Fix compile error in DatasmithRevitExporter.

#rnx

diff --git a/Engine/Source/Programs/Enterprise/Datasmith/DatasmithRevitExporter/Private/DatasmithRevitExportContext.cs b/Engine/Source/Programs/Enterprise/Datasmith/DatasmithRevitExporter/Private/DatasmithRevitExportContext.cs
--- a/Engine/Source/Programs/Enterprise/Datasmith/DatasmithRevitExporter/Private/DatasmithRevitExportContext.cs
+++ b/Engine/Source/Programs/Enterprise/Datasmith/DatasmithRevitExporter/Private/DatasmithRevitExportContext.cs
@@ -719,10 +719,6 @@ namespace DatasmithRevitExporter
 				InstanceData.InstanceActor.SetIsComponent(true);
 			}
 
-                // The Datasmith instance actor is a component in the hierarchy.
-                //InstanceData.InstanceActor.SetIsComponent(true);
-            }
-
             public FDatasmithFacadeMesh PopInstance()
 			{
 				FInstanceData InstanceData = InstanceDataStack.Pop();
