Fixes for DatasmithRevitExporter.

#rnx

diff --git a/Engine/Source/Programs/Enterprise/Datasmith/DatasmithRevitExporter/Private/DatasmithRevitExportContext.cs b/Engine/Source/Programs/Enterprise/Datasmith/DatasmithRevitExporter/Private/DatasmithRevitExportContext.cs
--- a/Engine/Source/Programs/Enterprise/Datasmith/DatasmithRevitExporter/Private/DatasmithRevitExportContext.cs
+++ b/Engine/Source/Programs/Enterprise/Datasmith/DatasmithRevitExporter/Private/DatasmithRevitExportContext.cs
@@ -719,7 +719,7 @@ namespace DatasmithRevitExporter
 				InstanceData.InstanceActor.SetIsComponent(true);
 			}
 
-            public FDatasmithFacadeMesh PopInstance()
+			public FDatasmithFacadeMesh PopInstance()
 			{
 				FInstanceData InstanceData = InstanceDataStack.Pop();
 
@@ -756,10 +756,6 @@ namespace DatasmithRevitExporter
 				string LayerName = Category.GetCategory(CurrentElement.Document, BuiltInCategory.OST_LightingFixtureSource)?.Name ?? "Light Sources";
 				SetActorProperties(LayerName, LightActor);
 
-				// Set the Datasmith light actor layer to its predefined name.
-				string CategoryName = Category.GetCategory(CurrentElement.Document, BuiltInCategory.OST_LightingFixtureSource)?.Name ?? "Light Sources";
-				LightActor.SetLayer(CategoryName);
-
 				// Set the specific properties of the Datasmith light actor.
 				FDatasmithRevitLight.SetLightProperties(InLightAsset, CurrentElement, LightActor);
 
diff --git a/Engine/Source/Programs/Enterprise/Datasmith/DatasmithRevitExporter/Private/DatasmithRevitMaterial.cs b/Engine/Source/Programs/Enterprise/Datasmith/DatasmithRevitExporter/Private/DatasmithRevitMaterial.cs
--- a/Engine/Source/Programs/Enterprise/Datasmith/DatasmithRevitExporter/Private/DatasmithRevitMaterial.cs
+++ b/Engine/Source/Programs/Enterprise/Datasmith/DatasmithRevitExporter/Private/DatasmithRevitMaterial.cs
@@ -55,7 +55,7 @@ namespace DatasmithRevitExporter
 			ExtraTexturePaths = InExtraTexturePaths;
 
 			// Create a new Datasmith master material.
-			MasterMaterial = new FDatasmithFacadeMaterial(GetMaterialName(InMaterialNode, CurrentMaterial), GetMaterialLabel());
+			MasterMaterial = new FDatasmithFacadeMaterial(GetMaterialName(InMaterialNode, CurrentMaterial), MaterialLabel);
 
 			// Hash the Datasmith master material name to shorten it.
 			MasterMaterial.HashName();
