Remove copy of max 2017 sdk from Enterprise

diff --git a/Engine/Source/Programs/Enterprise/Datasmith/DatasmithMaxExporter/DatasmithMax2017.Build.cs b/Engine/Source/Programs/Enterprise/Datasmith/DatasmithMaxExporter/DatasmithMax2017.Build.cs
--- a/Engine/Source/Programs/Enterprise/Datasmith/DatasmithMaxExporter/DatasmithMax2017.Build.cs
+++ b/Engine/Source/Programs/Enterprise/Datasmith/DatasmithMaxExporter/DatasmithMax2017.Build.cs
@@ -24,12 +24,11 @@ namespace UnrealBuildTool.Rules
 			// Max SDK setup
 			{
 				string MaxVersionString = GetMaxVersion();
+				string MaxSDKLocation = "";
 
-				string MaxSDKFolder = "3dsMax" + MaxVersionString + "SDK";
-				string MaxSDKLocation = Path.Combine("..", "..", "Enterprise", "Source", "ThirdParty", "NotForLicensees", MaxSDKFolder);
+				// Try with custom setup
 				string MaxSDKEnvVar = System.Environment.GetEnvironmentVariable("ADSK_3DSMAX_SDK_" + MaxVersionString);
-
-				if (!Directory.Exists(MaxSDKLocation) && MaxSDKEnvVar != null && MaxSDKEnvVar != "")
+				if (MaxSDKEnvVar != null && MaxSDKEnvVar != "")
 				{
 					MaxSDKLocation = MaxSDKEnvVar;
 				}
