Remove copy of SketchUp SDKs from Enterprise

diff --git a/Engine/Source/Programs/Enterprise/Datasmith/DatasmithSketchUpExporter/DatasmithSketchUp2017.Build.cs b/Engine/Source/Programs/Enterprise/Datasmith/DatasmithSketchUpExporter/DatasmithSketchUp2017.Build.cs
--- a/Engine/Source/Programs/Enterprise/Datasmith/DatasmithSketchUpExporter/DatasmithSketchUp2017.Build.cs
+++ b/Engine/Source/Programs/Enterprise/Datasmith/DatasmithSketchUpExporter/DatasmithSketchUp2017.Build.cs
@@ -34,19 +34,20 @@ namespace UnrealBuildTool.Rules
 
 			// Set up the SketchUp SDK paths and libraries.
 			{
-				string SketchUpSDKFolder = GetSketchUpSDKFolder();
+				string SketchUpSDKLocation = System.Environment.GetEnvironmentVariable(GetSketchUpEnvVar());
 
-				string SketchUpSDKLocation = Path.Combine("..", "..", "Enterprise", "Source", "ThirdParty", "NotForLicensees", "SketchUp", SketchUpSDKFolder);
-
-				// If SketchUp SDK is not part of workspace, look for environment variable
 				if (!Directory.Exists(SketchUpSDKLocation))
 				{
-					SketchUpSDKLocation = System.Environment.GetEnvironmentVariable(GetSketchUpEnvVar());
+					// Try with build machine setup
+					string SDKRootEnvVar = System.Environment.GetEnvironmentVariable("UE_SDKS_ROOT");
+					if (SDKRootEnvVar != null && SDKRootEnvVar != "")
+					{
+						SketchUpSDKLocation = Path.Combine(SDKRootEnvVar, "HostWin64", "Win64", "SketchUp", GetSketchUpSDKFolder());
+					}
 				}
 
 				// Make sure this version of the SketchUp SDK is actually installed.
 				if (Directory.Exists(SketchUpSDKLocation))
-
 				{
 					PrivateIncludePaths.Add(Path.Combine(SketchUpSDKLocation, "headers"));
 					PublicAdditionalLibraries.Add(Path.Combine(SketchUpSDKLocation, "binaries", "sketchup", "x64", "SketchUpAPI.lib"));
