Remove copy of Revit SDK from Enterprise

diff --git a/Engine/Source/Programs/Enterprise/Datasmith/DatasmithRevitExporter/DatasmithRevit2018.Target.cs b/Engine/Source/Programs/Enterprise/Datasmith/DatasmithRevitExporter/DatasmithRevit2018.Target.cs
--- a/Engine/Source/Programs/Enterprise/Datasmith/DatasmithRevitExporter/DatasmithRevit2018.Target.cs
+++ b/Engine/Source/Programs/Enterprise/Datasmith/DatasmithRevitExporter/DatasmithRevit2018.Target.cs
@@ -11,6 +11,12 @@ public abstract class DatasmithRevitBaseTarget : TargetRules
 		Type = TargetType.Program;
 		SolutionDirectory = "Programs/Enterprise";
 
+		string RevitVersionString = GetVersion();
+		string ProjectName = "DatasmithRevit" + RevitVersionString;
+
+		ExeBinariesSubFolder = Path.Combine("Revit", RevitVersionString);
+		LaunchModuleName = ProjectName;
+
 		bShouldCompileAsDLL = true;
 		LinkType = TargetLinkType.Monolithic;
 
@@ -24,23 +30,43 @@ public abstract class DatasmithRevitBaseTarget : TargetRules
 
 		bHasExports = true;
 		bForceEnableExceptions = true;
-	}
 
-	protected void AddPostBuildSteps(string ProjectName, string RevitAPIName)
-	{
+		// Define post-build step
+
+		string RevitSDKLocation = "";
+		string RevitSDKEnvVar = "Revit_" + RevitVersionString + "_API";
+
+		// Try with custom setup
+		string Location = System.Environment.GetEnvironmentVariable(RevitSDKEnvVar);
+		if (Location != null && Location != "")
+		{
+			RevitSDKLocation = Location;
+		}
+
+		if (!Directory.Exists(RevitSDKLocation))
+		{
+			// Try with build machine setup
+			string SDKRootEnvVar = System.Environment.GetEnvironmentVariable("UE_SDKS_ROOT");
+			if (SDKRootEnvVar != null && SDKRootEnvVar != "")
+			{
+				RevitSDKLocation = Path.Combine(SDKRootEnvVar, "HostWin64", "Win64", "Revit", RevitVersionString);
+			}
+		}
+
 		string RevitExporterPath = @"$(EngineDir)\Source\Programs\Enterprise\Datasmith\DatasmithRevitExporter";
-		string RevitAPILocation = Path.Combine(RevitExporterPath, "NotForLicensees", RevitAPIName);
 		string ProjectFile = Path.Combine(RevitExporterPath, ProjectName, ProjectName+".csproj");
-		string BuildCommand = string.Format(@"$(EngineDir)\Build\BatchFiles\MSBuild.bat /t:Build /p:Configuration=Release /p:{1}=%{1}% {0}", ProjectFile, RevitAPIName);
-		string ErrorMsg = string.Format("Cannot build {0}: Environment variable {1} is not defined.", ProjectName, RevitAPIName);
+		string BuildCommand = string.Format(@"$(EngineDir)\Build\BatchFiles\MSBuild.bat /t:Build /p:Configuration=Release /p:{1}=%{1}% {0}", ProjectFile, RevitSDKEnvVar);
+		string ErrorMsg = string.Format("Cannot build {0}: Environment variable {1} is not defined.", ProjectName, RevitSDKEnvVar);
 
 		// Since the Datasmith Revit Exporter is a C# project, build in batch the release configuration of the Visual Studio C# project file.
-		// Outside of Epic Games, environment variable <RevitAPIName> must be set to the Revit API directory on the developer's workstation.
+		// Outside of Epic Games, environment variable <RevitSDKEnvVar> (Revit_<year>_API) must be set to the Revit API directory on the developer's workstation.
 		PostBuildSteps.Add("setlocal enableextensions");
-		PostBuildSteps.Add(string.Format(@"if not defined {0} (if exist {1} (set {0}={1}) else ((echo {2}) & (exit /b 1)))", RevitAPIName, RevitAPILocation, ErrorMsg));
+		PostBuildSteps.Add(string.Format(@"if not defined {0} (if exist {1} (set {0}={1}) else ((echo {2}) & (exit /b 1)))", RevitSDKEnvVar, RevitSDKLocation, ErrorMsg));
 		PostBuildSteps.Add(string.Format(@"echo {0}", BuildCommand));
 		PostBuildSteps.Add(BuildCommand);
 	}
+
+	public abstract string GetVersion();
 }
 
 public class DatasmithRevit2018Target : DatasmithRevitBaseTarget
@@ -48,9 +74,7 @@ public class DatasmithRevit2018Target : DatasmithRevitBaseTarget
 	public DatasmithRevit2018Target(TargetInfo Target)
 		: base(Target)
 	{
-		LaunchModuleName = "DatasmithRevit2018";
-		ExeBinariesSubFolder = Path.Combine("Revit", "2018");
-
-		AddPostBuildSteps(LaunchModuleName, "Revit_2018_API");
 	}
+
+	public override string GetVersion() { return "2018"; }
 }
diff --git a/Engine/Source/Programs/Enterprise/Datasmith/DatasmithRevitExporter/DatasmithRevit2019.Target.cs b/Engine/Source/Programs/Enterprise/Datasmith/DatasmithRevitExporter/DatasmithRevit2019.Target.cs
--- a/Engine/Source/Programs/Enterprise/Datasmith/DatasmithRevitExporter/DatasmithRevit2019.Target.cs
+++ b/Engine/Source/Programs/Enterprise/Datasmith/DatasmithRevitExporter/DatasmithRevit2019.Target.cs
@@ -8,9 +8,7 @@ public class DatasmithRevit2019Target : DatasmithRevitBaseTarget
 	public DatasmithRevit2019Target(TargetInfo Target)
 		: base(Target)
 	{
-		LaunchModuleName = "DatasmithRevit2019";
-		ExeBinariesSubFolder = Path.Combine("Revit", "2019");
-		
-		AddPostBuildSteps(LaunchModuleName, "Revit_2019_API");
 	}
+
+	public override string GetVersion() { return "2019"; }
 }
diff --git a/Engine/Source/Programs/Enterprise/Datasmith/DatasmithRevitExporter/DatasmithRevit2020.Target.cs b/Engine/Source/Programs/Enterprise/Datasmith/DatasmithRevitExporter/DatasmithRevit2020.Target.cs
--- a/Engine/Source/Programs/Enterprise/Datasmith/DatasmithRevitExporter/DatasmithRevit2020.Target.cs
+++ b/Engine/Source/Programs/Enterprise/Datasmith/DatasmithRevitExporter/DatasmithRevit2020.Target.cs
@@ -8,9 +8,7 @@ public class DatasmithRevit2020Target : DatasmithRevitBaseTarget
 	public DatasmithRevit2020Target(TargetInfo Target)
 		: base(Target)
 	{
-		LaunchModuleName = "DatasmithRevit2020";
-		ExeBinariesSubFolder = Path.Combine("Revit", "2020");
-
-		AddPostBuildSteps(LaunchModuleName, "Revit_2020_API");
 	}
+
+	public override string GetVersion() { return "2020"; }
 }
