Doc update for DatasmithSDK

diff --git a/Engine/Source/Programs/Enterprise/Datasmith/DatasmithSDK/Documentation/2.SetUpYourProject.md b/Engine/Source/Programs/Enterprise/Datasmith/DatasmithSDK/Documentation/2.SetUpYourProject.md
--- a/Engine/Source/Programs/Enterprise/Datasmith/DatasmithSDK/Documentation/2.SetUpYourProject.md
+++ b/Engine/Source/Programs/Enterprise/Datasmith/DatasmithSDK/Documentation/2.SetUpYourProject.md
@@ -60,6 +60,7 @@ In order to use the classes and functions in the Datasmith SDK, you'll need to t
         _UNICODE
         UE_BUILD_DEVELOPMENT_WITH_DEBUGGAME=0
         UBT_COMPILED_PLATFORM=Windows
+        CORE_VTABLE=DLLIMPORT_VTABLE
 
 4.  Add the following additional library directories for your linker:
 
