Move DatasmithCADWorker in Programs/Enterprise/Datasmith
#rnx

diff --git a/Engine/Plugins/Enterprise/DatasmithCADImporter/Source/DatasmithDispatcher/Private/DatasmithWorkerHandler.cpp b/Engine/Plugins/Enterprise/DatasmithCADImporter/Source/DatasmithDispatcher/Private/DatasmithWorkerHandler.cpp
--- a/Engine/Plugins/Enterprise/DatasmithCADImporter/Source/DatasmithDispatcher/Private/DatasmithWorkerHandler.cpp
+++ b/Engine/Plugins/Enterprise/DatasmithCADImporter/Source/DatasmithDispatcher/Private/DatasmithWorkerHandler.cpp
@@ -17,15 +17,23 @@ namespace DatasmithDispatcher
 
 static FString GetWorkerExecutablePath()
 {
-	FString ProcessorPath = FPaths::Combine(FPaths::EnginePluginsDir(), TEXT("Enterprise/DatasmithCADImporter/Binaries"));
+	static FString ProcessorPath = [&]()
+	{
+		FString Path = FPaths::Combine(FPaths::EnginePluginsDir(), TEXT("Enterprise/DatasmithCADImporter/Binaries"));
 
 #if PLATFORM_MAC
-	ProcessorPath = FPaths::Combine(ProcessorPath, TEXT("Mac/DatasmithCADWorker"));
+		Path = FPaths::Combine(Path, TEXT("Mac/DatasmithCADWorker"));
 #elif PLATFORM_LINUX
-	ProcessorPath = FPaths::Combine(ProcessorPath, TEXT("Linux/DatasmithCADWorker"));
+		Path = FPaths::Combine(Path, TEXT("Linux/DatasmithCADWorker"));
 #elif PLATFORM_WINDOWS
-	ProcessorPath = FPaths::Combine(ProcessorPath, TEXT("Win64/DatasmithCADWorker.exe"));
+		Path = FPaths::Combine(Path, TEXT("Win64/DatasmithCADWorker.exe"));
 #endif
+		if (!FPaths::FileExists(Path))
+		{
+			UE_LOG(LogDatasmithDispatcher, Warning, TEXT("CADWorker executable not found. Expected location: %s"), *FPaths::ConvertRelativePathToFull(ProcessorPath));
+		}
+		return Path;
+	}();
 
 	return ProcessorPath;
 }
@@ -49,7 +57,7 @@ FDatasmithWorkerHandler::~FDatasmithWorkerHandler()
 void FDatasmithWorkerHandler::StartWorkerProcess()
 {
 	ensure(ErrorState == EWorkerErrorState::Ok);
-	static const FString ProcessorPath = GetWorkerExecutablePath();
+	FString ProcessorPath = GetWorkerExecutablePath();
 	if (FPaths::FileExists(ProcessorPath))
 	{
 		int32 ListenPort = NetworkInterface.GetListeningPort();
@@ -97,7 +105,7 @@ void FDatasmithWorkerHandler::Run()
 	WorkerState = EWorkerState::Uninitialized;
 	RunInternal();
 	WorkerState = EWorkerState::Terminated;
-	UE_CLOG(ErrorState != EWorkerErrorState::Ok, LogDatasmithDispatcher, Warning, TEXT("Handler ended with error %s"), EWorkerErrorStateAsString(ErrorState));
+	UE_CLOG(ErrorState != EWorkerErrorState::Ok, LogDatasmithDispatcher, Warning, TEXT("Handler ended with error: %s"), EWorkerErrorStateAsString(ErrorState));
 }
 
 void FDatasmithWorkerHandler::RunInternal()
@@ -320,17 +328,17 @@ void FDatasmithWorkerHandler::ProcessCommand(FCompletedTaskCommand& CompletedTas
 	CurrentTask.Reset();
 }
 
-const TCHAR* FDatasmithWorkerHandler::EWorkerErrorStateAsString(EWorkerErrorState e)
+const TCHAR* FDatasmithWorkerHandler::EWorkerErrorStateAsString(EWorkerErrorState Error)
 {
-	switch (e)
+	switch (Error)
 	{
 		case EWorkerErrorState::Ok: return TEXT("Ok");
-		case EWorkerErrorState::ConnectionFailed_NotBound: return TEXT("ConnectionFailed_NotBound");
-		case EWorkerErrorState::ConnectionFailed_NoClient: return TEXT("ConnectionFailed_NoClient");
-		case EWorkerErrorState::ConnectionLost: return TEXT("ConnectionLost");
-		case EWorkerErrorState::ConnectionLost_SendFailed: return TEXT("ConnectionLost_SendFailed");
-		case EWorkerErrorState::WorkerProcess_CantCreate: return TEXT("WorkerProcess_CantCreate");
-		case EWorkerErrorState::WorkerProcess_Lost: return TEXT("WorkerProcess_Lost");
+		case EWorkerErrorState::ConnectionFailed_NotBound: return TEXT("Connection failed (socket not bound)");
+		case EWorkerErrorState::ConnectionFailed_NoClient: return TEXT("Connection failed (No connection from client)");
+		case EWorkerErrorState::ConnectionLost: return TEXT("Connection lost");
+		case EWorkerErrorState::ConnectionLost_SendFailed: return TEXT("Connection lost (send failed)");
+		case EWorkerErrorState::WorkerProcess_CantCreate: return TEXT("Worker process issue (cannot create worker process)");
+		case EWorkerErrorState::WorkerProcess_Lost: return TEXT("Worker process issue (worker lost)");
 		default: return TEXT("unknown");
 	}
 }
diff --git a/Engine/Source/Programs/Enterprise/Datasmith/DatasmithCADWorker/DatasmithCADWorker.Build.cs b/Engine/Source/Programs/Enterprise/Datasmith/DatasmithCADWorker/DatasmithCADWorker.Build.cs
new file mode 100644
--- /dev/null
+++ b/Engine/Source/Programs/Enterprise/Datasmith/DatasmithCADWorker/DatasmithCADWorker.Build.cs
@@ -0,0 +1,29 @@
+// Copyright 1998-2019 Epic Games, Inc. All Rights Reserved.
+
+using UnrealBuildTool;
+using System.IO;
+
+public class DatasmithCADWorker : ModuleRules
+{
+	public DatasmithCADWorker(ReadOnlyTargetRules Target) : base(Target)
+	{
+
+		PublicIncludePaths.Add("Runtime/Launch/Public");
+		PrivateIncludePaths.Add("Runtime/Launch/Private");
+
+		PrivateDependencyModuleNames.AddRange(
+			new string[]
+			{
+				"Core",
+				"Projects",
+				"ApplicationCore",
+				"Sockets",
+				"DatasmithDispatcher",
+				"CADInterfaces",
+				"CADTools",
+			}
+		);
+
+		PublicDelayLoadDLLs.Add("kernel_io.dll");
+	}
+}
diff --git a/Engine/Source/Programs/Enterprise/Datasmith/DatasmithCADWorker/DatasmithCADWorker.Target.cs b/Engine/Source/Programs/Enterprise/Datasmith/DatasmithCADWorker/DatasmithCADWorker.Target.cs
new file mode 100644
--- /dev/null
+++ b/Engine/Source/Programs/Enterprise/Datasmith/DatasmithCADWorker/DatasmithCADWorker.Target.cs
@@ -0,0 +1,45 @@
+// Copyright 1998-2019 Epic Games, Inc. All Rights Reserved.
+
+using UnrealBuildTool;
+using System.Collections.Generic;
+
+[SupportedPlatforms(UnrealPlatformClass.Desktop)]
+public class DatasmithCADWorkerTarget : TargetRules
+{
+	public DatasmithCADWorkerTarget(TargetInfo Target) : base(Target)
+	{
+		Type = TargetType.Program;
+		LinkType = TargetLinkType.Monolithic;
+		LaunchModuleName = "DatasmithCADWorker";
+		SolutionDirectory = "Programs/Enterprise";
+
+        // Lean and mean
+        bBuildDeveloperTools = false;
+
+		// Never use malloc profiling in Unreal Header Tool.  We set this because often UHT is compiled right before the engine
+		// automatically by Unreal Build Tool, but if bUseMallocProfiler is defined, UHT can operate incorrectly.
+		bUseMallocProfiler = false;
+
+		// Editor-only data, however, is needed
+		bBuildWithEditorOnlyData = true;
+
+		// Currently this app is not linking against the engine, so we'll compile out references from Core to the rest of the engine
+		bCompileAgainstEngine = false;
+		bCompileAgainstCoreUObject = false;
+		bCompileAgainstApplicationCore = true;
+
+		// This is a console application, not a Windows app (sets entry point to main(), instead of WinMain())
+		bIsBuildingConsoleApplication = true;
+
+		bLegalToDistributeBinary = true;
+
+		bCompilePhysX = false;
+		bCompileAPEX = false;
+		bCompileNvCloth = false;
+		bCompileICU = false;
+		bCompileCEF3 = false;
+
+		string EngineDir = @"..\.."; // relative to default destination, which is $(EngineDir)\Binaries\Win64
+		ExeBinariesSubFolder = EngineDir + @"\Plugins\Enterprise\DatasmithCADImporter\Binaries\Win64\";
+	}
+}
diff --git a/Engine/Source/Programs/Enterprise/Datasmith/DatasmithCADWorker/Private/DatasmithCADWorker.cpp b/Engine/Source/Programs/Enterprise/Datasmith/DatasmithCADWorker/Private/DatasmithCADWorker.cpp
new file mode 100644
--- /dev/null
+++ b/Engine/Source/Programs/Enterprise/Datasmith/DatasmithCADWorker/Private/DatasmithCADWorker.cpp
@@ -0,0 +1,81 @@
+// Copyright 1998-2019 Epic Games, Inc. All Rights Reserved.
+
+#include "DatasmithCADWorker.h"
+#include "DatasmithCADWorkerImpl.h"
+
+#include "RequiredProgramMainCPPInclude.h"
+
+IMPLEMENT_APPLICATION(DatasmithCADWorker, "DatasmithCADWorker");
+
+
+void GetParameter(int32 Argc, TCHAR* Argv[], const FString& InParam, FString& OutValue)
+{
+	OutValue.Empty();
+	for (int32 Index = 1; Index < Argc; Index++)
+	{
+		if (! FCString::Strcmp(*InParam, Argv[Index]))
+		{
+			if (Index + 1 < Argc)
+			{
+				OutValue = Argv[Index + 1];
+			}
+		}
+	}
+}
+
+bool HasParameter(int32 Argc, TCHAR* Argv[], const FString& InParam)
+{
+	for (int32 Index = 1; Index < Argc; Index++)
+	{
+		if (!FCString::Strcmp(*InParam, Argv[Index]))
+		{
+			return true;
+		}
+	}
+	return false;
+}
+
+int32 Main(int32 Argc, TCHAR * Argv[])
+{
+	FString ServerPID, ServerPort, CacheDirectory;
+	GetParameter(Argc, Argv, "-ServerPID", ServerPID);
+	GetParameter(Argc, Argv, "-ServerPort", ServerPort);
+	GetParameter(Argc, Argv, "-CacheDir", CacheDirectory);
+
+	FString BinaryDir = FPaths::Combine(FPaths::EngineDir(), TEXT("Plugins"), TEXT("Enterprise"), TEXT("DatasmithCADImporter"), TEXT("Binaries"), TEXT("Win64"));
+	FString KernelIOLibDir = FPaths::Combine(BinaryDir, TEXT("kernel_io.dll"));
+	FPlatformProcess::PushDllDirectory(*BinaryDir);
+	void* KernelIODllHandle = FPlatformProcess::GetDllHandle(*KernelIOLibDir);
+	FPlatformProcess::PopDllDirectory(*BinaryDir);
+
+	FDatasmithCADWorkerImpl Worker(FCString::Atoi(*ServerPID), FCString::Atoi(*ServerPort), CacheDirectory);
+	Worker.Run();
+
+	return 0;
+}
+
+int32 Filter(uint32 Code, struct _EXCEPTION_POINTERS *Ep)
+{
+	return EXCEPTION_EXECUTE_HANDLER;
+}
+
+INT32_MAIN_INT32_ARGC_TCHAR_ARGV()
+{
+#ifdef CAD_INTERFACE
+	GEngineLoop.PreInit(ArgC, ArgV);
+
+	SetErrorMode(SEM_FAILCRITICALERRORS | SEM_NOGPFAULTERRORBOX);
+	_set_abort_behavior(0, _WRITE_ABORT_MSG);
+	__try
+	{
+		return Main(ArgC, ArgV);
+	}
+	__except (Filter(GetExceptionCode(), GetExceptionInformation()))
+	{
+		return EXIT_FAILURE;
+	}
+	return EXIT_SUCCESS;
+#else
+	return EXIT_FAILURE;
+#endif // CAD_INTERFACE
+}
diff --git a/Engine/Source/Programs/Enterprise/Datasmith/DatasmithCADWorker/Private/DatasmithCADWorker.h b/Engine/Source/Programs/Enterprise/Datasmith/DatasmithCADWorker/Private/DatasmithCADWorker.h
new file mode 100644
--- /dev/null
+++ b/Engine/Source/Programs/Enterprise/Datasmith/DatasmithCADWorker/Private/DatasmithCADWorker.h
@@ -0,0 +1,6 @@
+// Copyright 1998-2019 Epic Games, Inc. All Rights Reserved.
+
+#pragma once
+
+#include "CoreMinimal.h"
+
diff --git a/Engine/Source/Programs/Enterprise/Datasmith/DatasmithCADWorker/Private/DatasmithCADWorkerImpl.cpp b/Engine/Source/Programs/Enterprise/Datasmith/DatasmithCADWorker/Private/DatasmithCADWorkerImpl.cpp
new file mode 100644
--- /dev/null
+++ b/Engine/Source/Programs/Enterprise/Datasmith/DatasmithCADWorker/Private/DatasmithCADWorkerImpl.cpp
@@ -0,0 +1,144 @@
+// Copyright 1998-2019 Epic Games, Inc. All Rights Reserved.
+
+#include "DatasmithCADWorkerImpl.h"
+
+#include "CoreTechFileParser.h"
+#include "DatasmithCommands.h"
+#include "DatasmithDispatcherConfig.h"
+
+#include "HAL/FileManager.h"
+#include "Misc/FileHelper.h"
+#include "Misc/Paths.h"
+#include "SocketSubsystem.h"
+#include "Sockets.h"
+
+DEFINE_LOG_CATEGORY(LogDatasmithCADWorker);
+using namespace DatasmithDispatcher;
+
+
+FDatasmithCADWorkerImpl::FDatasmithCADWorkerImpl(int32 InServerPID, int32 InServerPort, const FString& InCachePath)
+	: ServerPID(InServerPID)
+	, ServerPort(InServerPort)
+	, CachePath(InCachePath)
+	, PingStartCycle(0)
+{
+	UE_SET_LOG_VERBOSITY(LogDatasmithCADWorker, Verbose);
+}
+
+bool FDatasmithCADWorkerImpl::Run()
+{
+	UE_LOG(LogDatasmithCADWorker, Verbose, TEXT("connect to %d..."), ServerPort);
+	bool bConnected = NetworkInterface.Connect(TEXT("Datasmith CAD Worker"), ServerPort, Config::ConnectTimeout_s);
+	UE_LOG(LogDatasmithCADWorker, Verbose, TEXT("connected to %d %s"), ServerPort, bConnected ? TEXT("OK") : TEXT("FAIL"));
+	if (bConnected)
+	{
+		CommandIO.SetNetworkInterface(&NetworkInterface);
+	}
+	else
+	{
+		UE_LOG(LogDatasmithCADWorker, Error, TEXT("Server connection failure. exit"));
+		return false;
+	}
+
+	InitiatePing();
+
+	bool bIsRunning = true;
+	while (bIsRunning)
+	{
+		if (TSharedPtr<ICommand> Command = CommandIO.GetNextCommand(1.0))
+		{
+			switch(Command->GetType())
+			{
+				case ECommandId::Ping:
+					ProcessCommand(*StaticCast<FPingCommand*>(Command.Get()));
+					break;
+
+				case ECommandId::BackPing:
+					ProcessCommand(*StaticCast<FBackPingCommand*>(Command.Get()));
+					break;
+
+				case ECommandId::RunTask:
+					ProcessCommand(*StaticCast<FRunTaskCommand*>(Command.Get()));
+					break;
+
+				case ECommandId::ImportParams:
+					ProcessCommand(*StaticCast<FImportParametersCommand*>(Command.Get()));
+					break;
+
+				case ECommandId::Terminate:
+					UE_LOG(LogDatasmithCADWorker, Verbose, TEXT("Terminate command received. Exiting."));
+					bIsRunning = false;
+					break;
+
+				case ECommandId::NotifyEndTask:
+				default:
+					break;
+			}
+		}
+		else
+		{
+			if (bIsRunning)
+			{
+				bIsRunning = ServerPID == 0 ? true : FPlatformProcess::IsApplicationRunning(ServerPID);
+				UE_CLOG(!bIsRunning, LogDatasmithCADWorker, Error, TEXT("Worker failure: server lost"));
+			}
+		}
+	}
+
+	UE_CLOG(!bIsRunning, LogDatasmithCADWorker, Verbose, TEXT("Worker loop exit..."));
+	CommandIO.Disconnect(0);
+	return true;
+}
+
+void FDatasmithCADWorkerImpl::InitiatePing()
+{
+	PingStartCycle = FPlatformTime::Cycles64();
+	FPingCommand Ping;
+	CommandIO.SendCommand(Ping, Config::SendCommandTimeout_s);
+}
+
+void FDatasmithCADWorkerImpl::ProcessCommand(const FPingCommand& PingCommand)
+{
+	FBackPingCommand BackPing;
+	CommandIO.SendCommand(BackPing, Config::SendCommandTimeout_s);
+}
+
+void FDatasmithCADWorkerImpl::ProcessCommand(const FBackPingCommand& BackPingCommand)
+{
+	if (PingStartCycle)
+	{
+		double ElapsedTime_s = FGenericPlatformTime::ToSeconds(FPlatformTime::Cycles64() - PingStartCycle);
+		UE_LOG(LogDatasmithCADWorker, Verbose, TEXT("Ping %f s"), ElapsedTime_s);
+	}
+	PingStartCycle = 0;
+}
+
+void FDatasmithCADWorkerImpl::ProcessCommand(const FImportParametersCommand& ImportParametersCommand)
+{
+	ImportParameters = ImportParametersCommand.ImportParameters;
+}
+
+void FDatasmithCADWorkerImpl::ProcessCommand(const FRunTaskCommand& RunTaskCommand)
+{
+	const FString& FileToProcess = RunTaskCommand.JobFilePath;
+	UE_LOG(LogDatasmithCADWorker, Verbose, TEXT("Process %s"), *FileToProcess);
+
+	FCompletedTaskCommand CompletedTask;
+#ifdef CAD_INTERFACE
+	CADLibrary::FCoreTechFileParser FileParser(FileToProcess, CachePath, ImportParameters, true, true);
+	CADLibrary::FCoreTechFileParser::EProcessResult ProcessResult = FileParser.ProcessFile();
+
+	CompletedTask.ProcessResult = ProcessResult;
+
+	if (CompletedTask.ProcessResult == ETaskState::ProcessOk)
+	{
+		CompletedTask.ExternalReferences = FileParser.GetExternalRefSet().Array();
+		CompletedTask.SceneGraphFileName = FileParser.GetSceneGraphFile();
+		CompletedTask.GeomFileName = FileParser.GetMeshFile();
+	}
+#endif // CAD_INTERFACE
+	CommandIO.SendCommand(CompletedTask, Config::SendCommandTimeout_s);
+
+	UE_LOG(LogDatasmithCADWorker, Verbose, TEXT("End of Process %s"), *FileToProcess);
+}
+
diff --git a/Engine/Source/Programs/Enterprise/Datasmith/DatasmithCADWorker/Private/DatasmithCADWorkerImpl.h b/Engine/Source/Programs/Enterprise/Datasmith/DatasmithCADWorker/Private/DatasmithCADWorkerImpl.h
new file mode 100644
--- /dev/null
+++ b/Engine/Source/Programs/Enterprise/Datasmith/DatasmithCADWorker/Private/DatasmithCADWorkerImpl.h
@@ -0,0 +1,42 @@
+// Copyright 1998-2019 Epic Games, Inc. All Rights Reserved.
+#pragma once
+
+#include "CoreMinimal.h"
+
+#include "CADOptions.h"
+#include "DatasmithCommands.h"
+#include "DatasmithDispatcherNetworking.h"
+
+#include "HAL/PlatformTime.h"
+#include "Logging/LogMacros.h"
+
+DECLARE_LOG_CATEGORY_EXTERN(LogDatasmithCADWorker, Log, All);
+
+
+struct FFileStatData;
+struct FImportParameters;
+
+class FDatasmithCADWorkerImpl
+{
+public:
+	FDatasmithCADWorkerImpl(int32 InServerPID, int32 InServerPort, const FString& InCachePath);
+	bool Run();
+
+private:
+	void InitiatePing();
+
+	void ProcessCommand(const DatasmithDispatcher::FPingCommand& PingCommand);
+	void ProcessCommand(const DatasmithDispatcher::FBackPingCommand& BackPingCommand);
+	void ProcessCommand(const DatasmithDispatcher::FImportParametersCommand& BackPingCommand);
+	void ProcessCommand(const DatasmithDispatcher::FRunTaskCommand& TerminateCommand);
+
+private:
+	DatasmithDispatcher::FNetworkClientNode NetworkInterface;
+	DatasmithDispatcher::FCommandQueue CommandIO;
+
+	int32 ServerPID;
+	int32 ServerPort;
+	FString CachePath;
+	CADLibrary::FImportParameters ImportParameters;
+	uint64 PingStartCycle;
+};
