Datasmith: in solution, renamed Program/Enterprise to Program/Datasmith

#rnx

diff --git a/Engine/Source/Programs/Enterprise/Datasmith/DatasmithCADWorker/DatasmithCADWorker.Target.cs b/Engine/Source/Programs/Enterprise/Datasmith/DatasmithCADWorker/DatasmithCADWorker.Target.cs
--- a/Engine/Source/Programs/Enterprise/Datasmith/DatasmithCADWorker/DatasmithCADWorker.Target.cs
+++ b/Engine/Source/Programs/Enterprise/Datasmith/DatasmithCADWorker/DatasmithCADWorker.Target.cs
@@ -11,7 +11,7 @@ public class DatasmithCADWorkerTarget : TargetRules
 		Type = TargetType.Program;
 		LinkType = TargetLinkType.Monolithic;
 		LaunchModuleName = "DatasmithCADWorker";
-		SolutionDirectory = "Programs/Enterprise";
+		SolutionDirectory = "Programs/Datasmith";
 
         // Lean and mean
         bBuildDeveloperTools = false;
diff --git a/Engine/Source/Programs/Enterprise/Datasmith/DatasmithFacadeCSharp/DatasmithFacadeCSharp.Target.cs b/Engine/Source/Programs/Enterprise/Datasmith/DatasmithFacadeCSharp/DatasmithFacadeCSharp.Target.cs
--- a/Engine/Source/Programs/Enterprise/Datasmith/DatasmithFacadeCSharp/DatasmithFacadeCSharp.Target.cs
+++ b/Engine/Source/Programs/Enterprise/Datasmith/DatasmithFacadeCSharp/DatasmithFacadeCSharp.Target.cs
@@ -10,7 +10,7 @@ public class DatasmithFacadeCSharpTarget : TargetRules
 		: base(Target)
 	{
 		Type = TargetType.Program;
-		SolutionDirectory = "Programs/Enterprise";
+		SolutionDirectory = "Programs/Datasmith";
 
 		LaunchModuleName = "DatasmithFacadeCSharp";
 		ExeBinariesSubFolder = "DatasmithFacadeCSharp";
diff --git a/Engine/Source/Programs/Enterprise/Datasmith/DatasmithMaxExporter/DatasmithMax2017.Target.cs b/Engine/Source/Programs/Enterprise/Datasmith/DatasmithMaxExporter/DatasmithMax2017.Target.cs
--- a/Engine/Source/Programs/Enterprise/Datasmith/DatasmithMaxExporter/DatasmithMax2017.Target.cs
+++ b/Engine/Source/Programs/Enterprise/Datasmith/DatasmithMaxExporter/DatasmithMax2017.Target.cs
@@ -8,7 +8,7 @@ public abstract class DatasmithMaxBaseTarget : TargetRules
 		: base(Target)
 	{
 		Type = TargetType.Program;
-		SolutionDirectory = "Programs/Enterprise";
+		SolutionDirectory = "Programs/Datasmith";
 		bLegalToDistributeBinary = true;
 
 		bShouldCompileAsDLL = true;
diff --git a/Engine/Source/Programs/Enterprise/Datasmith/DatasmithRevitExporter/DatasmithRevit2018.Target.cs b/Engine/Source/Programs/Enterprise/Datasmith/DatasmithRevitExporter/DatasmithRevit2018.Target.cs
--- a/Engine/Source/Programs/Enterprise/Datasmith/DatasmithRevitExporter/DatasmithRevit2018.Target.cs
+++ b/Engine/Source/Programs/Enterprise/Datasmith/DatasmithRevitExporter/DatasmithRevit2018.Target.cs
@@ -9,7 +9,7 @@ public abstract class DatasmithRevitBaseTarget : TargetRules
 		: base(Target)
 	{
 		Type = TargetType.Program;
-		SolutionDirectory = "Programs/Enterprise";
+		SolutionDirectory = "Programs/Datasmith";
 
 		string RevitVersionString = GetVersion();
 		string ProjectName = "DatasmithRevit" + RevitVersionString;
diff --git a/Engine/Source/Programs/Enterprise/Datasmith/DatasmithSDK/DatasmithSDK.Target.cs b/Engine/Source/Programs/Enterprise/Datasmith/DatasmithSDK/DatasmithSDK.Target.cs
--- a/Engine/Source/Programs/Enterprise/Datasmith/DatasmithSDK/DatasmithSDK.Target.cs
+++ b/Engine/Source/Programs/Enterprise/Datasmith/DatasmithSDK/DatasmithSDK.Target.cs
@@ -9,7 +9,7 @@ public class DatasmithSDKTarget : TargetRules
 		: base(Target)
 	{
 		Type = TargetType.Program;
-		SolutionDirectory = "Programs/Enterprise";
+		SolutionDirectory = "Programs/Datasmith";
 
 		LaunchModuleName = "DatasmithSDK";
 		ExeBinariesSubFolder = "DatasmithSDK";
diff --git a/Engine/Source/Programs/Enterprise/Datasmith/DatasmithSketchUpExporter/DatasmithSketchUp2017.Target.cs b/Engine/Source/Programs/Enterprise/Datasmith/DatasmithSketchUpExporter/DatasmithSketchUp2017.Target.cs
--- a/Engine/Source/Programs/Enterprise/Datasmith/DatasmithSketchUpExporter/DatasmithSketchUp2017.Target.cs
+++ b/Engine/Source/Programs/Enterprise/Datasmith/DatasmithSketchUpExporter/DatasmithSketchUp2017.Target.cs
@@ -8,7 +8,7 @@ public abstract class DatasmithSketchUpBaseTarget : TargetRules
 		: base(Target)
 	{
 		Type = TargetType.Program;
-		SolutionDirectory = "Programs/Enterprise";
+		SolutionDirectory = "Programs/Datasmith";
 
 		bShouldCompileAsDLL = true;
 		LinkType = TargetLinkType.Monolithic;
