Copy up //UE4/Dev-Enterprise/...at  Changelist 10508152  to //UE4/Dev-Tools-Staging/...

diff --git a/Engine/Source/Programs/Enterprise/Datasmith/DatasmithRevitExporter/Private/DatasmithRevitCommand.cs b/Engine/Source/Programs/Enterprise/Datasmith/DatasmithRevitExporter/Private/DatasmithRevitCommand.cs
--- a/Engine/Source/Programs/Enterprise/Datasmith/DatasmithRevitExporter/Private/DatasmithRevitCommand.cs
+++ b/Engine/Source/Programs/Enterprise/Datasmith/DatasmithRevitExporter/Private/DatasmithRevitCommand.cs
@@ -1,8 +1,10 @@
 // Copyright 1998-2019 Epic Games, Inc. All Rights Reserved.
 
 using System;
+using System.Collections.Generic;
 using System.Diagnostics;
 using System.IO;
+using System.Linq;
 using System.Runtime.InteropServices;
 using System.Text.RegularExpressions;
 using System.Windows.Forms;
@@ -11,7 +13,6 @@ using Autodesk.Revit.Attributes;
 using Autodesk.Revit.DB;
 using Autodesk.Revit.UI;
 
-
 namespace DatasmithRevitExporter
 {
 	// Add-in external command Export to Unreal Datasmith. 
@@ -22,111 +23,148 @@ namespace DatasmithRevitExporter
 
 		// Implement the interface to execute the command.
 		public Result Execute(
-			ExternalCommandData in_commandData,     // contains reference to Application and View
-			ref string          out_commandMessage, // error message to display in the failure dialog when the command returns "Failed"
-			ElementSet          out_elements        // set of problem elements to display in the failure dialog when the command returns "Failed"
+			ExternalCommandData InCommandData,     // contains reference to Application and View
+			ref string          OutCommandMessage, // error message to display in the failure dialog when the command returns "Failed"
+			ElementSet          OutElements        // set of problem elements to display in the failure dialog when the command returns "Failed"
 		)
 		{
-			Autodesk.Revit.ApplicationServices.Application application = in_commandData.Application.Application;
+			Autodesk.Revit.ApplicationServices.Application Application = InCommandData.Application.Application;
 
-			if (string.Compare(application.VersionNumber, "2018", StringComparison.Ordinal) == 0 && string.Compare(application.SubVersionNumber, "2018.3", StringComparison.Ordinal) < 0)
+			if (string.Compare(Application.VersionNumber, "2018", StringComparison.Ordinal) == 0 && string.Compare(Application.SubVersionNumber, "2018.3", StringComparison.Ordinal) < 0)
 			{
-				string message = string.Format("The running Revit is not supported.\nYou must use Revit 2018.3 or further updates to export.");
-				MessageBox.Show(message, DIALOG_CAPTION, MessageBoxButtons.OK, MessageBoxIcon.Warning);
+				string Message = string.Format("The running Revit is not supported.\nYou must use Revit 2018.3 or further updates to export.");
+				MessageBox.Show(Message, DIALOG_CAPTION, MessageBoxButtons.OK, MessageBoxIcon.Warning);
 				return Result.Cancelled;
 			}
 
 			if (!CustomExporter.IsRenderingSupported())
 			{
-				string message = "3D view rendering is not supported in the running Revit.";
-				MessageBox.Show(message, DIALOG_CAPTION, MessageBoxButtons.OK, MessageBoxIcon.Warning);
+				string Message = "3D view rendering is not supported in the running Revit.";
+				MessageBox.Show(Message, DIALOG_CAPTION, MessageBoxButtons.OK, MessageBoxIcon.Warning);
 				return Result.Cancelled;
 			}
 
-			UIDocument uiDocument = in_commandData.Application.ActiveUIDocument;
+			UIDocument UIDoc = InCommandData.Application.ActiveUIDocument;
 
-			if (uiDocument == null)
+			if (UIDoc == null)
 			{
-				string message = "You must be in a document to export.";
-				MessageBox.Show(message, DIALOG_CAPTION, MessageBoxButtons.OK, MessageBoxIcon.Warning);
+				string Message = "You must be in a document to export.";
+				MessageBox.Show(Message, DIALOG_CAPTION, MessageBoxButtons.OK, MessageBoxIcon.Warning);
 				return Result.Cancelled;
 			}
 
-			Document document = uiDocument.Document;
-	
-			View3D activeView = document.ActiveView as View3D;
+			Document Doc = UIDoc.Document;
 
-			if (activeView == null)
+			string DocumentPath = Doc.PathName;
+
+			if (string.IsNullOrWhiteSpace(DocumentPath))
 			{
-				string message = "You must be in a 3D view to export.";
+				string message = "Your document must be saved on disk before exporting.";
 				MessageBox.Show(message, DIALOG_CAPTION, MessageBoxButtons.OK, MessageBoxIcon.Warning);
 				return Result.Cancelled;
 			}
 
-			if (activeView.IsTemplate || !activeView.CanBePrinted)
+			bool ExportActiveViewOnly = true;
+
+			// Retrieve the Unreal Datasmith export options.
+			DatasmithRevitExportOptions ExportOptions = new DatasmithRevitExportOptions(Doc);
+			if ((System.Windows.Forms.Control.ModifierKeys & Keys.Control) == Keys.Control)
 			{
-				string message = "The active 3D view cannot be exported.";
-				MessageBox.Show(message, DIALOG_CAPTION, MessageBoxButtons.OK, MessageBoxIcon.Warning);
-				return Result.Cancelled;
+				if (ExportOptions.ShowDialog() != DialogResult.OK)
+				{
+					return Result.Cancelled;
+				}
+				ExportActiveViewOnly = false;
 			}
 
-			string documentPath = document.PathName;
+			// Generate file path for each view.
+			Dictionary<ElementId, string> FilePaths = new Dictionary<ElementId, string>();
+			List<View3D> ViewsToExport = new List<View3D>();
 
-			if (string.IsNullOrWhiteSpace(documentPath))
+			if (ExportActiveViewOnly)
 			{
-				string message = "Your document must be saved on disk before exporting.";
-				MessageBox.Show(message, DIALOG_CAPTION, MessageBoxButtons.OK, MessageBoxIcon.Warning);
-				return Result.Cancelled;
-			}
+				View3D ActiveView = Doc.ActiveView as View3D;
 
-			string ViewFamilyName = activeView.get_Parameter(BuiltInParameter.ELEM_FAMILY_PARAM).AsValueString().Replace(" ", "");
-			string FileName       = Regex.Replace($"{Path.GetFileNameWithoutExtension(documentPath)}-{ViewFamilyName}-{activeView.Name}.udatasmith", @"\s+", "_");
+				if (ActiveView == null)
+				{
+					string Message = "You must be in a 3D view to export.";
+					MessageBox.Show(Message, DIALOG_CAPTION, MessageBoxButtons.OK, MessageBoxIcon.Warning);
+					return Result.Cancelled;
+				}
 
-			SaveFileDialog dialog = new SaveFileDialog();
+				if (ActiveView.IsTemplate || !ActiveView.CanBePrinted)
+				{
+					string Message = "The active 3D view cannot be exported.";
+					MessageBox.Show(Message, DIALOG_CAPTION, MessageBoxButtons.OK, MessageBoxIcon.Warning);
+					return Result.Cancelled;
+				}
 
-			dialog.Title            = DIALOG_CAPTION;
-			dialog.InitialDirectory = Path.GetDirectoryName(documentPath);
-			dialog.FileName         = FileName;
-			dialog.DefaultExt       = "udatasmith";
-			dialog.Filter           = "Unreal Datasmith|*.udatasmith";
-			dialog.CheckFileExists  = false;
-			dialog.CheckPathExists  = true;
-			dialog.AddExtension     = true;
-			dialog.OverwritePrompt  = true;
+				string ViewFamilyName = ActiveView.get_Parameter(BuiltInParameter.ELEM_FAMILY_PARAM).AsValueString().Replace(" ", "");
+				string FileName       = Regex.Replace($"{Path.GetFileNameWithoutExtension(DocumentPath)}-{ViewFamilyName}-{ActiveView.Name}.udatasmith", @"\s+", "_");
 
-			if (dialog.ShowDialog() != DialogResult.OK)
-			{
-				return Result.Cancelled;
-			}
+				SaveFileDialog Dialog = new SaveFileDialog();
 
-			string filePath = dialog.FileName;
+				Dialog.Title            = DIALOG_CAPTION;
+				Dialog.InitialDirectory = Path.GetDirectoryName(DocumentPath);
+				Dialog.FileName         = FileName;
+				Dialog.DefaultExt       = "udatasmith";
+				Dialog.Filter           = "Unreal Datasmith|*.udatasmith";
+				Dialog.CheckFileExists  = false;
+				Dialog.CheckPathExists  = true;
+				Dialog.AddExtension     = true;
+				Dialog.OverwritePrompt  = true;
 
-			if (string.IsNullOrWhiteSpace(filePath))
-			{
-				string message = "The given Unreal Datasmith file name is blank.";
-				MessageBox.Show(message, DIALOG_CAPTION, MessageBoxButtons.OK, MessageBoxIcon.Warning);
-				return Result.Cancelled;
-			}
+				if (Dialog.ShowDialog() != DialogResult.OK)
+				{
+					return Result.Cancelled;
+				}
 
-			// Retrieve the Unreal Datasmith export options.
-			DatasmithRevitExportOptions exportOptions = new DatasmithRevitExportOptions();
-			if ((System.Windows.Forms.Control.ModifierKeys & Keys.Control) == Keys.Control)
-			{
-				if (exportOptions.ShowDialog() != DialogResult.OK)
+				string FilePath = Dialog.FileName;
+
+				if (string.IsNullOrWhiteSpace(FilePath))
 				{
+					string message = "The given Unreal Datasmith file name is blank.";
+					MessageBox.Show(message, DIALOG_CAPTION, MessageBoxButtons.OK, MessageBoxIcon.Warning);
 					return Result.Cancelled;
 				}
+
+				FilePaths.Add(ActiveView.Id, FilePath);
+				ViewsToExport.Add(ActiveView);
+			}
+			else
+			{
+				string SavePath;
+				using (var FBD = new FolderBrowserDialog())
+				{
+					FBD.ShowNewFolderButton = true;
+					DialogResult result = FBD.ShowDialog();
+
+					if (result != DialogResult.OK || string.IsNullOrWhiteSpace(FBD.SelectedPath))
+					{
+						return Result.Cancelled;
+					}
+
+					SavePath = FBD.SelectedPath;
+				}
+
+				foreach (var View in ExportOptions.Selected3DViews)
+				{
+					string ViewFamilyName = View.get_Parameter(BuiltInParameter.ELEM_FAMILY_PARAM).AsValueString().Replace(" ", "");
+					string FileName = Regex.Replace($"{Path.GetFileNameWithoutExtension(DocumentPath)}-{ViewFamilyName}-{View.Name}.udatasmith", @"\s+", "_");
+					FilePaths.Add(View.Id, Path.Combine(SavePath, FileName));
+					ViewsToExport.Add(View);
+				}
 			}
 
 			// Prevent user interaction with the active 3D view to avoid the termination of the custom export,
 			// without Revit providing any kind of internal or external feedback.
-			EnableViewWindow(in_commandData.Application, false);
+			EnableViewWindow(InCommandData.Application, false);
 
 			// Create a custom export context for command Export to Unreal Datasmith.
-			FDatasmithRevitExportContext exportContext = new FDatasmithRevitExportContext(in_commandData.Application.Application, document, filePath, exportOptions);
+			FDatasmithRevitExportContext ExportContext = new FDatasmithRevitExportContext(InCommandData.Application.Application, Doc, FilePaths, ExportOptions);
 
 			// Export the active 3D View to the given Unreal Datasmith file.
-			using( CustomExporter customExporter = new CustomExporter(document, exportContext) )
+			using( CustomExporter Exporter = new CustomExporter(Doc, ExportContext) )
 			{
 				// Add a progress bar callback.
 				// application.ProgressChanged += exportContext.HandleProgressChanged;
@@ -136,22 +174,25 @@ namespace DatasmithRevitExporter
 					// The export process will exclude output of geometric objects such as faces and curves,
 					// but the context needs to receive the calls related to Faces or Curves to gather data.
 					// The context always receive their tessellated geometry in form of polymeshes or lines.
-					customExporter.IncludeGeometricObjects = true;
+					Exporter.IncludeGeometricObjects = true;
 
 					// The export process should stop in case an error occurs during any of the exporting methods.
-					customExporter.ShouldStopOnError = true;
-
-					// Initiate the export process for the active 3D View.
-					#if REVIT_API_2020
-						customExporter.Export(activeView as Autodesk.Revit.DB.View);
-					#else
-						customExporter.Export(activeView);
-					#endif
+					Exporter.ShouldStopOnError = true;
+
+					// Initiate the export process for all 3D views.
+					foreach (var view in ViewsToExport)
+					{
+#if REVIT_API_2020
+						Exporter.Export(view as Autodesk.Revit.DB.View);
+#else
+						Exporter.Export(view);
+#endif
+					}
 				}
 				catch( System.Exception exception )
 				{
-					out_commandMessage = string.Format("Cannot export the active 3D view:\n\n{0}\n\n{1}", exception.Message, exception.StackTrace);
-					MessageBox.Show(out_commandMessage, DIALOG_CAPTION, MessageBoxButtons.OK, MessageBoxIcon.Error);
+					OutCommandMessage = string.Format("Cannot export the 3D view:\n\n{0}\n\n{1}", exception.Message, exception.StackTrace);
+					MessageBox.Show(OutCommandMessage, DIALOG_CAPTION, MessageBoxButtons.OK, MessageBoxIcon.Error);
 					return Result.Failed;
 				}
 				finally
@@ -160,11 +201,11 @@ namespace DatasmithRevitExporter
 					// application.ProgressChanged -= exportContext.HandleProgressChanged;
 
 					// Restore user interaction with the active 3D view.
-					EnableViewWindow(in_commandData.Application, true);
+					EnableViewWindow(InCommandData.Application, true);
 
-					if (exportContext.GetMessages().Count > 0)
+					if (ExportContext.GetMessages().Count > 0)
 					{
-						string Messages = string.Join($"{System.Environment.NewLine}", exportContext.GetMessages());
+						string Messages = string.Join($"{System.Environment.NewLine}", ExportContext.GetMessages());
 						DatasmithRevitApplication.ShowExportMessages(Messages);
 					}
 				}
@@ -185,44 +226,44 @@ namespace DatasmithRevitExporter
 			bool          in_enable
 		)
         {
-			#if REVIT_API_2018
-				IntPtr mainWindowHandle = Process.GetCurrentProcess().MainWindowHandle;
-			#else
-				IntPtr mainWindowHandle = in_application.MainWindowHandle;
-			#endif
+#if REVIT_API_2018
+				IntPtr MainWindowHandle = Process.GetCurrentProcess().MainWindowHandle;
+#else
+				IntPtr MainWindowHandle = in_application.MainWindowHandle;
+#endif
 
 			// "AfxFrameOrView140u" is the window class name of Revit active 3D view.
-			IntPtr viewWindowHandle = FindChildWindow(mainWindowHandle, "AfxFrameOrView140u");
+			IntPtr ViewWindowHandle = FindChildWindow(MainWindowHandle, "AfxFrameOrView140u");
 
-            if (viewWindowHandle != IntPtr.Zero)
+            if (ViewWindowHandle != IntPtr.Zero)
 			{
-                EnableWindow(viewWindowHandle, in_enable);
+                EnableWindow(ViewWindowHandle, in_enable);
 			}
         }
 
         private IntPtr FindChildWindow(
-			IntPtr in_parentWindowHandle,
-			string in_windowClassName
+			IntPtr InParentWindowHandle,
+			string InWindowClassName
 		)
         {
-            IntPtr windowHandle = FindWindowEx(in_parentWindowHandle, IntPtr.Zero, in_windowClassName, null);
+            IntPtr WindowHandle = FindWindowEx(InParentWindowHandle, IntPtr.Zero, InWindowClassName, null);
 			
-            if (windowHandle == IntPtr.Zero)
+            if (WindowHandle == IntPtr.Zero)
             {
-                IntPtr windowHandleChild = FindWindowEx(in_parentWindowHandle, IntPtr.Zero, null, null);
+                IntPtr WindowHandleChild = FindWindowEx(InParentWindowHandle, IntPtr.Zero, null, null);
 
-                while (windowHandleChild != IntPtr.Zero && windowHandle == IntPtr.Zero)
+                while (WindowHandleChild != IntPtr.Zero && WindowHandle == IntPtr.Zero)
                 {
-                    windowHandle = FindChildWindow(windowHandleChild, in_windowClassName);
+                    WindowHandle = FindChildWindow(WindowHandleChild, InWindowClassName);
 					
-                    if (windowHandle == IntPtr.Zero)
+                    if (WindowHandle == IntPtr.Zero)
                     {
-                        windowHandleChild = FindWindowEx(in_parentWindowHandle, windowHandleChild, null, null);
+                        WindowHandleChild = FindWindowEx(InParentWindowHandle, WindowHandleChild, null, null);
                     }
                 }
             }
 			
-            return windowHandle;
+            return WindowHandle;
         }
 	}
 }
diff --git a/Engine/Source/Programs/Enterprise/Datasmith/DatasmithRevitExporter/Private/DatasmithRevitExportContext.cs b/Engine/Source/Programs/Enterprise/Datasmith/DatasmithRevitExporter/Private/DatasmithRevitExportContext.cs
--- a/Engine/Source/Programs/Enterprise/Datasmith/DatasmithRevitExporter/Private/DatasmithRevitExportContext.cs
+++ b/Engine/Source/Programs/Enterprise/Datasmith/DatasmithRevitExporter/Private/DatasmithRevitExportContext.cs
@@ -30,8 +30,8 @@ namespace DatasmithRevitExporter
 		// Active Revit document being exported.
 		private Document RevitDocument = null;
 
-		// Datasmith output file path.
-		private string DatasmithFilePath = null;
+		// Datasmith file paths for each 3D view to be exported.
+		private Dictionary<ElementId, string> DatasmithFilePaths = null;
 
 		// Multi-line debug log.
 		private FDatasmithFacadeLog DebugLog = null;
@@ -54,16 +54,19 @@ namespace DatasmithRevitExporter
 		// List of messages generated during the export process.
 		private List<string> MessageList = new List<string>();
 
+		// The file path for the view that is currently being exported.
+		private string CurrentDatasmithFilePath = null;
+
 		public FDatasmithRevitExportContext(
-			Application                 InApplication,       // running Revit application
-			Document                    InDocument,          // active Revit document
-			string                      InDatasmithFilePath, // Datasmith output file path
-			DatasmithRevitExportOptions InExportOptions      // Unreal Datasmith export options
+			Application						InApplication,        // running Revit application
+			Document						InDocument,           // active Revit document
+			Dictionary<ElementId, string>	InDatasmithFilePaths, // Datasmith output file path
+			DatasmithRevitExportOptions		InExportOptions       // Unreal Datasmith export options
 		)
 		{
-			ProductVersion    = InApplication.VersionNumber;
-			RevitDocument     = InDocument;
-			DatasmithFilePath = InDatasmithFilePath;
+			ProductVersion     = InApplication.VersionNumber;
+			RevitDocument      = InDocument;
+			DatasmithFilePaths = InDatasmithFilePaths;
 
 			// Get the Unreal Datasmith export options.
 			DebugLog            = InExportOptions.GetWriteLogFile() ? new FDatasmithFacadeLog() : null;
@@ -111,7 +114,7 @@ namespace DatasmithRevitExporter
 		{
 			if (DebugLog != null)
 			{
-				DebugLog.WriteFile(DatasmithFilePath.Replace(".udatasmith", ".log"));
+				DebugLog.WriteFile(CurrentDatasmithFilePath.Replace(".udatasmith", ".log"));
 			}
 		}
 
@@ -136,8 +139,14 @@ namespace DatasmithRevitExporter
 			// Create an empty Datasmith scene.
 			DatasmithScene = new FDatasmithFacadeScene(HOST_NAME, VENDOR_NAME, PRODUCT_NAME, ProductVersion);
 
+			View3D ViewToExport = RevitDocument.GetElement(InViewNode.ViewId) as View3D;
+			if (!DatasmithFilePaths.TryGetValue(ViewToExport.Id, out CurrentDatasmithFilePath))
+			{
+				return RenderNodeAction.Skip; // TODO log error?
+			}
+
 			// Add a new camera actor to the Datasmith scene for the 3D view camera.
-			AddCameraActor(RevitDocument.GetElement(InViewNode.ViewId) as View3D, InViewNode.GetCameraInfo());
+			AddCameraActor(ViewToExport, InViewNode.GetCameraInfo());
 
 			// Keep track of the active Revit document being exported.
 			PushDocument(RevitDocument);
@@ -159,7 +168,7 @@ namespace DatasmithRevitExporter
 			DatasmithScene.Optimize();
 
 			// Build and export the Datasmith scene instance and its scene element assets.
-			DatasmithScene.ExportScene(DatasmithFilePath);
+			DatasmithScene.ExportScene(CurrentDatasmithFilePath);
 
 			// Dispose of the Datasmith scene.
 			DatasmithScene = null;
@@ -1045,9 +1054,9 @@ namespace DatasmithRevitExporter
 				FamilyInstance CurrentFamilyInstance = CurrentElement as FamilyInstance;
 				if (CurrentFamilyInstance != null)
 				{
-					IOActor.AddTag($"Revit.DB.FamiliyInstance.Mirrored.{CurrentFamilyInstance.Mirrored}");
-					IOActor.AddTag($"Revit.DB.FamiliyInstance.HandFlipped.{CurrentFamilyInstance.HandFlipped}");
-					IOActor.AddTag($"Revit.DB.FamiliyInstance.FaceFlipped.{CurrentFamilyInstance.FacingFlipped}");
+					IOActor.AddTag($"Revit.DB.FamilyInstance.Mirrored.{CurrentFamilyInstance.Mirrored}");
+					IOActor.AddTag($"Revit.DB.FamilyInstance.HandFlipped.{CurrentFamilyInstance.HandFlipped}");
+					IOActor.AddTag($"Revit.DB.FamilyInstance.FaceFlipped.{CurrentFamilyInstance.FacingFlipped}");
 
 					if (CurrentFamilyInstance.Host != null)
 					{
@@ -1519,6 +1528,10 @@ namespace DatasmithRevitExporter
 			{
 				return CurrentDocument.GetElement((SourceElement as ContinuousRail).HostRailingId);
 			}
+			else if (SourceElement.GetType().IsSubclassOf(typeof(InsulationLiningBase)))
+			{
+				return CurrentDocument.GetElement((SourceElement as InsulationLiningBase).HostElementId);
+			}
 
 			return null;
 		}
diff --git a/Engine/Source/Programs/Enterprise/Datasmith/DatasmithRevitExporter/Private/DatasmithRevitExportOptions.cs b/Engine/Source/Programs/Enterprise/Datasmith/DatasmithRevitExporter/Private/DatasmithRevitExportOptions.cs
--- a/Engine/Source/Programs/Enterprise/Datasmith/DatasmithRevitExporter/Private/DatasmithRevitExportOptions.cs
+++ b/Engine/Source/Programs/Enterprise/Datasmith/DatasmithRevitExporter/Private/DatasmithRevitExportOptions.cs
@@ -1,18 +1,23 @@
 ﻿// Copyright 1998-2019 Epic Games, Inc. All Rights Reserved.
 
 using System;
+using System.Collections.Generic;
 using System.Drawing;
+using System.Linq;
 using System.Windows.Forms;
 
-
 namespace DatasmithRevitExporter
 {
 	public class DatasmithRevitExportOptions : Form
 	{
-		private CheckBox      WriteLogFile;
-		private NumericUpDown LevelOfTessellation;
+		private CheckBox		WriteLogFile;
+		private NumericUpDown	LevelOfTessellation;
+		private CheckedListBox	ViewsList;
+
+		public List<Autodesk.Revit.DB.View3D> Selected3DViews { get; private set; } = new List<Autodesk.Revit.DB.View3D>();
+		private List<Autodesk.Revit.DB.View3D> All3DViews;
 
-		public DatasmithRevitExportOptions()
+		public DatasmithRevitExportOptions(Autodesk.Revit.DB.Document document)
 		{
 			ToolTip OptionToolTip = new ToolTip();
 			OptionToolTip.AutoPopDelay = 10000; // milliseconds
@@ -55,7 +60,7 @@ namespace DatasmithRevitExporter
 			LevelOfTessellation.Name = "LevelOfTessellation";
 			LevelOfTessellation.Minimum = -1;
 			LevelOfTessellation.Maximum = 15;
-			LevelOfTessellation.Value   = 8;
+			LevelOfTessellation.Value = 8;
 			LevelOfTessellation.Anchor = AnchorStyles.Left;
 			LevelOfTessellation.AutoSize = true;
 			LevelOfTessellation.Location = new Point(113, 23);
@@ -64,7 +69,7 @@ namespace DatasmithRevitExporter
 			LevelOfTessellation.TextAlign = HorizontalAlignment.Right;
 
 			OptionToolTip.SetToolTip(LevelOfTessellation,
-				                     "-1 :\tRevit will use its default algorithm, which is based on output resolution.\n" +
+									 "-1 :\tRevit will use its default algorithm, which is based on output resolution.\n" +
 									 "0 to 15 :\tRevit will use the suggested level of detail when tessellating faces.\n" +
 									 "\tUsing a value close to the middle of the range yields a very reasonable tessellation.\n" +
 									 "\tRevit uses level 8 as its 'normal' level of detail.");
@@ -79,13 +84,40 @@ namespace DatasmithRevitExporter
 			OKButton.Size = new Size(75, 23);
 			OKButton.TabIndex = 4;
 			OKButton.UseVisualStyleBackColor = true;
+			OKButton.Click += (s, e) => 
+			{
+				for (int ItemIndex = 0; ItemIndex < ViewsList.Items.Count; ++ItemIndex)
+				{
+					if (ViewsList.GetItemChecked(ItemIndex))
+					{
+						Selected3DViews.Add(All3DViews[ItemIndex]);
+					}
+				}
+				if (Selected3DViews.Count == 0)
+				{
+					string message = "No views selected for export! Please select at least one 3D view.";
+					MessageBox.Show(message, "View selection", MessageBoxButtons.OK, MessageBoxIcon.Warning);
+					DialogResult = DialogResult.None;
+				}
+			};
+
+			ViewsList = new CheckedListBox();
+			ViewsList.Size = new Size(300, 100);
+
+			Label ViewsLabel = new Label();
+			ViewsLabel.Name = "ViewsLabel";
+			ViewsLabel.Text = "Select 3D Views";
+			ViewsLabel.Anchor = AnchorStyles.Left;
+			ViewsLabel.AutoSize = true;
 
 			TableLayoutPanel DialogLayout = new TableLayoutPanel();
 			DialogLayout.Name = "DialogLayout";
 			DialogLayout.ColumnCount = 2;
 			DialogLayout.ColumnStyles.Add(new ColumnStyle());
 			DialogLayout.ColumnStyles.Add(new ColumnStyle());
-			DialogLayout.RowCount = 3;
+			DialogLayout.RowCount = 5;
+			DialogLayout.RowStyles.Add(new RowStyle());
+			DialogLayout.RowStyles.Add(new RowStyle(SizeType.Percent, 100F));
 			DialogLayout.RowStyles.Add(new RowStyle());
 			DialogLayout.RowStyles.Add(new RowStyle());
 			DialogLayout.RowStyles.Add(new RowStyle());
@@ -96,11 +128,15 @@ namespace DatasmithRevitExporter
 			DialogLayout.Size = new Size(300, 100);
 			DialogLayout.TabIndex = 0;
 
-			DialogLayout.Controls.Add(WriteLogFileLabel,        0, 0);
-			DialogLayout.Controls.Add(WriteLogFile,             1, 0);
-			DialogLayout.Controls.Add(LevelOfTessellationLabel, 0, 1);
-			DialogLayout.Controls.Add(LevelOfTessellation,      1, 1);
-			DialogLayout.Controls.Add(OKButton,                 1, 2);
+			DialogLayout.Controls.Add(ViewsLabel,				0, 0);
+			DialogLayout.Controls.Add(ViewsList,				0, 1);
+			DialogLayout.Controls.Add(WriteLogFileLabel,		0, 2);
+			DialogLayout.Controls.Add(WriteLogFile,				1, 2);
+			DialogLayout.Controls.Add(LevelOfTessellationLabel, 0, 3);
+			DialogLayout.Controls.Add(LevelOfTessellation,		1, 3);
+			DialogLayout.Controls.Add(OKButton,					1, 4);
+
+			DialogLayout.SetColumnSpan(ViewsList, 2);
 
 			Name = "UnrealDatasmithExportOptions";
 			Text = "Unreal Datasmith Export - Debug Options";
@@ -118,6 +154,15 @@ namespace DatasmithRevitExporter
 			StartPosition = FormStartPosition.CenterParent;
 
 			Controls.Add(DialogLayout);
+
+			// Get all 3D views that can be exported.
+			All3DViews = new Autodesk.Revit.DB.FilteredElementCollector(document).OfClass(typeof(Autodesk.Revit.DB.View3D)).Cast<Autodesk.Revit.DB.View3D>().ToList();
+			All3DViews.RemoveAll(view => (view.IsTemplate || !view.CanBePrinted));
+
+			foreach (var View in All3DViews)
+			{
+				ViewsList.Items.Add(View.Name, document.ActiveView.Id == View.Id);
+			}
 		}
 
 		public bool GetWriteLogFile()
