Copy up //UE4/Dev-Enterprise/...at  Changelist 10508152  to //UE4/Dev-Tools-Staging/...

diff --git a/Engine/Source/Programs/Enterprise/Datasmith/DatasmithRevitExporter/Private/DatasmithRevitExportOptions.cs b/Engine/Source/Programs/Enterprise/Datasmith/DatasmithRevitExporter/Private/DatasmithRevitExportOptions.cs
--- a/Engine/Source/Programs/Enterprise/Datasmith/DatasmithRevitExporter/Private/DatasmithRevitExportOptions.cs
+++ b/Engine/Source/Programs/Enterprise/Datasmith/DatasmithRevitExporter/Private/DatasmithRevitExportOptions.cs
@@ -1,18 +1,23 @@
 ﻿// Copyright 1998-2019 Epic Games, Inc. All Rights Reserved.
 
 using System;
+using System.Collections.Generic;
 using System.Drawing;
+using System.Linq;
 using System.Windows.Forms;
 
-
 namespace DatasmithRevitExporter
 {
 	public class DatasmithRevitExportOptions : Form
 	{
-		private CheckBox      WriteLogFile;
-		private NumericUpDown LevelOfTessellation;
+		private CheckBox		WriteLogFile;
+		private NumericUpDown	LevelOfTessellation;
+		private CheckedListBox	ViewsList;
+
+		public List<Autodesk.Revit.DB.View3D> Selected3DViews { get; private set; } = new List<Autodesk.Revit.DB.View3D>();
+		private List<Autodesk.Revit.DB.View3D> All3DViews;
 
-		public DatasmithRevitExportOptions()
+		public DatasmithRevitExportOptions(Autodesk.Revit.DB.Document document)
 		{
 			ToolTip OptionToolTip = new ToolTip();
 			OptionToolTip.AutoPopDelay = 10000; // milliseconds
@@ -55,7 +60,7 @@ namespace DatasmithRevitExporter
 			LevelOfTessellation.Name = "LevelOfTessellation";
 			LevelOfTessellation.Minimum = -1;
 			LevelOfTessellation.Maximum = 15;
-			LevelOfTessellation.Value   = 8;
+			LevelOfTessellation.Value = 8;
 			LevelOfTessellation.Anchor = AnchorStyles.Left;
 			LevelOfTessellation.AutoSize = true;
 			LevelOfTessellation.Location = new Point(113, 23);
@@ -64,7 +69,7 @@ namespace DatasmithRevitExporter
 			LevelOfTessellation.TextAlign = HorizontalAlignment.Right;
 
 			OptionToolTip.SetToolTip(LevelOfTessellation,
-				                     "-1 :\tRevit will use its default algorithm, which is based on output resolution.\n" +
+									 "-1 :\tRevit will use its default algorithm, which is based on output resolution.\n" +
 									 "0 to 15 :\tRevit will use the suggested level of detail when tessellating faces.\n" +
 									 "\tUsing a value close to the middle of the range yields a very reasonable tessellation.\n" +
 									 "\tRevit uses level 8 as its 'normal' level of detail.");
@@ -79,13 +84,40 @@ namespace DatasmithRevitExporter
 			OKButton.Size = new Size(75, 23);
 			OKButton.TabIndex = 4;
 			OKButton.UseVisualStyleBackColor = true;
+			OKButton.Click += (s, e) => 
+			{
+				for (int ItemIndex = 0; ItemIndex < ViewsList.Items.Count; ++ItemIndex)
+				{
+					if (ViewsList.GetItemChecked(ItemIndex))
+					{
+						Selected3DViews.Add(All3DViews[ItemIndex]);
+					}
+				}
+				if (Selected3DViews.Count == 0)
+				{
+					string message = "No views selected for export! Please select at least one 3D view.";
+					MessageBox.Show(message, "View selection", MessageBoxButtons.OK, MessageBoxIcon.Warning);
+					DialogResult = DialogResult.None;
+				}
+			};
+
+			ViewsList = new CheckedListBox();
+			ViewsList.Size = new Size(300, 100);
+
+			Label ViewsLabel = new Label();
+			ViewsLabel.Name = "ViewsLabel";
+			ViewsLabel.Text = "Select 3D Views";
+			ViewsLabel.Anchor = AnchorStyles.Left;
+			ViewsLabel.AutoSize = true;
 
 			TableLayoutPanel DialogLayout = new TableLayoutPanel();
 			DialogLayout.Name = "DialogLayout";
 			DialogLayout.ColumnCount = 2;
 			DialogLayout.ColumnStyles.Add(new ColumnStyle());
 			DialogLayout.ColumnStyles.Add(new ColumnStyle());
-			DialogLayout.RowCount = 3;
+			DialogLayout.RowCount = 5;
+			DialogLayout.RowStyles.Add(new RowStyle());
+			DialogLayout.RowStyles.Add(new RowStyle(SizeType.Percent, 100F));
 			DialogLayout.RowStyles.Add(new RowStyle());
 			DialogLayout.RowStyles.Add(new RowStyle());
 			DialogLayout.RowStyles.Add(new RowStyle());
@@ -96,11 +128,15 @@ namespace DatasmithRevitExporter
 			DialogLayout.Size = new Size(300, 100);
 			DialogLayout.TabIndex = 0;
 
-			DialogLayout.Controls.Add(WriteLogFileLabel,        0, 0);
-			DialogLayout.Controls.Add(WriteLogFile,             1, 0);
-			DialogLayout.Controls.Add(LevelOfTessellationLabel, 0, 1);
-			DialogLayout.Controls.Add(LevelOfTessellation,      1, 1);
-			DialogLayout.Controls.Add(OKButton,                 1, 2);
+			DialogLayout.Controls.Add(ViewsLabel,				0, 0);
+			DialogLayout.Controls.Add(ViewsList,				0, 1);
+			DialogLayout.Controls.Add(WriteLogFileLabel,		0, 2);
+			DialogLayout.Controls.Add(WriteLogFile,				1, 2);
+			DialogLayout.Controls.Add(LevelOfTessellationLabel, 0, 3);
+			DialogLayout.Controls.Add(LevelOfTessellation,		1, 3);
+			DialogLayout.Controls.Add(OKButton,					1, 4);
+
+			DialogLayout.SetColumnSpan(ViewsList, 2);
 
 			Name = "UnrealDatasmithExportOptions";
 			Text = "Unreal Datasmith Export - Debug Options";
@@ -118,6 +154,15 @@ namespace DatasmithRevitExporter
 			StartPosition = FormStartPosition.CenterParent;
 
 			Controls.Add(DialogLayout);
+
+			// Get all 3D views that can be exported.
+			All3DViews = new Autodesk.Revit.DB.FilteredElementCollector(document).OfClass(typeof(Autodesk.Revit.DB.View3D)).Cast<Autodesk.Revit.DB.View3D>().ToList();
+			All3DViews.RemoveAll(view => (view.IsTemplate || !view.CanBePrinted));
+
+			foreach (var View in All3DViews)
+			{
+				ViewsList.Items.Add(View.Name, document.ActiveView.Id == View.Id);
+			}
 		}
 
 		public bool GetWriteLogFile()
