Copy of 10340048 from //UE4/Dev-Enterprise:

Datasmith: flag exporters with bBuildInSolutionByDefault = false;
Fixes Github build when the whole solution is being built

diff --git a/Engine/Source/Programs/Enterprise/Datasmith/DatasmithFacadeCSharp/DatasmithFacadeCSharp.Target.cs b/Engine/Source/Programs/Enterprise/Datasmith/DatasmithFacadeCSharp/DatasmithFacadeCSharp.Target.cs
--- a/Engine/Source/Programs/Enterprise/Datasmith/DatasmithFacadeCSharp/DatasmithFacadeCSharp.Target.cs
+++ b/Engine/Source/Programs/Enterprise/Datasmith/DatasmithFacadeCSharp/DatasmithFacadeCSharp.Target.cs
@@ -11,6 +11,7 @@ public class DatasmithFacadeCSharpTarget : TargetRules
 	{
 		Type = TargetType.Program;
 		SolutionDirectory = "Programs/Datasmith";
+		bBuildInSolutionByDefault = false;
 
 		LaunchModuleName = "DatasmithFacadeCSharp";
 		ExeBinariesSubFolder = "DatasmithFacadeCSharp";
diff --git a/Engine/Source/Programs/Enterprise/Datasmith/DatasmithMaxExporter/DatasmithMax2017.Target.cs b/Engine/Source/Programs/Enterprise/Datasmith/DatasmithMaxExporter/DatasmithMax2017.Target.cs
--- a/Engine/Source/Programs/Enterprise/Datasmith/DatasmithMaxExporter/DatasmithMax2017.Target.cs
+++ b/Engine/Source/Programs/Enterprise/Datasmith/DatasmithMaxExporter/DatasmithMax2017.Target.cs
@@ -9,6 +9,7 @@ public abstract class DatasmithMaxBaseTarget : TargetRules
 	{
 		Type = TargetType.Program;
 		SolutionDirectory = "Programs/Datasmith";
+		bBuildInSolutionByDefault = false;
 		bLegalToDistributeBinary = true;
 
 		bShouldCompileAsDLL = true;
diff --git a/Engine/Source/Programs/Enterprise/Datasmith/DatasmithRevitExporter/DatasmithRevit2018.Target.cs b/Engine/Source/Programs/Enterprise/Datasmith/DatasmithRevitExporter/DatasmithRevit2018.Target.cs
--- a/Engine/Source/Programs/Enterprise/Datasmith/DatasmithRevitExporter/DatasmithRevit2018.Target.cs
+++ b/Engine/Source/Programs/Enterprise/Datasmith/DatasmithRevitExporter/DatasmithRevit2018.Target.cs
@@ -10,6 +10,7 @@ public abstract class DatasmithRevitBaseTarget : TargetRules
 	{
 		Type = TargetType.Program;
 		SolutionDirectory = "Programs/Datasmith";
+		bBuildInSolutionByDefault = false;
 
 		string RevitVersionString = GetVersion();
 		string ProjectName = "DatasmithRevit" + RevitVersionString;
diff --git a/Engine/Source/Programs/Enterprise/Datasmith/DatasmithSketchUpExporter/DatasmithSketchUp2017.Target.cs b/Engine/Source/Programs/Enterprise/Datasmith/DatasmithSketchUpExporter/DatasmithSketchUp2017.Target.cs
--- a/Engine/Source/Programs/Enterprise/Datasmith/DatasmithSketchUpExporter/DatasmithSketchUp2017.Target.cs
+++ b/Engine/Source/Programs/Enterprise/Datasmith/DatasmithSketchUpExporter/DatasmithSketchUp2017.Target.cs
@@ -9,6 +9,7 @@ public abstract class DatasmithSketchUpBaseTarget : TargetRules
 	{
 		Type = TargetType.Program;
 		SolutionDirectory = "Programs/Datasmith";
+		bBuildInSolutionByDefault = false;
 
 		bShouldCompileAsDLL = true;
 		LinkType = TargetLinkType.Monolithic;
