More diplomatic relations between Dev-Core and Main to avoid a conflict.

#rnx

diff --git a/Engine/Source/Programs/Enterprise/Datasmith/DatasmithRevitExporter/Private/DatasmithRevitExportContext.cs b/Engine/Source/Programs/Enterprise/Datasmith/DatasmithRevitExporter/Private/DatasmithRevitExportContext.cs
new file mode 100644
--- /dev/null
+++ b/Engine/Source/Programs/Enterprise/Datasmith/DatasmithRevitExporter/Private/DatasmithRevitExportContext.cs
@@ -0,0 +1,1681 @@
+﻿// Copyright 1998-2019 Epic Games, Inc. All Rights Reserved.
+
+using System;
+using System.Collections.Generic;
+using System.IO;
+
+using Autodesk.Revit.ApplicationServices;
+using Autodesk.Revit.DB;
+using Autodesk.Revit.DB.Architecture;
+using Autodesk.Revit.DB.Events;
+using Autodesk.Revit.DB.Structure;
+using Autodesk.Revit.DB.Visual;
+
+namespace DatasmithRevitExporter
+{
+	// Custom export context for command Export to Unreal Datasmith. 
+	public class FDatasmithRevitExportContext : IPhotoRenderContext
+	{
+		// Revit application information for Datasmith.
+		private const string HOST_NAME    = "Revit";
+		private const string VENDOR_NAME  = "Autodesk Inc.";
+		private const string PRODUCT_NAME = "Revit";
+
+		// Revit to Datasmith unit convertion factor.
+		private const float CENTIMETERS_PER_FOOT = 30.48F;
+
+		// Running Revit version.
+		private string ProductVersion = "";
+
+		// Active Revit document being exported.
+		private Document RevitDocument = null;
+
+		// Datasmith output file path.
+		private string DatasmithFilePath = null;
+
+		// Multi-line debug log.
+		private FDatasmithFacadeLog DebugLog = null;
+
+		// Level of detail when tessellating faces (between -1 and 15).
+		private int LevelOfTessellation = 8;
+
+		// Stack of world Transforms for the Revit instances being processed.
+		private Stack<Transform> WorldTransformStack = new Stack<Transform>();
+
+		// Datasmith scene being built.
+		private FDatasmithFacadeScene DatasmithScene = null;
+
+		// Stack of Revit document data being processed.
+		private Stack<FDocumentData> DocumentDataStack = new Stack<FDocumentData>();
+
+		// List of extra search paths for Revit texture files.
+		private IList<string> ExtraTexturePaths = new List<string>();
+		
+		// List of messages generated during the export process.
+		private List<string> MessageList = new List<string>();
+
+		public FDatasmithRevitExportContext(
+			Application                 InApplication,       // running Revit application
+			Document                    InDocument,          // active Revit document
+			string                      InDatasmithFilePath, // Datasmith output file path
+			DatasmithRevitExportOptions InExportOptions      // Unreal Datasmith export options
+		)
+		{
+			ProductVersion    = InApplication.VersionNumber;
+			RevitDocument     = InDocument;
+			DatasmithFilePath = InDatasmithFilePath;
+
+			// Get the Unreal Datasmith export options.
+			DebugLog            = InExportOptions.GetWriteLogFile() ? new FDatasmithFacadeLog() : null;
+			LevelOfTessellation = InExportOptions.GetLevelOfTessellation();
+		}
+
+		// Progress bar callback.
+		public void HandleProgressChanged(
+			object                   InSender,
+			ProgressChangedEventArgs InArgs
+		)
+		{
+			// DebugLog.AddLine($"HandleProgressChanged: {InArgs.Stage} {InArgs.Position} {InArgs.UpperRange} {InArgs.Caption}");
+		}
+
+		public IList<string> GetMessages()
+		{
+			return MessageList;
+		}
+
+		//========================================================================================================================
+		// Implement IPhotoRenderContext interface methods called at times of drawing geometry as if executing the Render command.
+		// Only actual geometry suitable to appear in a rendered view will be processed and output.
+
+		// Start is called at the very start of the export process, still before the first entity of the model was send out.
+		public bool Start()
+		{
+			// Retrieve the list of extra search paths for Revit texture files.
+			RetrieveExtraTexturePaths();
+
+			// Set the coordinate system type of the world geometries and transforms.
+			// Revit uses a right-handed Z-up coordinate system.
+			FDatasmithFacadeElement.SetCoordinateSystemType(FDatasmithFacadeElement.ECoordinateSystemType.RightHandedZup);
+
+			// Set the scale factor from Revit world units to Datasmith centimeters.
+			// Revit uses foot as internal system unit for all 3D coordinates.
+			FDatasmithFacadeElement.SetWorldUnitScale(CENTIMETERS_PER_FOOT);
+
+			// We are ready to proceed with the export.
+			return true;
+		}
+
+		// Finish is called at the very end of the export process, after all entities were processed (or after the process was cancelled).
+		public void Finish()
+		{
+			if (DebugLog != null)
+			{
+				DebugLog.WriteFile(DatasmithFilePath.Replace(".udatasmith", ".log"));
+			}
+		}
+
+		// IsCanceled is queried at the beginning of every element.
+		public bool IsCanceled()
+		{
+			// Return whether or not the export process should be canceled.
+			return false;
+		}
+
+		// OnViewBegin marks the beginning of a 3D view to be exported.
+		public RenderNodeAction OnViewBegin(
+			ViewNode InViewNode // render node associated with the 3D view
+		)
+		{
+			// Set the level of detail when tessellating faces (between -1 and 15).
+			InViewNode.LevelOfDetail = LevelOfTessellation;
+
+			// Initialize the world transform for the 3D view being processed.
+			WorldTransformStack.Push(Transform.Identity);
+
+			// Create an empty Datasmith scene.
+			DatasmithScene = new FDatasmithFacadeScene(HOST_NAME, VENDOR_NAME, PRODUCT_NAME, ProductVersion);
+
+			// Add a new camera actor to the Datasmith scene for the 3D view camera.
+			AddCameraActor(RevitDocument.GetElement(InViewNode.ViewId) as View3D, InViewNode.GetCameraInfo());
+
+			// Keep track of the active Revit document being exported.
+			PushDocument(RevitDocument);
+
+			// We want to export the 3D view.
+			return RenderNodeAction.Proceed;
+		}
+
+		// OnViewEnd marks the end of a 3D view being exported.
+		// This method is invoked even for 3D views that were skipped.
+		public void OnViewEnd(
+			ElementId InElementId // exported 3D view ID
+		)
+		{
+			// Forget the active Revit document being exported.
+			PopDocument();
+
+			// Optimize the Datasmith actor hierarchy by removing the intermediate single child actors.
+			DatasmithScene.Optimize();
+
+			// Build and export the Datasmith scene instance and its scene element assets.
+			DatasmithScene.ExportScene(DatasmithFilePath);
+
+			// Dispose of the Datasmith scene.
+			DatasmithScene = null;
+			
+			// Forget the 3D view world transform.
+			WorldTransformStack.Pop();
+		}
+
+		// OnElementBegin marks the beginning of an element to be exported.
+		public RenderNodeAction OnElementBegin(
+			ElementId InElementId // exported element ID
+		)
+		{
+			Element CurrentElement = GetElement(InElementId);
+
+			if (CurrentElement != null)
+			{
+				// Keep track of the element being processed.
+				PushElement(CurrentElement, WorldTransformStack.Peek(), "Element Begin");
+			
+				// We want to export the element.
+				return RenderNodeAction.Proceed;
+			}
+
+			return RenderNodeAction.Skip;
+		}
+
+		// OnElementEnd marks the end of an element being exported.
+		// This method is invoked even for elements that were skipped.
+		public void OnElementEnd(
+			ElementId InElementId // exported element ID
+		)
+		{
+			if (GetElement(InElementId) != null)
+			{
+				// Forget the current element being exported.
+				PopElement("Element End");
+			}
+		}
+
+		// OnInstanceBegin marks the beginning of a family instance to be exported.
+		public RenderNodeAction OnInstanceBegin(
+			InstanceNode InInstanceNode // family instance output node
+		)
+		{
+			Element CurrentInstance = GetElement(InInstanceNode.GetSymbolId());
+
+			if (CurrentInstance != null)
+			{
+				// Keep track of the world transform for the instance being processed.
+				WorldTransformStack.Push(WorldTransformStack.Peek().Multiply(InInstanceNode.GetTransform()));
+
+				ElementType CurrentInstanceType = CurrentInstance as ElementType;
+
+				// Keep track of the instance being processed.
+				if (CurrentInstanceType != null)
+				{
+					PushInstance(CurrentInstanceType, WorldTransformStack.Peek(), "Instance Begin");
+				}
+				else
+				{
+					PushElement(CurrentInstance, WorldTransformStack.Peek(), "Symbol Begin");
+				}
+
+				// We want to process the instance.
+				return RenderNodeAction.Proceed;
+			}
+
+			return RenderNodeAction.Skip;
+		}
+
+		// OnInstanceEnd marks the end of a family instance being exported.
+		// This method is invoked even for family instances that were skipped.
+		public void OnInstanceEnd(
+			InstanceNode InInstanceNode // family instance output node
+		)
+		{
+			Element CurrentInstance = GetElement(InInstanceNode.GetSymbolId());
+
+			if (CurrentInstance != null)
+			{
+				// Forget the current instance being exported.
+				if (CurrentInstance as ElementType != null)
+				{
+					PopInstance("Instance End");
+				}
+				else
+				{
+					PopElement("Symbol End");
+				}
+
+				// Forget the current world transform.
+				WorldTransformStack.Pop();
+			}
+		}
+
+		// OnLinkBegin marks the beginning of a link instance to be exported.
+		public RenderNodeAction OnLinkBegin(
+			LinkNode InLinkNode // linked Revit document output node
+		)
+		{
+			ElementType CurrentInstanceType = GetElement(InLinkNode.GetSymbolId()) as ElementType;
+
+			if (CurrentInstanceType != null)
+			{
+				// Keep track of the world transform for the instance being processed.
+				WorldTransformStack.Push(WorldTransformStack.Peek().Multiply(InLinkNode.GetTransform()));
+
+				// Keep track of the instance being processed.
+				PushInstance(CurrentInstanceType, WorldTransformStack.Peek(), "Link Begin");
+			}
+
+			Document LinkedDocument = InLinkNode.GetDocument();
+
+			if (LinkedDocument != null)
+			{
+				// Keep track of the linked document being processed.
+				PushDocument(LinkedDocument);
+			}
+
+			return (CurrentInstanceType != null && LinkedDocument != null) ? RenderNodeAction.Proceed : RenderNodeAction.Skip;
+		}
+
+		// OnLinkEnd marks the end of a link instance being exported.
+		// This method is invoked even for link instances that were skipped.
+		public void OnLinkEnd(
+			LinkNode InLinkNode // linked Revit document output node
+		)
+		{
+			if (InLinkNode.GetDocument() != null)
+			{
+				// Forget the current linked document being exported.
+				PopDocument();
+			}
+
+			if (GetElement(InLinkNode.GetSymbolId()) as ElementType != null)
+			{
+				// Forget the current instance being exported.
+				PopInstance("Link End");
+
+				// Forget the current world transform.
+				WorldTransformStack.Pop();
+			}
+		}
+
+		// OnLight marks the beginning of export of a light which is enabled for rendering.
+		// This method is only called for interface IPhotoRenderContext.
+		public void OnLight(
+			LightNode InLightNode // light output node
+		)
+		{
+			// Keep track of the world transform for the light being processed.
+			WorldTransformStack.Push(WorldTransformStack.Peek().Multiply(InLightNode.GetTransform()));
+
+			// Add a light actor in the hierarchy of Datasmith actors being processed.
+			AddLightActor(WorldTransformStack.Peek(), InLightNode.GetAsset());
+
+			// Forget the current light world transform.
+			WorldTransformStack.Pop();
+		}
+		
+		// OnRPC marks the beginning of export of an RPC object.
+		// This method is only called for interface IPhotoRenderContext.
+		public void OnRPC(
+			RPCNode InRPCNode // RPC content output node
+		)
+		{
+			// We ignore the RPC node local transform since the RPC location point will be used later.
+
+			// Add an RPC mesh actor in the hierarchy of Datasmith actors being processed.
+			AddRPCActor(WorldTransformStack.Peek(), InRPCNode.GetAsset());
+		}
+
+		// OnFaceBegin marks the beginning of a Face to be exported.
+		// This method is invoked only when the custom exporter was set up to include geometric objects in the output stream (IncludeGeometricObjects).
+		public RenderNodeAction OnFaceBegin(
+			FaceNode InFaceNode // face output node
+		)
+		{
+			// We want to receive geometry (polymesh) for this face.
+			return RenderNodeAction.Proceed;
+		}
+
+		// OnFaceEnd marks the end of the current face being exported.
+		// This method is invoked only when the custom exporter was set up to include geometric objects in the output stream (IncludeGeometricObjects).
+		// This method is invoked even for faces that were skipped.
+		public void OnFaceEnd(
+			FaceNode InFaceNode // face output node
+		)
+		{
+			// Nothing to do here.
+		}
+
+		// OnMaterial marks a change of the material.
+		// This method can be invoked for every single out-coming mesh even when the material has not actually changed.
+		public void OnMaterial(
+			MaterialNode InMaterialNode // current material output node
+		)
+		{
+			SetMaterial(InMaterialNode, ExtraTexturePaths);
+		}
+
+		// OnPolymesh is called when a tessellated polymesh of a 3D face is being output.
+		// The current material is applied to the polymesh.
+		public void OnPolymesh(
+			PolymeshTopology InPolymeshNode // tessellated polymesh output node
+		)
+		{
+			if (IgnoreElementGeometry())
+			{
+				return;
+			}
+
+			// Retrieve the Datasmith mesh being processed.
+			FDatasmithFacadeMesh CurrentMesh = GetCurrentMesh();
+
+			// Retrieve the index of the current material and make the Datasmith mesh keep track of it.
+			int CurrentMaterialIndex = GetCurrentMaterialIndex();
+
+			int initialVertexCount = CurrentMesh.GetVertexCount();
+
+			// Add the vertex points (in right-handed Z-up coordinates) to the Datasmith mesh.
+			foreach (XYZ point in InPolymeshNode.GetPoints())
+			{
+				CurrentMesh.AddVertex((float) point.X, (float) point.Y, (float) point.Z);
+			}
+
+			// Add the vertex UV texture coordinates to the Datasmith mesh.
+			foreach (UV uv in InPolymeshNode.GetUVs())
+			{
+				CurrentMesh.AddUV(0, (float) uv.U, (float) -uv.V);
+			}
+
+			// Add the triangle vertex indexes to the Datasmith mesh.
+			foreach (PolymeshFacet facet in InPolymeshNode.GetFacets())
+			{
+				CurrentMesh.AddTriangle(initialVertexCount + facet.V1, initialVertexCount + facet.V2, initialVertexCount + facet.V3, CurrentMaterialIndex);
+			}
+			
+			// Add the triangle vertex normals (in right-handed Z-up coordinates) to the Datasmith mesh.
+			// Normals can be associated with either points or facets of the polymesh.
+			switch (InPolymeshNode.DistributionOfNormals)
+			{
+				case DistributionOfNormals.AtEachPoint:
+				{
+					IList<XYZ> normals = InPolymeshNode.GetNormals();
+					foreach (PolymeshFacet facet in InPolymeshNode.GetFacets())
+					{
+						XYZ normal1 = normals[facet.V1];
+						XYZ normal2 = normals[facet.V2];
+						XYZ normal3 = normals[facet.V3];
+
+						CurrentMesh.AddNormal((float) normal1.X, (float) normal1.Y, (float) normal1.Z);
+						CurrentMesh.AddNormal((float) normal2.X, (float) normal2.Y, (float) normal2.Z);
+						CurrentMesh.AddNormal((float) normal3.X, (float) normal3.Y, (float) normal3.Z);
+					}
+					break;
+				}
+				case DistributionOfNormals.OnePerFace:
+				{
+					XYZ normal = InPolymeshNode.GetNormals()[0];
+					for (int i = 0; i < 3 * InPolymeshNode.NumberOfFacets; i++)
+					{
+						CurrentMesh.AddNormal((float) normal.X, (float) normal.Y, (float) normal.Z);
+					}
+					break;
+				}
+				case DistributionOfNormals.OnEachFacet:
+				{
+					foreach (XYZ normal in InPolymeshNode.GetNormals())
+					{
+						CurrentMesh.AddNormal((float) normal.X, (float) normal.Y, (float) normal.Z);
+						CurrentMesh.AddNormal((float) normal.X, (float) normal.Y, (float) normal.Z);
+						CurrentMesh.AddNormal((float) normal.X, (float) normal.Y, (float) normal.Z);
+					}
+					break;
+				}
+			}
+		}
+
+		// End of IPhotoRenderContext interface method implementation.
+		//========================================================================================================================
+
+		// Retrieve the list of extra search paths for Revit texture files.
+		// This is done by reading user's Revit.ini file and searching for field AdditionalRenderAppearancePaths
+		// which contains search paths that Revit will use to locate texture files.
+		// Note that the behavior in Revit is to search in the directory itself and not in child sub-directories.
+		private void RetrieveExtraTexturePaths()
+		{
+			string UserSpecificDirectoryPath = Environment.GetFolderPath(Environment.SpecialFolder.UserProfile);
+
+			if (!string.IsNullOrEmpty(UserSpecificDirectoryPath) && Path.IsPathRooted(UserSpecificDirectoryPath) && Directory.Exists(UserSpecificDirectoryPath))
+			{
+				string FullRevitIniPath = $"{UserSpecificDirectoryPath}\\AppData\\Roaming\\Autodesk\\Revit\\Autodesk Revit {ProductVersion}\\Revit.ini";
+
+				if (File.Exists(FullRevitIniPath))
+				{
+					FileStream RevitIniStream = new FileStream(FullRevitIniPath, FileMode.Open, FileAccess.Read, FileShare.Read);
+
+					using (StreamReader RevitIniReader = new StreamReader(RevitIniStream))
+					{
+						string ConfigLine;
+
+						while ((ConfigLine = RevitIniReader.ReadLine()) != null)
+						{
+							if (ConfigLine.Contains("AdditionalRenderAppearancePaths"))
+							{
+								string[] SplitLineArray = ConfigLine.Split('=');
+
+								if (SplitLineArray.Length > 1)
+								{
+									string[] TexturePaths = SplitLineArray[1].Split('|');
+
+									foreach (string TexturePath in TexturePaths)
+									{
+										ExtraTexturePaths.Add(TexturePath);
+									}
+
+									break;
+								}
+							}
+						}
+					}
+				}
+			}
+		}
+
+		private Element GetElement(
+			ElementId InElementId
+		)
+		{
+			return DocumentDataStack.Peek().GetElement(InElementId);
+		}
+
+		private void PushDocument(
+			Document InDocument
+		)
+		{
+			DocumentDataStack.Push(new FDocumentData(InDocument, ref MessageList));
+
+			DocumentDataStack.Peek().AddLocationActors(WorldTransformStack.Peek());
+		}
+
+		private void PopDocument()
+		{
+			FDocumentData DocumentData = DocumentDataStack.Pop();
+
+			if (DocumentDataStack.Count == 0)
+			{
+				DocumentData.WrapupScene(DatasmithScene);
+			}
+			else
+			{
+				DocumentData.WrapupLink(DatasmithScene, DocumentDataStack.Peek().GetCurrentActor());
+			}
+
+		}
+
+		private void PushElement(
+			Element   InElement,
+			Transform InWorldTransform,
+			string    InLogLinePrefix
+		)
+		{
+			FDocumentData DocumentData = DocumentDataStack.Peek();
+
+			DocumentData.PushElement(InElement, InWorldTransform);
+			DocumentData.LogElement(DebugLog, InLogLinePrefix, +1);
+		}
+
+		private void PopElement(
+			string InLogLinePrefix
+		)
+		{
+			FDocumentData DocumentData = DocumentDataStack.Peek();
+
+			DocumentData.LogElement(DebugLog, InLogLinePrefix, -1);
+			DocumentData.PopElement();
+		}
+
+		private void PushInstance(
+			ElementType InInstanceType,
+			Transform   InWorldTransform,
+			string      InLogLinePrefix
+		)
+		{
+			FDocumentData DocumentData = DocumentDataStack.Peek();
+
+			DocumentData.PushInstance(InInstanceType, InWorldTransform);
+			DocumentData.LogElement(DebugLog, InLogLinePrefix, +1);
+		}
+
+		private void PopInstance(
+			string InLogLinePrefix
+		)
+		{
+			FDocumentData DocumentData = DocumentDataStack.Peek();
+
+			DocumentData.LogElement(DebugLog, InLogLinePrefix, -1);
+			DocumentData.PopInstance();
+		}
+
+		private void AddLightActor(
+			Transform InWorldTransform,
+			Asset     InLightAsset
+		)
+		{
+			DocumentDataStack.Peek().AddLightActor(InWorldTransform, InLightAsset);
+		}
+
+		private void AddRPCActor(
+			Transform InWorldTransform,
+			Asset     InRPCAsset
+		)
+		{
+			DocumentDataStack.Peek().AddRPCActor(InWorldTransform, InRPCAsset);
+		}
+
+		private void SetMaterial(
+			MaterialNode  InMaterialNode,
+			IList<string> InExtraTexturePaths
+		)
+		{
+			FDocumentData DocumentData = DocumentDataStack.Peek();
+
+			if (DocumentData.SetMaterial(InMaterialNode, InExtraTexturePaths))
+			{
+				DocumentData.LogMaterial(InMaterialNode, DebugLog, "Add Material");
+			}
+		}
+
+		private bool IgnoreElementGeometry()
+		{
+			return DocumentDataStack.Peek().IgnoreElementGeometry();
+		}
+
+		private FDatasmithFacadeMesh GetCurrentMesh()
+		{
+			return DocumentDataStack.Peek().GetCurrentMesh();
+		}
+
+		private int GetCurrentMaterialIndex()
+		{
+			return DocumentDataStack.Peek().GetCurrentMaterialIndex();
+		}
+
+		private void AddCameraActor(
+			View3D     InView3D,
+			CameraInfo InViewCamera
+		)
+		{
+			// Create a new Datasmith camera actor.
+			FDatasmithFacadeActorCamera CameraActor = new FDatasmithFacadeActorCamera(InView3D.UniqueId, InView3D.Name);
+
+			// Hash the Datasmith camera actor name to shorten it.
+			CameraActor.HashName();
+
+			if (InView3D.Category != null)
+			{
+				// Set the Datasmith camera actor layer to be the 3D view category name.
+				CameraActor.SetLayer(InView3D.Category.Name);
+			}
+
+			// Gets the current non-saved orientation of the 3D view.
+			ViewOrientation3D ViewOrientation = InView3D.GetOrientation();
+
+			// Set the world position (in right-handed Z-up coordinates) of the Datasmith camera actor.
+			XYZ CameraPosition = ViewOrientation.EyePosition;
+			CameraActor.SetCameraPosition((float) CameraPosition.X, (float) CameraPosition.Y, (float) CameraPosition.Z);
+
+			// Set the world rotation of the Datasmith camera actor with
+			// the camera world forward and up vectors (in right-handed Z-up coordinates).
+			XYZ CameraForward = ViewOrientation.ForwardDirection;
+			XYZ CameraUp      = ViewOrientation.UpDirection;
+			CameraActor.SetCameraRotation((float) CameraForward.X, (float) CameraForward.Y, (float) CameraForward.Z, (float) CameraUp.X, (float) CameraUp.Y, (float) CameraUp.Z);
+
+			// When the 3D view camera is not available, an orthographic view should be assumed.
+			if (InViewCamera != null)
+			{
+				// Compute the aspect ratio (width/height) of the Revit 3D view camera, where
+				// HorizontalExtent is the distance between left and right planes on the target plane,
+				// VerticalExtent is the distance between top and bottom planes on the target plane.
+				float AspectRatio = (float) (InViewCamera.HorizontalExtent / InViewCamera.VerticalExtent);
+
+				// Set the aspect ratio of the Datasmith camera.
+				CameraActor.SetAspectRatio(AspectRatio);
+
+				if(InView3D.IsPerspective)
+				{
+					// Set the sensor width of the Datasmith camera.
+					CameraActor.SetSensorWidth((float) (InViewCamera.HorizontalExtent * /* millimeters per foot */ 304.8));
+
+					// Get the distance from eye point along view direction to target plane.
+					// This value is appropriate for perspective views only.
+					float TargetDistance = (float) InViewCamera.TargetDistance;
+
+					// Set the Datasmith camera focus distance.
+					CameraActor.SetFocusDistance(TargetDistance);
+
+					// Set the Datasmith camera focal length.
+					CameraActor.SetFocalLength(TargetDistance * /* millimeters per foot */ 304.8F);
+				}
+			}
+
+			// Add the camera actor to the Datasmith scene.
+			DatasmithScene.AddElement(CameraActor);
+		}
+	}
+
+	public class FDocumentData
+	{
+		private class FElementData
+		{
+			private class FInstanceData
+			{
+				public ElementType               InstanceType;
+				public FDatasmithFacadeMesh      InstanceMesh  = null;
+				public FDatasmithFacadeActorMesh InstanceActor = null;
+
+				public FInstanceData(
+					ElementType InInstanceType
+				)
+				{
+					InstanceType = InInstanceType;
+				}
+			}
+
+			public  Element                   CurrentElement;
+			private ElementType               CurrentElementType;
+			private Stack<FInstanceData>      InstanceDataStack = new Stack<FInstanceData>();
+			public  FDatasmithFacadeMesh      ElementMesh       = null;
+			public  FDatasmithFacadeActorMesh ElementActor      = null;
+
+			public FElementData(
+				Element   InElement,
+				Transform InWorldTransform
+			)
+			{
+				CurrentElement     = InElement;
+				CurrentElementType = InElement.Document.GetElement(InElement.GetTypeId()) as ElementType;
+
+				CreateMeshActor(InWorldTransform, out ElementMesh, out ElementActor);
+			}
+
+			public void PushInstance(
+				ElementType InInstanceType,
+				Transform   InWorldTransform
+			)
+			{
+				FInstanceData InstanceData = new FInstanceData(InInstanceType);
+
+				InstanceDataStack.Push(InstanceData);
+
+				CreateMeshActor(InWorldTransform, out InstanceData.InstanceMesh, out InstanceData.InstanceActor);
+
+				// The Datasmith instance actor is a component in the hierarchy.
+				InstanceData.InstanceActor.SetIsComponent(true);
+			}
+
+			public FDatasmithFacadeMesh PopInstance()
+			{
+				FInstanceData InstanceData = InstanceDataStack.Pop();
+
+				FDatasmithFacadeMesh      InstanceMesh  = InstanceData.InstanceMesh;
+				FDatasmithFacadeActorMesh InstanceActor = InstanceData.InstanceActor;
+
+				if (InstanceMesh.GetVertexCount() > 0 && InstanceMesh.GetTriangleCount() > 0)
+				{
+					// Set the static mesh of the Datasmith instance actor.
+					InstanceActor.SetMesh(InstanceMesh.GetName());
+				}
+
+				// Add the instance mesh actor to the Datasmith actor hierarchy.
+				AddChildActor(InstanceActor);
+
+				return InstanceMesh;
+			}
+
+			public void AddLightActor(
+				Transform InWorldTransform,
+				Asset     InLightAsset
+			)
+			{
+				// Create a new Datasmith light actor.
+				FDatasmithFacadeActorLight LightActor = new FDatasmithFacadeActorLight("A:" + GetActorName(), GetActorLabel());
+
+				// Hash the Datasmith light actor name to shorten it.
+				LightActor.HashName();
+
+				// Set the world transform of the Datasmith light actor.
+				FDocumentData.SetActorTransform(InWorldTransform, LightActor);
+
+				// Set the base properties of the Datasmith light actor.
+				string LayerName = Category.GetCategory(CurrentElement.Document, BuiltInCategory.OST_LightingFixtureSource)?.Name ?? "Light Sources";
+				SetActorProperties(LayerName, LightActor);
+
+				// Set the specific properties of the Datasmith light actor.
+				FDatasmithRevitLight.SetLightProperties(InLightAsset, CurrentElement, LightActor);
+
+				// Add the light actor to the Datasmith actor hierarchy.
+				AddChildActor(LightActor);
+			}
+
+			public FDatasmithFacadeMesh AddRPCActor(
+				Transform InWorldTransform,
+				Asset     InRPCAsset,
+				int       InMaterialIndex
+			)
+			{
+				// Create a new Datasmith RPC mesh.
+				FDatasmithFacadeMesh RPCMesh = new FDatasmithFacadeMesh("M:" + GetActorName(), GetActorLabel());
+
+				// Hash the Datasmith RPC mesh name to shorten it.
+				RPCMesh.HashName();
+
+				Transform AffineTransform = Transform.Identity;
+
+				LocationPoint RPCLocationPoint = CurrentElement.Location as LocationPoint;
+
+				if (RPCLocationPoint != null)
+				{
+					if (RPCLocationPoint.Rotation != 0.0)
+					{
+						AffineTransform = AffineTransform.Multiply(Transform.CreateRotation(XYZ.BasisZ, -RPCLocationPoint.Rotation));
+						AffineTransform = AffineTransform.Multiply(Transform.CreateTranslation(RPCLocationPoint.Point.Negate()));
+					}
+					else
+					{
+						AffineTransform = Transform.CreateTranslation(RPCLocationPoint.Point.Negate());
+					}
+				}
+
+				GeometryElement RPCGeometryElement = CurrentElement.get_Geometry(new Options());
+
+				foreach (GeometryObject RPCGeometryObject in RPCGeometryElement)
+				{
+					GeometryInstance RPCGeometryInstance = RPCGeometryObject as GeometryInstance;
+
+					if (RPCGeometryInstance != null)
+					{
+						GeometryElement RPCInstanceGeometry = RPCGeometryInstance.GetInstanceGeometry();
+
+						foreach (GeometryObject RPCInstanceGeometryObject in RPCInstanceGeometry)
+						{
+							Mesh RPCInstanceGeometryMesh = RPCInstanceGeometryObject as Mesh;
+
+							if (RPCInstanceGeometryMesh == null || RPCInstanceGeometryMesh.NumTriangles < 1)
+							{
+								continue;
+							}
+
+							// RPC geometry does not have normals nor UVs available through the Revit Mesh interface.
+							int InitialVertexCount = RPCMesh.GetVertexCount();
+							int TriangleCount      = RPCInstanceGeometryMesh.NumTriangles;
+
+							// Add the RPC geometry vertices to the Datasmith RPC mesh.
+							foreach (XYZ Vertex in RPCInstanceGeometryMesh.Vertices)
+							{
+								XYZ PositionedVertex = AffineTransform.OfPoint(Vertex);
+								RPCMesh.AddVertex((float)PositionedVertex.X, (float)PositionedVertex.Y, (float)PositionedVertex.Z);
+							}
+
+							// Add the RPC geometry triangles to the Datasmith RPC mesh.
+							for (int TriangleNo = 0; TriangleNo < TriangleCount; TriangleNo++)
+							{
+								MeshTriangle Triangle = RPCInstanceGeometryMesh.get_Triangle(TriangleNo);
+
+								try
+								{
+									int Index0 = Convert.ToInt32(Triangle.get_Index(0));
+									int Index1 = Convert.ToInt32(Triangle.get_Index(1));
+									int Index2 = Convert.ToInt32(Triangle.get_Index(2));
+
+									// Add triangles for both the front and back faces.
+									RPCMesh.AddTriangle(InitialVertexCount + Index0, InitialVertexCount + Index1, InitialVertexCount + Index2, InMaterialIndex);
+									RPCMesh.AddTriangle(InitialVertexCount + Index2, InitialVertexCount + Index1, InitialVertexCount + Index0, InMaterialIndex);
+								}
+								catch (OverflowException)
+								{
+									continue;
+								}
+							}
+						}
+					}
+				}
+
+				// Create a new Datasmith RPC mesh actor.
+				FDatasmithFacadeActorMesh RPCMeshActor = new FDatasmithFacadeActorMesh("A:" + GetActorName(), GetActorLabel());
+
+				// Hash the Datasmith RPC mesh actor name to shorten it.
+				RPCMeshActor.HashName();
+
+				// Prevent the Datasmith RPC mesh actor from being removed by optimization.
+				RPCMeshActor.KeepActor();
+
+				if (RPCMesh.GetVertexCount() > 0 && RPCMesh.GetTriangleCount() > 0)
+				{
+					RPCMeshActor.SetMesh(RPCMesh.GetName());
+				}
+
+				// Set the world transform of the Datasmith RPC mesh actor.
+				FDocumentData.SetActorTransform(InWorldTransform, RPCMeshActor);
+
+				// Set the base properties of the Datasmith RPC mesh actor.
+				string LayerName = GetCategoryName();
+				SetActorProperties(LayerName, RPCMeshActor);
+
+				// Add a Revit element RPC tag to the Datasmith RPC mesh actor.
+				RPCMeshActor.AddTag("Revit.Element.RPC");
+
+				// Add some Revit element RPC metadata to the Datasmith RPC mesh actor.
+				AssetProperty RPCTypeId   = InRPCAsset.FindByName("RPCTypeId");
+				AssetProperty RPCFilePath = InRPCAsset.FindByName("RPCFilePath");
+
+				if (RPCTypeId != null)
+				{
+					RPCMeshActor.AddMetadataString("Type*RPCTypeId", (RPCTypeId as AssetPropertyString).Value);
+				}
+
+				if (RPCFilePath != null)
+				{
+					RPCMeshActor.AddMetadataString("Type*RPCFilePath", (RPCFilePath as AssetPropertyString).Value);
+				}
+
+				// Add the RPC mesh actor to the Datasmith actor hierarchy.
+				AddChildActor(RPCMeshActor);
+
+				return RPCMesh;
+			}
+
+			public void AddChildActor(
+				FDatasmithFacadeActor InChildActor
+			)
+			{
+				if (InstanceDataStack.Count == 0)
+				{
+					ElementActor.AddChild(InChildActor);
+				}
+				else
+				{
+					InstanceDataStack.Peek().InstanceActor.AddChild(InChildActor);
+				}
+			}
+
+			public string GetCategoryName()
+			{
+				return CurrentElementType?.Category?.Name ?? CurrentElement.Category?.Name;
+			}
+
+			public bool IgnoreElementGeometry()
+			{
+				// Ignore elements that have unwanted geometry, such as level symbols.
+				return (CurrentElementType as LevelType) != null;
+			}
+
+			public FDatasmithFacadeMesh GetCurrentMesh()
+			{
+				if (InstanceDataStack.Count == 0)
+				{
+					return ElementMesh;
+				}
+				else
+				{
+					return InstanceDataStack.Peek().InstanceMesh;
+				}
+			}
+
+			public FDatasmithFacadeActorMesh GetCurrentActor()
+			{
+				if (InstanceDataStack.Count == 0)
+				{
+					return ElementActor;
+				}
+				else
+				{
+					return InstanceDataStack.Peek().InstanceActor;
+				}
+			}
+
+			public void Log(
+				FDatasmithFacadeLog InDebugLog,
+				string              InLinePrefix,
+				int                 InLineIndentation
+			)
+			{
+				if (InDebugLog != null)
+				{
+					if (InLineIndentation < 0)
+					{
+						InDebugLog.LessIndentation();
+					}
+
+					Element SourceElement = (InstanceDataStack.Count == 0) ? CurrentElement : InstanceDataStack.Peek().InstanceType;
+
+					InDebugLog.AddLine($"{InLinePrefix} {SourceElement.Id.IntegerValue} '{SourceElement.Name}' {SourceElement.GetType()}: '{GetActorLabel()}'");
+
+					if (InLineIndentation > 0)
+					{
+						InDebugLog.MoreIndentation();
+					}
+				}
+			}
+
+			private string GetActorName()
+			{
+				string DocumentName = Path.GetFileNameWithoutExtension(CurrentElement.Document.PathName);
+
+				if (InstanceDataStack.Count == 0)
+				{
+					return $"{DocumentName}:{CurrentElement.UniqueId}";
+				}
+				else
+				{
+					return $"{DocumentName}:{InstanceDataStack.Peek().InstanceType.UniqueId}";
+				}
+			}
+
+			private string GetActorLabel()
+			{
+				string CategoryName = GetCategoryName();
+				string FamilyName   = CurrentElementType?.FamilyName;
+				string TypeName     = CurrentElementType?.Name;
+				string InstanceName = (InstanceDataStack.Count > 1) ? InstanceDataStack.Peek().InstanceType?.Name : null;
+
+				string ActorLabel = "";
+
+				if (CurrentElement as Level != null)
+				{
+					ActorLabel += string.IsNullOrEmpty(FamilyName) ? "" : FamilyName + "*";
+					ActorLabel += string.IsNullOrEmpty(TypeName)   ? "" : TypeName + "*";
+					ActorLabel += CurrentElement.Name;
+				}
+				else
+				{
+					ActorLabel += string.IsNullOrEmpty(CategoryName) ? ""                  : CategoryName + "*";
+					ActorLabel += string.IsNullOrEmpty(FamilyName)   ? ""                  : FamilyName + "*";
+					ActorLabel += string.IsNullOrEmpty(TypeName)     ? CurrentElement.Name : TypeName;
+					ActorLabel += string.IsNullOrEmpty(InstanceName) ? ""                  : "*" + InstanceName;
+				}
+
+				return ActorLabel;
+			}
+
+			private void CreateMeshActor(
+				    Transform                 InWorldTransform,
+				out FDatasmithFacadeMesh      OutMesh,
+				out FDatasmithFacadeActorMesh OutMeshActor
+			)
+			{
+				// Create a new Datasmith mesh.
+				OutMesh = new FDatasmithFacadeMesh("M:" + GetActorName(), GetActorLabel());
+
+				// Hash the Datasmith mesh name to shorten it.
+				OutMesh.HashName();
+
+				// Create a new Datasmith mesh actor.
+				OutMeshActor = new FDatasmithFacadeActorMesh("A:" + GetActorName(), GetActorLabel());
+
+				// Hash the Datasmith mesh actor name to shorten it.
+				OutMeshActor.HashName();
+
+				// Set the world transform of the Datasmith mesh actor.
+				FDocumentData.SetActorTransform(InWorldTransform, OutMeshActor);
+
+				// Set the base properties of the Datasmith mesh actor.
+				string LayerName = GetCategoryName();
+				SetActorProperties(LayerName, OutMeshActor);
+			}
+
+			private void SetActorProperties(
+				string                InLayerName,
+				FDatasmithFacadeActor IOActor
+			)
+			{
+				// Set the Datasmith actor layer to the element type category name.
+				IOActor.SetLayer(InLayerName);
+
+				// Add the Revit element ID and Unique ID tags to the Datasmith actor.
+				IOActor.AddTag($"Revit.Element.Id.{CurrentElement.Id.IntegerValue}");
+				IOActor.AddTag($"Revit.Element.UniqueId.{CurrentElement.UniqueId}");
+
+				// For an hosted Revit family instance, add the host ID and Unique ID tags to the Datasmith actor.
+				FamilyInstance CurrentFamilyInstance = CurrentElement as FamilyInstance;
+				if (CurrentFamilyInstance?.Host != null)
+				{
+					IOActor.AddTag($"Revit.Host.Id.{CurrentFamilyInstance.Host.Id.IntegerValue}");
+					IOActor.AddTag($"Revit.Host.UniqueId.{CurrentFamilyInstance.Host.UniqueId}");
+				}
+
+				// Add the Revit element category name metadata to the Datasmith actor.
+				string CategoryName = GetCategoryName();
+				if (!string.IsNullOrEmpty(CategoryName))
+				{
+					IOActor.AddMetadataString("Element*Category", CategoryName);
+				}
+
+				// Add the Revit element family name metadata to the Datasmith actor.
+				string FamilyName = CurrentElementType?.FamilyName;
+				if (!string.IsNullOrEmpty(FamilyName))
+				{
+					IOActor.AddMetadataString("Element*Family", FamilyName);
+				}
+
+				// Add the Revit element type name metadata to the Datasmith actor.
+				string TypeName = CurrentElementType?.Name;
+				if (!string.IsNullOrEmpty(TypeName))
+				{
+					IOActor.AddMetadataString("Element*Type", TypeName);
+				}
+
+				// Add Revit element metadata to the Datasmith actor.
+				FDocumentData.AddActorMetadata(CurrentElement, "Element*", IOActor);
+
+				if (CurrentElementType != null)
+				{
+					// Add Revit element type metadata to the Datasmith actor.
+					FDocumentData.AddActorMetadata(CurrentElementType, "Type*", IOActor);
+				}
+			}
+		}
+
+		private Document                                 CurrentDocument;
+		private Stack<FElementData>                      ElementDataStack         = new Stack<FElementData>();
+		private Dictionary<string, FDatasmithFacadeMesh> CollectedMeshMap         = new Dictionary<string, FDatasmithFacadeMesh>();
+		private Dictionary<int, FDatasmithFacadeActor>   CollectedActorMap        = new Dictionary<int, FDatasmithFacadeActor>();
+		private string                                   CurrentMaterialDataKey   = null;
+		private Dictionary<string, FMaterialData>        CollectedMaterialDataMap = new Dictionary<string, FMaterialData>();
+		private int                                      LatestMaterialIndex      = 0;
+		private List<string>                             MessageList              = null;
+
+		public FDocumentData(
+			Document         InDocument,
+			ref List<string> InMessageList
+		)
+		{
+			CurrentDocument = InDocument;
+			MessageList     = InMessageList;
+		}
+
+		public Element GetElement(
+			ElementId InElementId
+		)
+		{
+			return (InElementId != ElementId.InvalidElementId) ? CurrentDocument.GetElement(InElementId) : null;
+		}
+
+		public void PushElement(
+			Element   InElement,
+			Transform InWorldTransform
+		)
+		{
+			ElementDataStack.Push(new FElementData(InElement, InWorldTransform));
+		}
+
+		public void PopElement()
+		{
+			FElementData ElementData = ElementDataStack.Pop();
+
+			FDatasmithFacadeMesh      ElementMesh  = ElementData.ElementMesh;
+			FDatasmithFacadeActorMesh ElementActor = ElementData.ElementActor;
+
+			if (ElementMesh.GetVertexCount() > 0 && ElementMesh.GetTriangleCount() > 0)
+			{
+				// Set the static mesh of the Datasmith actor.
+				ElementActor.SetMesh(ElementMesh.GetName());
+			}
+
+			// Collect the element Datasmith mesh into the mesh dictionary.
+			CollectMesh(ElementMesh);
+
+			if (ElementDataStack.Count == 0)
+			{
+				int ElementId = ElementData.CurrentElement.Id.IntegerValue;
+
+				if (CollectedActorMap.ContainsKey(ElementId))
+				{
+					// Handle the spurious case of Revit Custom Exporter calling back more than once for the same element.
+					// These extra empty actors will be cleaned up later by the Datasmith actor hierarchy optimization.
+					CollectedActorMap[ElementId].AddChild(ElementActor);
+				}
+				else
+				{
+					// Collect the element mesh actor into the Datasmith actor dictionary.
+					CollectedActorMap[ElementId] = ElementActor;
+				}
+			}
+			else
+			{
+				// Add the element mesh actor to the Datasmith actor hierarchy.
+				ElementDataStack.Peek().AddChildActor(ElementActor);
+			}
+		}
+
+		public void PushInstance(
+			ElementType InInstanceType,
+			Transform   InWorldTransform
+		)
+		{
+			ElementDataStack.Peek().PushInstance(InInstanceType, InWorldTransform);
+		}
+
+		public void PopInstance()
+		{
+			FDatasmithFacadeMesh InstanceMesh = ElementDataStack.Peek().PopInstance();
+
+			// Collect the instance mesh into the Datasmith mesh dictionary.
+			CollectMesh(InstanceMesh);
+		}
+
+		public void AddLocationActors(
+			Transform InWorldTransform
+		)
+		{
+			// Add a new Datasmith placeholder actor for this document site location.
+			AddSiteLocation(CurrentDocument.SiteLocation);
+
+			// Add new Datasmith placeholder actors for the project base point and survey points.
+			// A project has one base point and at least one survey point. Linked documents also have their own points.
+			AddPointLocations(InWorldTransform);
+		}
+
+		public void AddLightActor(
+			Transform InWorldTransform,
+			Asset     InLightAsset
+		)
+		{
+			ElementDataStack.Peek().AddLightActor(InWorldTransform, InLightAsset);
+		}
+		
+		public void AddRPCActor(
+			Transform InWorldTransform,
+			Asset     InRPCAsset
+		)
+		{
+			// Create a simple fallback material for the RPC mesh.
+			string RPCCategoryName = ElementDataStack.Peek().GetCategoryName();
+			bool   isRPCPlant      = !string.IsNullOrEmpty(RPCCategoryName) && RPCCategoryName == Category.GetCategory(CurrentDocument, BuiltInCategory.OST_Planting)?.Name;
+			string RPCMaterialName = isRPCPlant ? "RPC_Plant" : "RPC_Material";
+
+			if (!CollectedMaterialDataMap.ContainsKey(RPCMaterialName))
+			{
+				// Color reference: https://www.color-hex.com/color-palette/70002
+				Color RPCColor = isRPCPlant ? /* green */ new Color(88, 126, 96) : /* gray */ new Color(128, 128, 128);
+
+				// Keep track of a new RPC master material.
+				CollectedMaterialDataMap[RPCMaterialName] = new FMaterialData(RPCMaterialName, RPCColor, ++LatestMaterialIndex);
+			}
+
+			FMaterialData RPCMaterialData = CollectedMaterialDataMap[RPCMaterialName];
+
+			FDatasmithFacadeMesh RPCMesh = ElementDataStack.Peek().AddRPCActor(InWorldTransform, InRPCAsset, RPCMaterialData.MaterialIndex);
+
+			// Add the RPC master material name to the dictionary of material names utilized by the RPC mesh.
+			RPCMesh.AddMaterial(RPCMaterialData.MaterialIndex, RPCMaterialData.MasterMaterial.GetName());
+
+			// Collect the RPC mesh into the Datasmith mesh dictionary.
+			CollectMesh(RPCMesh);
+		}
+
+		public bool SetMaterial(
+			MaterialNode  InMaterialNode,
+			IList<string> InExtraTexturePaths
+		)
+		{
+			Material CurrentMaterial = GetElement(InMaterialNode.MaterialId) as Material;
+
+			CurrentMaterialDataKey = FMaterialData.GetMaterialName(InMaterialNode, CurrentMaterial);
+
+			if (!CollectedMaterialDataMap.ContainsKey(CurrentMaterialDataKey))
+			{
+				// Keep track of a new Datasmith master material.
+				CollectedMaterialDataMap[CurrentMaterialDataKey] = new FMaterialData(InMaterialNode, CurrentMaterial, ++LatestMaterialIndex, InExtraTexturePaths);
+
+				// A new Datasmith master material was created.
+				return true;
+			}
+
+			// No new Datasmith master material created.
+			return false;
+		}
+
+		public bool IgnoreElementGeometry()
+		{
+			return ElementDataStack.Peek().IgnoreElementGeometry();
+		}
+
+		public FDatasmithFacadeMesh GetCurrentMesh()
+		{
+			return ElementDataStack.Peek().GetCurrentMesh();
+		}
+
+		public int GetCurrentMaterialIndex()
+		{
+			if (CollectedMaterialDataMap.ContainsKey(CurrentMaterialDataKey))
+			{
+				FMaterialData MaterialData = CollectedMaterialDataMap[CurrentMaterialDataKey];
+
+				// Add the current Datasmith master material name to the dictionary of material names utilized by the Datasmith mesh being processed.
+				GetCurrentMesh().AddMaterial(MaterialData.MaterialIndex, MaterialData.MasterMaterial.GetName());
+
+				// Return the index of the current material.
+				return MaterialData.MaterialIndex;
+			}
+
+			return 0;
+		}
+
+		public FDatasmithFacadeActorMesh GetCurrentActor()
+		{
+			return ElementDataStack.Peek().GetCurrentActor();
+		}
+
+		public void WrapupLink(
+			FDatasmithFacadeScene InDatasmithScene,
+			FDatasmithFacadeActor InLinkActor
+		)
+		{
+			// Add the collected meshes from the Datasmith mesh dictionary to the Datasmith scene.
+			AddCollectedMeshes(InDatasmithScene);
+
+			// Factor in the Datasmith actor hierarchy the Revit document host hierarchy.
+			AddHostHierarchy();
+
+			// Factor in the Datasmith actor hierarchy the Revit document level hierarchy.
+			AddLevelHierarchy();
+
+			if (CollectedActorMap.Count > 0)
+			{
+				// Prevent the Datasmith link actor from being removed by optimization.
+				InLinkActor.KeepActor();
+
+				// Add the collected actors from the Datasmith actor dictionary as children of the Datasmith link actor.
+				foreach (FDatasmithFacadeActor CollectedActor in CollectedActorMap.Values)
+				{
+					InLinkActor.AddChild(CollectedActor);
+				}
+			}
+
+			// Add the collected master materials from the material data dictionary to the Datasmith scene.
+			AddCollectedMaterials(InDatasmithScene);
+		}
+
+		public void WrapupScene(
+			FDatasmithFacadeScene InDatasmithScene
+		)
+		{
+			// Add the collected meshes from the Datasmith mesh dictionary to the Datasmith scene.
+			AddCollectedMeshes(InDatasmithScene);
+
+			// Factor in the Datasmith actor hierarchy the Revit document host hierarchy.
+			AddHostHierarchy();
+
+			// Factor in the Datasmith actor hierarchy the Revit document level hierarchy.
+			AddLevelHierarchy();
+
+			// Add the collected actors from the Datasmith actor dictionary to the Datasmith scene.
+			foreach (FDatasmithFacadeActor CollectedActor in CollectedActorMap.Values)
+			{
+				// Make sure all the actor names are unique and persistent in the Datasmith actor hierarchy.
+				CollectedActor.SanitizeActorHierarchyNames();
+
+				InDatasmithScene.AddElement(CollectedActor);
+			}
+
+			// Add the collected master materials from the material data dictionary to the Datasmith scene.
+			AddCollectedMaterials(InDatasmithScene);
+		}
+
+		public void LogElement(
+			FDatasmithFacadeLog InDebugLog,
+			string              InLinePrefix,
+			int                 InLineIndentation
+		)
+		{
+			ElementDataStack.Peek().Log(InDebugLog, InLinePrefix, InLineIndentation);
+		}
+
+		public void LogMaterial(
+			MaterialNode        InMaterialNode,
+			FDatasmithFacadeLog InDebugLog,
+			string              InLinePrefix
+		)
+		{
+			if (CollectedMaterialDataMap.ContainsKey(CurrentMaterialDataKey))
+			{
+				CollectedMaterialDataMap[CurrentMaterialDataKey].Log(InMaterialNode, InDebugLog, InLinePrefix);
+			}
+		}
+
+		private void AddSiteLocation(
+			SiteLocation InSiteLocation
+		)
+		{
+			if (!InSiteLocation.IsValidObject)
+			{
+				return;
+			}
+
+			// Create a new Datasmith placeholder actor for the site location.
+			FDatasmithFacadeActor SiteLocationActor = new FDatasmithFacadeActor("SiteLocation", "Site Location");
+
+			// Hash the Datasmith placeholder actor name to shorten it.
+			SiteLocationActor.HashName();
+
+			// Prevent the Datasmith placeholder actor from being removed by optimization.
+			SiteLocationActor.KeepActor();
+
+			// Set the Datasmith placeholder actor layer to the site location category name.
+			SiteLocationActor.SetLayer(InSiteLocation.Category.Name);
+
+			// Add the Revit element ID and Unique ID tags to the Datasmith placeholder actor.
+			SiteLocationActor.AddTag($"Revit.Element.Id.{InSiteLocation.Id.IntegerValue}");
+			SiteLocationActor.AddTag($"Revit.Element.UniqueId.{InSiteLocation.UniqueId}");
+
+			// Add a Revit element site location tag to the Datasmith placeholder actor.
+			SiteLocationActor.AddTag("Revit.Element.SiteLocation");
+
+			// Add site location metadata to the Datasmith placeholder actor.
+			const double RadiansToDegrees = 180.0 / Math.PI;
+			SiteLocationActor.AddMetadataFloat("SiteLocation*Latitude",  (float) (InSiteLocation.Latitude * RadiansToDegrees));
+			SiteLocationActor.AddMetadataFloat("SiteLocation*Longitude", (float) (InSiteLocation.Longitude * RadiansToDegrees));
+			SiteLocationActor.AddMetadataFloat("SiteLocation*Elevation", (float) InSiteLocation.Elevation);
+			SiteLocationActor.AddMetadataFloat("SiteLocation*TimeZone",  (float) InSiteLocation.TimeZone);
+			SiteLocationActor.AddMetadataString("SiteLocation*Place",    InSiteLocation.PlaceName);
+
+			// Collect the site location placeholder actor into the Datasmith actor dictionary.
+			CollectedActorMap[InSiteLocation.Id.IntegerValue] = SiteLocationActor;
+		}
+
+		private void AddPointLocations(
+			Transform InWorldTransform
+		)
+		{
+			FilteredElementCollector Collector = new FilteredElementCollector(CurrentDocument);
+			ICollection<Element> PointLocations = Collector.OfClass(typeof(BasePoint)).ToElements();
+
+			foreach (Element PointLocation in PointLocations)
+			{
+				BasePoint BasePointLocation = PointLocation as BasePoint;
+
+				if (BasePointLocation != null)
+				{
+					// Create a new Datasmith placeholder actor for the base point.
+					string ActorName  = BasePointLocation.IsShared ? "SurveyPoint"  : "BasePoint";
+					string ActorLabel = BasePointLocation.IsShared ? "Survey Point" : "Base Point";
+					FDatasmithFacadeActor BasePointActor = new FDatasmithFacadeActor(ActorName, ActorLabel);
+
+					// Hash the Datasmith placeholder actor name to shorten it.
+					BasePointActor.HashName();
+
+					// Prevent the Datasmith placeholder actor from being removed by optimization.
+					BasePointActor.KeepActor();
+
+					// Set the world transform of the Datasmith placeholder actor.
+					// Since BasePoint.Location is not a location point we cannot get a position from it; so we use a bounding box approach.
+					// Note that, as of Revit 2020, BasePoint has 2 new properties: Position for base point and SharedPosition for survey point.
+					XYZ BasePointPosition = BasePointLocation.get_BoundingBox(CurrentDocument.ActiveView).Min;
+					Transform TranslationMatrix = Transform.CreateTranslation(BasePointPosition);
+					FDocumentData.SetActorTransform(TranslationMatrix.Multiply(InWorldTransform), BasePointActor);
+
+					// Set the Datasmith placeholder actor layer to the base point category name.
+					BasePointActor.SetLayer(BasePointLocation.Category.Name);
+
+					// Add the Revit element ID and Unique ID tags to the Datasmith placeholder actor.
+					BasePointActor.AddTag($"Revit.Element.Id.{BasePointLocation.Id.IntegerValue}");
+					BasePointActor.AddTag($"Revit.Element.UniqueId.{BasePointLocation.UniqueId}");
+
+					// Add a Revit element base point tag to the Datasmith placeholder actor.
+					BasePointActor.AddTag("Revit.Element." + ActorName);
+
+					// Add base point metadata to the Datasmith actor.
+					string MetadataPrefix = BasePointLocation.IsShared ? "SurveyPointLocation*" : "BasePointLocation*";
+					BasePointActor.AddMetadataVector(MetadataPrefix + "Location", $"{BasePointPosition.X} {BasePointPosition.Y} {BasePointPosition.Z}");
+					FDocumentData.AddActorMetadata(BasePointLocation, MetadataPrefix, BasePointActor);
+
+					// Collect the base point placeholder actor into the Datasmith actor dictionary.
+					CollectedActorMap[BasePointLocation.Id.IntegerValue] = BasePointActor;
+				}
+			}
+		}
+
+		private void CollectMesh(
+			FDatasmithFacadeMesh InMesh
+		)
+		{
+			if (InMesh.GetVertexCount() > 0 && InMesh.GetTriangleCount() > 0)
+			{
+				string MeshName = InMesh.GetName();
+
+				if (!CollectedMeshMap.ContainsKey(MeshName))
+				{
+					// Keep track of the Datasmith mesh.
+					CollectedMeshMap[MeshName] = InMesh;
+				}
+			}
+		}
+
+		private void AddCollectedMeshes(
+			FDatasmithFacadeScene InDatasmithScene
+		)
+		{
+			// Add the collected meshes from the Datasmith mesh dictionary to the Datasmith scene.
+			foreach (FDatasmithFacadeMesh CollectedMesh in CollectedMeshMap.Values)
+			{
+				InDatasmithScene.AddElement(CollectedMesh);
+			}
+		}
+
+		private void AddHostHierarchy()
+		{
+			AddParentElementHierarchy(GetHostElement);
+		}
+
+		private void AddLevelHierarchy()
+		{
+			AddParentElementHierarchy(GetLevelElement);
+		}
+
+		private void AddCollectedMaterials(
+			FDatasmithFacadeScene InDatasmithScene
+		)
+		{
+			// Add the collected master materials from the material data dictionary to the Datasmith scene.
+			foreach (FMaterialData CollectedMaterialData in CollectedMaterialDataMap.Values)
+			{
+				InDatasmithScene.AddElement(CollectedMaterialData.MasterMaterial);
+
+				if (CollectedMaterialData.MessageList.Count > 0)
+				{
+					MessageList.AddRange(CollectedMaterialData.MessageList);
+				}
+			}
+		}
+
+		private Element GetHostElement(
+			int InElementId
+		)
+		{
+			Element SourceElement = CurrentDocument.GetElement(new ElementId(InElementId));
+
+			if (SourceElement as FamilyInstance != null)
+			{
+				return (SourceElement as FamilyInstance).Host;
+			}
+			else if (SourceElement as Wall != null)
+			{
+				return CurrentDocument.GetElement((SourceElement as Wall).StackedWallOwnerId);
+			}
+			else if (SourceElement as ContinuousRail != null)
+			{
+				return CurrentDocument.GetElement((SourceElement as ContinuousRail).HostRailingId);
+			}
+
+			return null;
+		}
+
+		private Element GetLevelElement(
+			int InElementId
+		)
+		{
+			Element SourceElement = CurrentDocument.GetElement(new ElementId(InElementId));
+
+			return (SourceElement == null) ? null : CurrentDocument.GetElement(SourceElement.LevelId);
+		}
+
+		private void AddParentElementHierarchy(
+			Func<int, Element> InGetParentElement
+		)
+		{
+			Queue<int> ElementIdQueue = new Queue<int>(CollectedActorMap.Keys);
+
+			// Make sure the Datasmith actor dictionary contains actors for all the Revit parent elements.
+			while (ElementIdQueue.Count > 0)
+			{
+				Element ParentElement = InGetParentElement(ElementIdQueue.Dequeue());
+
+				if (ParentElement == null)
+				{
+					continue;
+				}
+
+				int ParentElementId = ParentElement.Id.IntegerValue;
+
+				if (CollectedActorMap.ContainsKey(ParentElementId))
+				{
+					continue;
+				}
+
+				// Create a new Datasmith actor for the Revit parent element.
+				PushElement(ParentElement, Transform.Identity);
+				PopElement();
+
+				ElementIdQueue.Enqueue(ParentElementId);
+			}
+
+			// Add the parented actors as children of the parent Datasmith actors.
+			foreach (int ElementId in new List<int>(CollectedActorMap.Keys))
+			{
+				Element ParentElement = InGetParentElement(ElementId);
+
+				if (ParentElement == null)
+				{
+					continue;
+				}
+
+				Element SourceElement = CurrentDocument.GetElement(new ElementId(ElementId));
+
+				if ((SourceElement as FamilyInstance != null && ParentElement as Truss != null) ||
+				    (SourceElement as Mullion != null) ||
+				    (SourceElement as Panel != null) ||
+				    (SourceElement as ContinuousRail != null))
+				{
+					// The Datasmith actor is a component in the hierarchy.
+					CollectedActorMap[ElementId].SetIsComponent(true);
+				}
+
+				int ParentElementId = ParentElement.Id.IntegerValue;
+
+				// Add the parented actor as child of the parent Datasmith actor.
+				CollectedActorMap[ParentElementId].AddChild(CollectedActorMap[ElementId]);
+
+				// Prevent the parent Datasmith actor from being removed by optimization.
+				CollectedActorMap[ParentElementId].KeepActor();
+			}
+
+			// Remove the parented child actors from the Datasmith actor dictionary.
+			foreach (int ElementId in new List<int>(CollectedActorMap.Keys))
+			{
+				Element ParentElement = InGetParentElement(ElementId);
+
+				if (ParentElement == null)
+				{
+					continue;
+				}
+
+				// Remove the parented child actor from the Datasmith actor dictionary.
+				CollectedActorMap.Remove(ElementId);
+			}
+		}
+
+		private static void SetActorTransform(
+			Transform             InWorldTransform,
+			FDatasmithFacadeActor IOActor
+		)
+		{
+			XYZ transformBasisX = InWorldTransform.BasisX;
+			XYZ transformBasisY = InWorldTransform.BasisY;
+			XYZ transformBasisZ = InWorldTransform.BasisZ;
+			XYZ transformOrigin = InWorldTransform.Origin;
+
+			float[] worldMatrix = new float[16];
+
+			worldMatrix[0]  = (float) transformBasisX.X;
+			worldMatrix[1]  = (float) transformBasisX.Y;
+			worldMatrix[2]  = (float) transformBasisX.Z;
+			worldMatrix[3]  = 0.0F;
+			worldMatrix[4]  = (float) transformBasisY.X;
+			worldMatrix[5]  = (float) transformBasisY.Y;
+			worldMatrix[6]  = (float) transformBasisY.Z;
+			worldMatrix[7]  = 0.0F;
+			worldMatrix[8]  = (float) transformBasisZ.X;
+			worldMatrix[9]  = (float) transformBasisZ.Y;
+			worldMatrix[10] = (float) transformBasisZ.Z;
+			worldMatrix[11] = 0.0F;
+			worldMatrix[12] = (float) transformOrigin.X;
+			worldMatrix[13] = (float) transformOrigin.Y;
+			worldMatrix[14] = (float) transformOrigin.Z;
+			worldMatrix[15] = 1.0F;
+
+			// Set the world transform of the Datasmith actor.
+			IOActor.SetWorldTransform(worldMatrix);
+		}
+
+		private static void AddActorMetadata(
+			Element               InSourceElement,
+			string                InMetadataPrefix,
+			FDatasmithFacadeActor IOActor
+		)
+		{
+			IList<Parameter> Parameters = InSourceElement.GetOrderedParameters();
+
+			if (Parameters != null)
+			{
+				foreach (Parameter Parameter in Parameters)
+				{
+					if (Parameter.HasValue)
+					{
+						string ParameterValue = Parameter.AsValueString();
+
+						if (string.IsNullOrEmpty(ParameterValue))
+						{
+							switch (Parameter.StorageType)
+							{
+								case StorageType.Integer:
+									ParameterValue = Parameter.AsInteger().ToString();
+									break;
+								case StorageType.Double:
+									ParameterValue = Parameter.AsDouble().ToString();
+									break;
+								case StorageType.String:
+									ParameterValue = Parameter.AsString();
+									break;
+								case StorageType.ElementId:
+									ParameterValue = Parameter.AsElementId().ToString();
+									break;
+							}
+						}
+
+						if (!string.IsNullOrEmpty(ParameterValue))
+						{
+							string MetadataKey = InMetadataPrefix + Parameter.Definition.Name;
+							IOActor.AddMetadataString(MetadataKey, ParameterValue);
+						}
+					}
+				}
+			}
+		}
+	}
+}
diff --git a/Engine/Source/Programs/Enterprise/Datasmith/DatasmithRevitExporter/Private/DatasmithRevitMaterial.cs b/Engine/Source/Programs/Enterprise/Datasmith/DatasmithRevitExporter/Private/DatasmithRevitMaterial.cs
new file mode 100644
--- /dev/null
+++ b/Engine/Source/Programs/Enterprise/Datasmith/DatasmithRevitExporter/Private/DatasmithRevitMaterial.cs
@@ -0,0 +1,1628 @@
+﻿// Copyright 1998-2019 Epic Games, Inc. All Rights Reserved.
+
+using System;
+using System.Collections.Generic;
+using System.IO;
+
+using Autodesk.Revit.DB;
+using Autodesk.Revit.DB.Visual;
+
+
+namespace DatasmithRevitExporter
+{
+	public class FMaterialData
+	{
+		private Material                 CurrentMaterial   = null;
+		private string                   MaterialLabel     = null;	
+		public  int                      MaterialIndex     = 0;
+		public  FDatasmithFacadeMaterial MasterMaterial    = null;
+		private IList<string>            ExtraTexturePaths = null;
+		public  IList<string>            MessageList       = new List<string>();
+
+		// Multi-line debug log.
+		// private FDatasmithFacadeLog DebugLog = null;
+
+		public static string GetMaterialName(
+			MaterialNode InMaterialNode,
+			Material     InMaterial
+		)
+		{
+			if (InMaterial == null)
+			{
+				Color MaterialColor        = InMaterialNode.Color.IsValid ? InMaterialNode.Color : new Color(255, 255, 255);
+				int   MaterialTransparency = (int)(InMaterialNode.Transparency * 100.0);
+				int   MaterialSmoothness   = InMaterialNode.Smoothness;
+
+				// Generate a unique name for the fallback material.
+				return $"{MaterialColor.Red:x2}{MaterialColor.Green:x2}{MaterialColor.Blue:x2}{MaterialTransparency:x2}{MaterialSmoothness:x2}";
+			}
+			else
+			{
+				return $"{Path.GetFileNameWithoutExtension(InMaterial.Document.PathName)}:{InMaterial.UniqueId}";
+			}
+		}
+
+		public FMaterialData(
+			MaterialNode  InMaterialNode,
+			Material      InMaterial,
+			int           InMaterialIndex,
+			IList<string> InExtraTexturePaths
+		)
+		{
+			CurrentMaterial   = InMaterial;
+			MaterialLabel     = GetMaterialLabel(InMaterialNode, InMaterial);
+			MaterialIndex     = InMaterialIndex;
+			ExtraTexturePaths = InExtraTexturePaths;
+
+			// Create a new Datasmith master material.
+			MasterMaterial = new FDatasmithFacadeMaterial(GetMaterialName(InMaterialNode, CurrentMaterial), MaterialLabel);
+
+			// Hash the Datasmith master material name to shorten it.
+			MasterMaterial.HashName();
+
+			// Set the properties of the Datasmith master material.
+			if (!SetMasterMaterial(CurrentMaterial, MasterMaterial))
+			{
+				SetFallbackMaterial(InMaterialNode.Color, (float) InMaterialNode.Transparency, InMaterialNode.Smoothness / 100.0F, MasterMaterial);
+			}
+		}
+
+		public FMaterialData(
+			string InMaterialName,
+			Color  InMaterialColor,
+			int    InMaterialIndex
+		)
+		{
+			MaterialLabel = InMaterialName;
+			MaterialIndex = InMaterialIndex;
+
+			// Create a new Datasmith master material.
+			MasterMaterial = new FDatasmithFacadeMaterial(InMaterialName, MaterialLabel);
+
+			// Hash the Datasmith master material name to shorten it.
+			MasterMaterial.HashName();
+
+			// Set the properties of the Datasmith master material.
+			SetFallbackMaterial(InMaterialColor, 0.0F, 0.5F, MasterMaterial);
+		}
+
+		public void Log(
+			MaterialNode        InMaterialNode,
+			FDatasmithFacadeLog InDebugLog,
+			string              InLinePrefix
+		)
+		{
+			if (InDebugLog != null)
+			{
+				int MaterialId = (CurrentMaterial == null) ? 0 : CurrentMaterial.Id.IntegerValue;
+
+				InDebugLog.AddLine($"{InLinePrefix} {MaterialId}: '{GetMaterialLabel(InMaterialNode, CurrentMaterial)}'");
+			}
+		}
+
+		private static string GetMaterialLabel(
+			MaterialNode InMaterialNode,
+			Material     InMaterial
+		)
+		{
+			if (InMaterial != null)
+			{
+				Asset RenderingAsset = (InMaterial.Document.GetElement(InMaterial.AppearanceAssetId) as AppearanceAssetElement)?.GetRenderingAsset();
+
+				if (RenderingAsset != null)
+				{
+					string RenderingAssetName = RenderingAsset.Name.Replace("Schema", "");
+					Type RenderingAssetType = Type.GetType($"Autodesk.Revit.DB.Visual.{RenderingAssetName},RevitAPI");
+
+					if (RenderingAssetType != null)
+					{
+						switch (RenderingAssetType.Name)
+						{
+							case "Ceramic":
+							case "Concrete":
+							case "Glazing":
+							case "Hardwood":
+							case "MasonryCMU":
+							case "Metal":
+							case "MetallicPaint":
+							case "Mirror":
+							case "PlasticVinyl":
+							case "SolidGlass":
+							case "Stone":
+							case "WallPaint":
+							case "Generic":
+								return InMaterial.Name;
+							default:
+								break;
+						}
+					}
+
+					return InMaterial.Name;
+				}
+
+				return InMaterial.Name;
+			}
+
+			return InMaterialNode.NodeName;
+		}
+
+		private bool SetMasterMaterial(
+			Material                 InMaterial,
+			FDatasmithFacadeMaterial IOMasterMaterial
+		)
+		{
+			if (InMaterial == null)
+			{
+				// The properties of the Datasmith master material cannot be set.
+				return false;
+			}
+
+			Asset RenderingAsset = (InMaterial.Document.GetElement(InMaterial.AppearanceAssetId) as AppearanceAssetElement)?.GetRenderingAsset();
+
+			if (RenderingAsset == null)
+			{
+				// The properties of the Datasmith master material cannot be set.
+				return false;
+			}
+
+			string RenderingAssetName = RenderingAsset.Name.Replace("Schema", "");
+			Type   RenderingAssetType = Type.GetType($"Autodesk.Revit.DB.Visual.{RenderingAssetName},RevitAPI");
+
+			if (RenderingAssetType == null)
+			{
+				// The properties of the Datasmith master material cannot be set.
+				return false;
+			}
+
+			Color sourceMaterialColor = InMaterial.Color.IsValid ? InMaterial.Color : new Color(255, 255, 255);
+
+			// TODO: Some master material setup code below should be put in reusable methods.
+
+			switch (RenderingAssetType.Name)
+			{
+				case "Ceramic":
+				{
+					IOMasterMaterial.SetMasterMaterialType(FDatasmithFacadeMaterial.EMasterMaterialType.Opaque);
+
+					Color color = GetColorPropertyValue(RenderingAsset, Ceramic.CeramicColor, sourceMaterialColor);
+
+					// Control the Unreal material Base Color.
+					IOMasterMaterial.AddColor("DiffuseColor", color.Red / 255.0F, color.Green / 255.0F, color.Blue / 255.0F, 1.0F);
+
+					string diffuseMapPath = GetTexturePropertyPath(RenderingAsset, Ceramic.CeramicColor);
+
+					if (!string.IsNullOrEmpty(diffuseMapPath))
+					{
+						float diffuseMapUVOffsetX = GetTexturePropertyDistance(RenderingAsset, Ceramic.CeramicColor, UnifiedBitmap.TextureRealWorldOffsetX, 0.0F);
+						float diffuseMapUVOffsetY = GetTexturePropertyDistance(RenderingAsset, Ceramic.CeramicColor, UnifiedBitmap.TextureRealWorldOffsetY, 0.0F);
+						float diffuseMapUVScaleX  = 1.0F / GetTexturePropertyDistance(RenderingAsset, Ceramic.CeramicColor, UnifiedBitmap.TextureRealWorldScaleX, 1.0F);
+						float diffuseMapUVScaleY  = 1.0F / GetTexturePropertyDistance(RenderingAsset, Ceramic.CeramicColor, UnifiedBitmap.TextureRealWorldScaleY, 1.0F);
+						float diffuseMapUVWAngle  = GetTexturePropertyAngle(RenderingAsset, Ceramic.CeramicColor, UnifiedBitmap.TextureWAngle);
+
+						// Control the Unreal material Base Color.
+						IOMasterMaterial.AddFloat("DiffuseMapFading", 1.0F);                          
+						IOMasterMaterial.AddTexture("DiffuseMap", diffuseMapPath);                    
+						IOMasterMaterial.AddFloat("DiffuseMap_UVOffsetX", diffuseMapUVOffsetX);
+						IOMasterMaterial.AddFloat("DiffuseMap_UVOffsetY", diffuseMapUVOffsetY);
+						IOMasterMaterial.AddFloat("DiffuseMap_UVScaleX",  diffuseMapUVScaleX);
+						IOMasterMaterial.AddFloat("DiffuseMap_UVScaleY",  diffuseMapUVScaleY);
+						IOMasterMaterial.AddFloat("DiffuseMap_UVWAngle",  diffuseMapUVWAngle);
+					}
+					else
+					{
+						IOMasterMaterial.AddFloat("DiffuseMapFading", 0.0F);                          
+					}
+
+					bool  tintEnabled = GetBooleanPropertyValue(RenderingAsset, Ceramic.CommonTintToggle, false);
+					Color tintColor   = GetColorPropertyValue(RenderingAsset, Ceramic.CommonTintColor, sourceMaterialColor);
+
+					// Control the Unreal material Base Color.
+					IOMasterMaterial.AddBoolean("TintEnabled", tintEnabled);
+					IOMasterMaterial.AddColor("TintColor", tintColor.Red / 255.0F, tintColor.Green / 255.0F, tintColor.Blue / 255.0F, 1.0F);
+
+					// Control the Unreal material Metallic.
+					IOMasterMaterial.AddBoolean("IsMetal", false);
+
+					float glossiness = 0.6F; // CeramicApplicationType.Satin
+					switch ((CeramicApplicationType) GetIntegerPropertyValue(RenderingAsset, Ceramic.CeramicApplication, (int) CeramicApplicationType.Satin))
+					{
+						case CeramicApplicationType.HighGlossy:
+							glossiness = 0.9F;
+							break;
+						case CeramicApplicationType.Satin:
+							glossiness = 0.6F;
+							break;
+						case CeramicApplicationType.Matte:
+							glossiness = 0.35F;
+							break;
+					}                        
+
+					// Control the Unreal material Roughness.
+					IOMasterMaterial.AddFloat("Glossiness", glossiness);
+
+					// Control the Unreal material Emissive Color.
+					IOMasterMaterial.AddFloat("SelfIlluminationLuminance", 0.0F);
+
+					if ((CeramicPatternType) GetIntegerPropertyValue(RenderingAsset, Ceramic.CeramicPattern, (int) CeramicPatternType.None) == CeramicPatternType.Custom)
+					{
+						string bumpMapPath = GetTexturePropertyPath(RenderingAsset, Ceramic.CeramicPatternMap);
+
+						if (!string.IsNullOrEmpty(bumpMapPath))
+						{
+							float bumpAmount       = GetFloatPropertyValue(RenderingAsset, Ceramic.CeramicPatternAmount, 0.0F);
+							float bumpMapUVOffsetX = GetTexturePropertyDistance(RenderingAsset, Ceramic.CeramicPatternMap, UnifiedBitmap.TextureRealWorldOffsetX, 0.0F);
+							float bumpMapUVOffsetY = GetTexturePropertyDistance(RenderingAsset, Ceramic.CeramicPatternMap, UnifiedBitmap.TextureRealWorldOffsetY, 0.0F);
+							float bumpMapUVScaleX  = 1.0F / GetTexturePropertyDistance(RenderingAsset, Ceramic.CeramicPatternMap, UnifiedBitmap.TextureRealWorldScaleX, 1.0F);
+							float bumpMapUVScaleY  = 1.0F / GetTexturePropertyDistance(RenderingAsset, Ceramic.CeramicPatternMap, UnifiedBitmap.TextureRealWorldScaleY, 1.0F);
+							float bumpMapUVWAngle  = GetTexturePropertyAngle(RenderingAsset, Ceramic.CeramicPatternMap, UnifiedBitmap.TextureWAngle);
+
+							// Control the Unreal material Normal.
+							IOMasterMaterial.AddFloat("BumpAmount", bumpAmount);
+							IOMasterMaterial.AddTexture("BumpMap", bumpMapPath);
+							IOMasterMaterial.AddFloat("BumpMap_UVOffsetX", bumpMapUVOffsetX);
+							IOMasterMaterial.AddFloat("BumpMap_UVOffsetY", bumpMapUVOffsetY);
+							IOMasterMaterial.AddFloat("BumpMap_UVScaleX",  bumpMapUVScaleX);
+							IOMasterMaterial.AddFloat("BumpMap_UVScaleY",  bumpMapUVScaleY);
+							IOMasterMaterial.AddFloat("BumpMap_UVWAngle",  bumpMapUVWAngle);
+						}
+						else
+						{
+							IOMasterMaterial.AddFloat("BumpAmount", 0.0F);
+						}
+					}
+					else
+					{
+						IOMasterMaterial.AddFloat("BumpAmount", 0.0F);
+					}
+				}
+				break;
+
+				case "Concrete":
+				{
+					IOMasterMaterial.SetMasterMaterialType(FDatasmithFacadeMaterial.EMasterMaterialType.Opaque);
+
+					Color color = GetColorPropertyValue(RenderingAsset, Concrete.ConcreteColor, sourceMaterialColor);
+
+					// Control the Unreal material Base Color.
+					IOMasterMaterial.AddColor("DiffuseColor", color.Red / 255.0F, color.Green / 255.0F, color.Blue / 255.0F, 1.0F);
+
+					string diffuseMapPath = GetTexturePropertyPath(RenderingAsset, Concrete.ConcreteColor);
+
+					if (!string.IsNullOrEmpty(diffuseMapPath))
+					{
+						float diffuseMapUVOffsetX = GetTexturePropertyDistance(RenderingAsset, Concrete.ConcreteColor, UnifiedBitmap.TextureRealWorldOffsetX, 0.0F);
+						float diffuseMapUVOffsetY = GetTexturePropertyDistance(RenderingAsset, Concrete.ConcreteColor, UnifiedBitmap.TextureRealWorldOffsetY, 0.0F);
+						float diffuseMapUVScaleX  = 1.0F / GetTexturePropertyDistance(RenderingAsset, Concrete.ConcreteColor, UnifiedBitmap.TextureRealWorldScaleX, 1.0F);
+						float diffuseMapUVScaleY  = 1.0F / GetTexturePropertyDistance(RenderingAsset, Concrete.ConcreteColor, UnifiedBitmap.TextureRealWorldScaleY, 1.0F);
+						float diffuseMapUVWAngle  = GetTexturePropertyAngle(RenderingAsset, Concrete.ConcreteColor, UnifiedBitmap.TextureWAngle);
+
+						// Control the Unreal material Base Color.
+						IOMasterMaterial.AddFloat("DiffuseMapFading", 1.0F);                          
+						IOMasterMaterial.AddTexture("DiffuseMap", diffuseMapPath);                    
+						IOMasterMaterial.AddFloat("DiffuseMap_UVOffsetX", diffuseMapUVOffsetX);
+						IOMasterMaterial.AddFloat("DiffuseMap_UVOffsetY", diffuseMapUVOffsetY);
+						IOMasterMaterial.AddFloat("DiffuseMap_UVScaleX",  diffuseMapUVScaleX);
+						IOMasterMaterial.AddFloat("DiffuseMap_UVScaleY",  diffuseMapUVScaleY);
+						IOMasterMaterial.AddFloat("DiffuseMap_UVWAngle",  diffuseMapUVWAngle);
+					}
+					else
+					{
+						IOMasterMaterial.AddFloat("DiffuseMapFading", 0.0F);                          
+					}
+
+					bool  tintEnabled = GetBooleanPropertyValue(RenderingAsset, Concrete.CommonTintToggle, false);
+					Color tintColor   = GetColorPropertyValue(RenderingAsset, Concrete.CommonTintColor, sourceMaterialColor);
+
+					// Control the Unreal material Base Color.
+					IOMasterMaterial.AddBoolean("TintEnabled", tintEnabled);
+					IOMasterMaterial.AddColor("TintColor", tintColor.Red / 255.0F, tintColor.Green / 255.0F, tintColor.Blue / 255.0F, 1.0F);
+
+					// Control the Unreal material Metallic.
+					IOMasterMaterial.AddBoolean("IsMetal", false);
+
+					// Control the Unreal material Roughness.
+					IOMasterMaterial.AddFloat("Glossiness", 0.5F);
+
+					// Control the Unreal material Emissive Color.
+					IOMasterMaterial.AddFloat("SelfIlluminationLuminance", 0.0F);
+
+					if ((ConcreteFinishType) GetIntegerPropertyValue(RenderingAsset, Concrete.ConcreteFinish, (int) ConcreteFinishType.Straight) == ConcreteFinishType.Custom)
+					{
+						string bumpMapPath = GetTexturePropertyPath(RenderingAsset, Concrete.ConcreteBumpMap);
+
+						if (!string.IsNullOrEmpty(bumpMapPath))
+						{
+							float bumpAmount  = GetFloatPropertyValue(RenderingAsset, Concrete.ConcreteBumpAmount, 0.0F);
+							float bumpMapUVOffsetX = GetTexturePropertyDistance(RenderingAsset, Concrete.ConcreteBumpMap, UnifiedBitmap.TextureRealWorldOffsetX, 0.0F);
+							float bumpMapUVOffsetY = GetTexturePropertyDistance(RenderingAsset, Concrete.ConcreteBumpMap, UnifiedBitmap.TextureRealWorldOffsetY, 0.0F);
+							float bumpMapUVScaleX  = 1.0F / GetTexturePropertyDistance(RenderingAsset, Concrete.ConcreteBumpMap, UnifiedBitmap.TextureRealWorldScaleX, 1.0F);
+							float bumpMapUVScaleY  = 1.0F / GetTexturePropertyDistance(RenderingAsset, Concrete.ConcreteBumpMap, UnifiedBitmap.TextureRealWorldScaleY, 1.0F);
+							float bumpMapUVWAngle  = GetTexturePropertyAngle(RenderingAsset, Concrete.ConcreteBumpMap, UnifiedBitmap.TextureWAngle);
+
+							// Control the Unreal material Normal.
+							IOMasterMaterial.AddFloat("BumpAmount", bumpAmount);
+							IOMasterMaterial.AddTexture("BumpMap", bumpMapPath);
+							IOMasterMaterial.AddFloat("BumpMap_UVOffsetX", bumpMapUVOffsetX);
+							IOMasterMaterial.AddFloat("BumpMap_UVOffsetY", bumpMapUVOffsetY);
+							IOMasterMaterial.AddFloat("BumpMap_UVScaleX",  bumpMapUVScaleX);
+							IOMasterMaterial.AddFloat("BumpMap_UVScaleY",  bumpMapUVScaleY);
+							IOMasterMaterial.AddFloat("BumpMap_UVWAngle",  bumpMapUVWAngle);
+						}
+						else
+						{
+							IOMasterMaterial.AddFloat("BumpAmount", 0.0F);
+						}
+					}
+					else
+					{
+						IOMasterMaterial.AddFloat("BumpAmount", 0.0F);
+					}
+				}
+				break;
+
+				case "Glazing":
+				{
+					IOMasterMaterial.SetMasterMaterialType(FDatasmithFacadeMaterial.EMasterMaterialType.Transparent);
+
+					// TODO: Should use the Glazing.GlazingTransmittanceColor to select a predefined color value.
+					Color color           = sourceMaterialColor;
+					string diffuseMapPath = null;
+					if ((GlazingTransmittanceColorType) GetIntegerPropertyValue(RenderingAsset, Glazing.GlazingTransmittanceColor, (int) GlazingTransmittanceColorType.Clear) == GlazingTransmittanceColorType.Custom)
+					{
+						color          = GetColorPropertyValue(RenderingAsset, Glazing.GlazingTransmittanceMap, sourceMaterialColor);
+						diffuseMapPath = GetTexturePropertyPath(RenderingAsset, Glazing.GlazingTransmittanceMap);
+					}
+
+					// Control the Unreal material Base Color.
+					IOMasterMaterial.AddColor("DiffuseColor", color.Red / 255.0F, color.Green / 255.0F, color.Blue / 255.0F, 1.0F);
+
+					if (!string.IsNullOrEmpty(diffuseMapPath))
+					{
+						float diffuseMapUVOffsetX = GetTexturePropertyDistance(RenderingAsset, SolidGlass.SolidglassTransmittanceCustomColor, UnifiedBitmap.TextureRealWorldOffsetX, 0.0F);
+						float diffuseMapUVOffsetY = GetTexturePropertyDistance(RenderingAsset, SolidGlass.SolidglassTransmittanceCustomColor, UnifiedBitmap.TextureRealWorldOffsetY, 0.0F);
+						float diffuseMapUVScaleX  = 1.0F / GetTexturePropertyDistance(RenderingAsset, SolidGlass.SolidglassTransmittanceCustomColor, UnifiedBitmap.TextureRealWorldScaleX, 1.0F);
+						float diffuseMapUVScaleY  = 1.0F / GetTexturePropertyDistance(RenderingAsset, SolidGlass.SolidglassTransmittanceCustomColor, UnifiedBitmap.TextureRealWorldScaleY, 1.0F);
+						float diffuseMapUVWAngle  = GetTexturePropertyAngle(RenderingAsset, SolidGlass.SolidglassTransmittanceCustomColor, UnifiedBitmap.TextureWAngle);
+
+						// Control the Unreal material Base Color.
+						IOMasterMaterial.AddFloat("DiffuseMapFading", 1.0F);                          
+						IOMasterMaterial.AddTexture("DiffuseMap", diffuseMapPath);                    
+						IOMasterMaterial.AddFloat("DiffuseMap_UVOffsetX", diffuseMapUVOffsetX);
+						IOMasterMaterial.AddFloat("DiffuseMap_UVOffsetY", diffuseMapUVOffsetY);
+						IOMasterMaterial.AddFloat("DiffuseMap_UVScaleX",  diffuseMapUVScaleX);
+						IOMasterMaterial.AddFloat("DiffuseMap_UVScaleY",  diffuseMapUVScaleY);
+						IOMasterMaterial.AddFloat("DiffuseMap_UVWAngle",  diffuseMapUVWAngle);
+					}
+					else
+					{
+						IOMasterMaterial.AddFloat("DiffuseMapFading", 0.0F);                          
+					}
+
+					bool  tintEnabled = GetBooleanPropertyValue(RenderingAsset, Glazing.CommonTintToggle, false);
+					Color tintColor   = GetColorPropertyValue(RenderingAsset, Glazing.CommonTintColor, sourceMaterialColor);
+
+					// Control the Unreal material Base Color.
+					IOMasterMaterial.AddBoolean("TintEnabled", tintEnabled);
+					IOMasterMaterial.AddColor("TintColor", tintColor.Red / 255.0F, tintColor.Green / 255.0F, tintColor.Blue / 255.0F, 1.0F);
+
+					// Control the Unreal material Emissive Color.
+					IOMasterMaterial.AddFloat("SelfIlluminationLuminance", 0.0F);
+
+					float transparency = 1.0F - GetFloatPropertyValue(RenderingAsset, Glazing.GlazingReflectance, 0.0F);
+
+					// Control the Unreal material Opacity.
+					IOMasterMaterial.AddFloat("Transparency", transparency);
+					IOMasterMaterial.AddFloat("TransparencyMapFading", 0.0F);
+
+					// Control the Unreal material Normal.
+					IOMasterMaterial.AddFloat("BumpAmount", 0.0F);
+
+					// Control the Unreal material Refraction.
+					IOMasterMaterial.AddFloat("RefractionIndex", 1.01F);
+				}
+				break;
+
+				case "Hardwood":
+				{
+					IOMasterMaterial.SetMasterMaterialType(FDatasmithFacadeMaterial.EMasterMaterialType.Opaque);
+
+					// Control the Unreal material Base Color.
+					IOMasterMaterial.AddColor("DiffuseColor", sourceMaterialColor.Red / 255.0F, sourceMaterialColor.Green / 255.0F, sourceMaterialColor.Blue / 255.0F, 1.0F);
+
+					string diffuseMapPath = GetTexturePropertyPath(RenderingAsset, Hardwood.HardwoodColor);
+
+					if (!string.IsNullOrEmpty(diffuseMapPath))
+					{
+						float diffuseMapUVOffsetX = GetTexturePropertyDistance(RenderingAsset, Hardwood.HardwoodColor, UnifiedBitmap.TextureRealWorldOffsetX, 0.0F);
+						float diffuseMapUVOffsetY = GetTexturePropertyDistance(RenderingAsset, Hardwood.HardwoodColor, UnifiedBitmap.TextureRealWorldOffsetY, 0.0F);
+						float diffuseMapUVScaleX  = 1.0F / GetTexturePropertyDistance(RenderingAsset, Hardwood.HardwoodColor, UnifiedBitmap.TextureRealWorldScaleX, 1.0F);
+						float diffuseMapUVScaleY  = 1.0F / GetTexturePropertyDistance(RenderingAsset, Hardwood.HardwoodColor, UnifiedBitmap.TextureRealWorldScaleY, 1.0F);
+						float diffuseMapUVWAngle  = GetTexturePropertyAngle(RenderingAsset, Hardwood.HardwoodColor, UnifiedBitmap.TextureWAngle);
+
+						// Control the Unreal material Base Color.
+						IOMasterMaterial.AddFloat("DiffuseMapFading", 1.0F);                          
+						IOMasterMaterial.AddTexture("DiffuseMap", diffuseMapPath);                    
+						IOMasterMaterial.AddFloat("DiffuseMap_UVOffsetX", diffuseMapUVOffsetX);
+						IOMasterMaterial.AddFloat("DiffuseMap_UVOffsetY", diffuseMapUVOffsetY);
+						IOMasterMaterial.AddFloat("DiffuseMap_UVScaleX",  diffuseMapUVScaleX);
+						IOMasterMaterial.AddFloat("DiffuseMap_UVScaleY",  diffuseMapUVScaleY);
+						IOMasterMaterial.AddFloat("DiffuseMap_UVWAngle",  diffuseMapUVWAngle);
+					}
+					else
+					{
+						IOMasterMaterial.AddFloat("DiffuseMapFading", 0.0F);                          
+					}
+
+					bool  tintEnabled = GetBooleanPropertyValue(RenderingAsset, Hardwood.CommonTintToggle, false);
+					Color tintColor   = GetColorPropertyValue(RenderingAsset, Hardwood.CommonTintColor, sourceMaterialColor);
+
+					// Control the Unreal material Base Color.
+					IOMasterMaterial.AddBoolean("TintEnabled", tintEnabled);
+					IOMasterMaterial.AddColor("TintColor", tintColor.Red / 255.0F, tintColor.Green / 255.0F, tintColor.Blue / 255.0F, 1.0F);
+
+					// Control the Unreal material Metallic.
+					IOMasterMaterial.AddBoolean("IsMetal", false);
+
+					float glossiness = 0.45F; // HardwoodFinishType.Satin
+					switch ((HardwoodFinishType) GetIntegerPropertyValue(RenderingAsset, Hardwood.HardwoodFinish, (int) HardwoodFinishType.Satin))
+					{
+						case HardwoodFinishType.Gloss:
+							glossiness = 0.8F;
+							break;
+						case HardwoodFinishType.Semigloss:
+							glossiness = 0.65F;
+							break;
+						case HardwoodFinishType.Satin:
+							glossiness = 0.45F;
+							break;
+						case HardwoodFinishType.Unfinished:
+							glossiness = 0.05F;
+							break;
+					}                        
+
+					// Control the Unreal material Roughness.
+					IOMasterMaterial.AddFloat("Glossiness", glossiness);
+
+					// Control the Unreal material Emissive Color.
+					IOMasterMaterial.AddFloat("SelfIlluminationLuminance", 0.0F);
+
+					HardwoodImperfectionsType hardwoodImperfectionsType = (HardwoodImperfectionsType) GetIntegerPropertyValue(RenderingAsset, Hardwood.HardwoodImperfections, (int) HardwoodImperfectionsType.None);
+					if (hardwoodImperfectionsType != HardwoodImperfectionsType.None)
+					{
+						string texturePropertyName = (hardwoodImperfectionsType == HardwoodImperfectionsType.Automatic) ? Hardwood.HardwoodColor : Hardwood.HardwoodImperfectionsShader;
+
+						string bumpMapPath = GetTexturePropertyPath(RenderingAsset, texturePropertyName);
+
+						if (!string.IsNullOrEmpty(bumpMapPath))
+						{
+							float bumpAmount       = GetFloatPropertyValue(RenderingAsset, Hardwood.HardwoodImperfectionsAmount, 0.0F);
+							float bumpMapUVOffsetX = GetTexturePropertyDistance(RenderingAsset, texturePropertyName, UnifiedBitmap.TextureRealWorldOffsetX, 0.0F);
+							float bumpMapUVOffsetY = GetTexturePropertyDistance(RenderingAsset, texturePropertyName, UnifiedBitmap.TextureRealWorldOffsetY, 0.0F);
+							float bumpMapUVScaleX  = 1.0F / GetTexturePropertyDistance(RenderingAsset, texturePropertyName, UnifiedBitmap.TextureRealWorldScaleX, 1.0F);
+							float bumpMapUVScaleY  = 1.0F / GetTexturePropertyDistance(RenderingAsset, texturePropertyName, UnifiedBitmap.TextureRealWorldScaleY, 1.0F);
+							float bumpMapUVWAngle  = GetTexturePropertyAngle(RenderingAsset, texturePropertyName, UnifiedBitmap.TextureWAngle);
+
+							// Control the Unreal material Normal.
+							IOMasterMaterial.AddFloat("BumpAmount", bumpAmount);
+							IOMasterMaterial.AddTexture("BumpMap", bumpMapPath);
+							IOMasterMaterial.AddFloat("BumpMap_UVOffsetX", bumpMapUVOffsetX);
+							IOMasterMaterial.AddFloat("BumpMap_UVOffsetY", bumpMapUVOffsetY);
+							IOMasterMaterial.AddFloat("BumpMap_UVScaleX",  bumpMapUVScaleX);
+							IOMasterMaterial.AddFloat("BumpMap_UVScaleY",  bumpMapUVScaleY);
+							IOMasterMaterial.AddFloat("BumpMap_UVWAngle",  bumpMapUVWAngle);
+						}
+						else
+						{
+							IOMasterMaterial.AddFloat("BumpAmount", 0.0F);
+						}
+					}
+					else
+					{
+						IOMasterMaterial.AddFloat("BumpAmount", 0.0F);
+					}
+				}
+				break;
+
+				case "MasonryCMU":
+				{
+					IOMasterMaterial.SetMasterMaterialType(FDatasmithFacadeMaterial.EMasterMaterialType.Opaque);
+
+					Color color = GetColorPropertyValue(RenderingAsset, MasonryCMU.MasonryCMUColor, sourceMaterialColor);
+
+					// Control the Unreal material Base Color.
+					IOMasterMaterial.AddColor("DiffuseColor", color.Red / 255.0F, color.Green / 255.0F, color.Blue / 255.0F, 1.0F);
+
+					string diffuseMapPath = GetTexturePropertyPath(RenderingAsset, MasonryCMU.MasonryCMUColor);
+
+					if (!string.IsNullOrEmpty(diffuseMapPath))
+					{
+						float diffuseMapUVOffsetX = GetTexturePropertyDistance(RenderingAsset, MasonryCMU.MasonryCMUColor, UnifiedBitmap.TextureRealWorldOffsetX, 0.0F);
+						float diffuseMapUVOffsetY = GetTexturePropertyDistance(RenderingAsset, MasonryCMU.MasonryCMUColor, UnifiedBitmap.TextureRealWorldOffsetY, 0.0F);
+						float diffuseMapUVScaleX  = 1.0F / GetTexturePropertyDistance(RenderingAsset, MasonryCMU.MasonryCMUColor, UnifiedBitmap.TextureRealWorldScaleX, 1.0F);
+						float diffuseMapUVScaleY  = 1.0F / GetTexturePropertyDistance(RenderingAsset, MasonryCMU.MasonryCMUColor, UnifiedBitmap.TextureRealWorldScaleY, 1.0F);
+						float diffuseMapUVWAngle  = GetTexturePropertyAngle(RenderingAsset, MasonryCMU.MasonryCMUColor, UnifiedBitmap.TextureWAngle);
+
+						// Control the Unreal material Base Color.
+						IOMasterMaterial.AddFloat("DiffuseMapFading", 1.0F);                          
+						IOMasterMaterial.AddTexture("DiffuseMap", diffuseMapPath);                    
+						IOMasterMaterial.AddFloat("DiffuseMap_UVOffsetX", diffuseMapUVOffsetX);
+						IOMasterMaterial.AddFloat("DiffuseMap_UVOffsetY", diffuseMapUVOffsetY);
+						IOMasterMaterial.AddFloat("DiffuseMap_UVScaleX",  diffuseMapUVScaleX);
+						IOMasterMaterial.AddFloat("DiffuseMap_UVScaleY",  diffuseMapUVScaleY);
+						IOMasterMaterial.AddFloat("DiffuseMap_UVWAngle",  diffuseMapUVWAngle);
+					}
+					else
+					{
+						IOMasterMaterial.AddFloat("DiffuseMapFading", 0.0F);                          
+					}
+
+					bool  tintEnabled = GetBooleanPropertyValue(RenderingAsset, MasonryCMU.CommonTintToggle, false);
+					Color tintColor   = GetColorPropertyValue(RenderingAsset, MasonryCMU.CommonTintColor, sourceMaterialColor);
+
+					// Control the Unreal material Base Color.
+					IOMasterMaterial.AddBoolean("TintEnabled", tintEnabled);
+					IOMasterMaterial.AddColor("TintColor", tintColor.Red / 255.0F, tintColor.Green / 255.0F, tintColor.Blue / 255.0F, 1.0F);
+
+					// Control the Unreal material Metallic.
+					IOMasterMaterial.AddBoolean("IsMetal", false);
+
+					float glossiness = 0.25F; // MasonryCMUApplicationType.Matte
+					switch ((MasonryCMUApplicationType) GetIntegerPropertyValue(RenderingAsset, MasonryCMU.MasonryCMUApplication, (int) MasonryCMUApplicationType.Matte))
+					{
+						case MasonryCMUApplicationType.Glossy:
+							glossiness = 0.9F;
+							break;
+						case MasonryCMUApplicationType.Matte:
+						case MasonryCMUApplicationType.Unfinished:
+							glossiness = 0.25F;
+							break;
+					}                        
+
+					// Control the Unreal material Roughness.
+					IOMasterMaterial.AddFloat("Glossiness", glossiness);
+
+					// Control the Unreal material Emissive Color.
+					IOMasterMaterial.AddFloat("SelfIlluminationLuminance", 0.0F);
+
+					if ((MasonryCMUPatternType) GetIntegerPropertyValue(RenderingAsset, MasonryCMU.MasonryCMUPattern, (int) MasonryCMUPatternType.None) == MasonryCMUPatternType.Custom)
+					{
+						string bumpMapPath = GetTexturePropertyPath(RenderingAsset, MasonryCMU.MasonryCMUPatternMap);
+
+						if (!string.IsNullOrEmpty(bumpMapPath))
+						{
+							float bumpAmount       = GetFloatPropertyValue(RenderingAsset, MasonryCMU.MasonryCMUPatternHeight, 0.0F);
+							float bumpMapUVOffsetX = GetTexturePropertyDistance(RenderingAsset, MasonryCMU.MasonryCMUPatternMap, UnifiedBitmap.TextureRealWorldOffsetX, 0.0F);
+							float bumpMapUVOffsetY = GetTexturePropertyDistance(RenderingAsset, MasonryCMU.MasonryCMUPatternMap, UnifiedBitmap.TextureRealWorldOffsetY, 0.0F);
+							float bumpMapUVScaleX  = 1.0F / GetTexturePropertyDistance(RenderingAsset, MasonryCMU.MasonryCMUPatternMap, UnifiedBitmap.TextureRealWorldScaleX, 1.0F);
+							float bumpMapUVScaleY  = 1.0F / GetTexturePropertyDistance(RenderingAsset, MasonryCMU.MasonryCMUPatternMap, UnifiedBitmap.TextureRealWorldScaleY, 1.0F);
+							float bumpMapUVWAngle  = GetTexturePropertyAngle(RenderingAsset, MasonryCMU.MasonryCMUPatternMap, UnifiedBitmap.TextureWAngle);
+
+							// Control the Unreal material Normal.
+							IOMasterMaterial.AddFloat("BumpAmount", bumpAmount);
+							IOMasterMaterial.AddTexture("BumpMap", bumpMapPath);
+							IOMasterMaterial.AddFloat("BumpMap_UVOffsetX", bumpMapUVOffsetX);
+							IOMasterMaterial.AddFloat("BumpMap_UVOffsetY", bumpMapUVOffsetY);
+							IOMasterMaterial.AddFloat("BumpMap_UVScaleX",  bumpMapUVScaleX);
+							IOMasterMaterial.AddFloat("BumpMap_UVScaleY",  bumpMapUVScaleY);
+							IOMasterMaterial.AddFloat("BumpMap_UVWAngle",  bumpMapUVWAngle);
+						}
+						else
+						{
+							IOMasterMaterial.AddFloat("BumpAmount", 0.0F);
+						}
+					}
+					else
+					{
+						IOMasterMaterial.AddFloat("BumpAmount", 0.0F);
+					}
+				}
+				break;
+
+				case "Metal":
+				{
+					IOMasterMaterial.SetMasterMaterialType(FDatasmithFacadeMaterial.EMasterMaterialType.Opaque);
+
+					// TODO: Should use the Metal.MetalColor to select a predefined color value.
+
+					// Control the Unreal material Base Color.
+					IOMasterMaterial.AddColor("DiffuseColor", sourceMaterialColor.Red / 255.0F, sourceMaterialColor.Green / 255.0F, sourceMaterialColor.Blue / 255.0F, 1.0F);
+					IOMasterMaterial.AddFloat("DiffuseMapFading", 0.0F);
+
+					bool  tintEnabled = GetBooleanPropertyValue(RenderingAsset, Metal.CommonTintToggle, false);
+					Color tintColor   = GetColorPropertyValue(RenderingAsset, Metal.CommonTintColor, sourceMaterialColor);
+
+					// Control the Unreal material Base Color.
+					IOMasterMaterial.AddBoolean("TintEnabled", tintEnabled);
+					IOMasterMaterial.AddColor("TintColor", tintColor.Red / 255.0F, tintColor.Green / 255.0F, tintColor.Blue / 255.0F, 1.0F);
+
+					// Control the Unreal material Metallic.
+					IOMasterMaterial.AddBoolean("IsMetal", true);
+
+					float glossiness = 0.5F; // MetalFinishType.SemiPolished
+					switch ((MetalFinishType) GetIntegerPropertyValue(RenderingAsset, Metal.MetalFinish, (int) MetalFinishType.SemiPolished))
+					{
+						case MetalFinishType.Polished:
+							glossiness = 1.0F;
+							break;
+						case MetalFinishType.SemiPolished:
+							glossiness = 0.5F;
+							break;
+						case MetalFinishType.Satin:
+						case MetalFinishType.Brushed:
+							glossiness = 0.25F;
+							break;
+					}
+
+					// Control the Unreal material Roughness.
+					IOMasterMaterial.AddFloat("Glossiness", glossiness);
+
+					// Control the Unreal material Emissive Color.
+					IOMasterMaterial.AddFloat("SelfIlluminationLuminance", 0.0F);
+
+					if ((MetalPerforationsType) GetIntegerPropertyValue(RenderingAsset, Metal.MetalPerforations, (int) MetalPerforationsType.None) == MetalPerforationsType.Custom)
+					{
+						string cutoutMapPath = GetTexturePropertyPath(RenderingAsset, Metal.MetalPerforationsShader);
+
+						if (!string.IsNullOrEmpty(cutoutMapPath))
+						{
+							float cutoutMapUVOffsetX = GetTexturePropertyDistance(RenderingAsset, Metal.MetalPerforationsShader, UnifiedBitmap.TextureRealWorldOffsetX, 0.0F);
+							float cutoutMapUVOffsetY = GetTexturePropertyDistance(RenderingAsset, Metal.MetalPerforationsShader, UnifiedBitmap.TextureRealWorldOffsetY, 0.0F);
+							float cutoutMapUVScaleX  = 1.0F / GetTexturePropertyDistance(RenderingAsset, Metal.MetalPerforationsShader, UnifiedBitmap.TextureRealWorldScaleX, 1.0F);
+							float cutoutMapUVScaleY  = 1.0F / GetTexturePropertyDistance(RenderingAsset, Metal.MetalPerforationsShader, UnifiedBitmap.TextureRealWorldScaleY, 1.0F);
+							float cutoutMapUVWAngle  = GetTexturePropertyAngle(RenderingAsset, Metal.MetalPerforationsShader, UnifiedBitmap.TextureWAngle);
+
+							IOMasterMaterial.SetMasterMaterialType(FDatasmithFacadeMaterial.EMasterMaterialType.CutOut);
+
+							// Control the Unreal material Opacity Mask.
+							IOMasterMaterial.AddTexture("CutoutOpacityMap",  cutoutMapPath);
+							IOMasterMaterial.AddFloat("CutoutMap_UVOffsetX", cutoutMapUVOffsetX);
+							IOMasterMaterial.AddFloat("CutoutMap_UVOffsetY", cutoutMapUVOffsetY);
+							IOMasterMaterial.AddFloat("CutoutMap_UVScaleX",  cutoutMapUVScaleX);
+							IOMasterMaterial.AddFloat("CutoutMap_UVScaleY",  cutoutMapUVScaleY);
+							IOMasterMaterial.AddFloat("CutoutMap_UVWAngle",  cutoutMapUVWAngle);
+						}
+					}
+
+					if ((MetalPatternType) GetIntegerPropertyValue(RenderingAsset, Metal.MetalPattern, (int) MetalPatternType.None) == MetalPatternType.Custom)
+					{
+						string bumpMapPath = GetTexturePropertyPath(RenderingAsset, Metal.MetalPatternShader);
+
+						if (!string.IsNullOrEmpty(bumpMapPath))
+						{
+							float bumpAmount       = GetFloatPropertyValue(RenderingAsset, Metal.MetalPatternHeight, 0.0F);
+							float bumpMapUVOffsetX = GetTexturePropertyDistance(RenderingAsset, Metal.MetalPatternShader, UnifiedBitmap.TextureRealWorldOffsetX, 0.0F);
+							float bumpMapUVOffsetY = GetTexturePropertyDistance(RenderingAsset, Metal.MetalPatternShader, UnifiedBitmap.TextureRealWorldOffsetY, 0.0F);
+							float bumpMapUVScaleX  = 1.0F / GetTexturePropertyDistance(RenderingAsset, Metal.MetalPatternShader, UnifiedBitmap.TextureRealWorldScaleX, 1.0F);
+							float bumpMapUVScaleY  = 1.0F / GetTexturePropertyDistance(RenderingAsset, Metal.MetalPatternShader, UnifiedBitmap.TextureRealWorldScaleY, 1.0F);
+							float bumpMapUVWAngle  = GetTexturePropertyAngle(RenderingAsset, Metal.MetalPatternShader, UnifiedBitmap.TextureWAngle);
+
+							// Control the Unreal material Normal.
+							IOMasterMaterial.AddFloat("BumpAmount", bumpAmount);
+							IOMasterMaterial.AddTexture("BumpMap", bumpMapPath);
+							IOMasterMaterial.AddFloat("BumpMap_UVOffsetX", bumpMapUVOffsetX);
+							IOMasterMaterial.AddFloat("BumpMap_UVOffsetY", bumpMapUVOffsetY);
+							IOMasterMaterial.AddFloat("BumpMap_UVScaleX",  bumpMapUVScaleX);
+							IOMasterMaterial.AddFloat("BumpMap_UVScaleY",  bumpMapUVScaleY);
+							IOMasterMaterial.AddFloat("BumpMap_UVWAngle",  bumpMapUVWAngle);
+						}
+						else
+						{
+							IOMasterMaterial.AddFloat("BumpAmount", 0.0F);
+						}
+					}
+					else
+					{
+						IOMasterMaterial.AddFloat("BumpAmount", 0.0F);
+					}
+				}
+				break;
+
+				case "MetallicPaint":
+				{
+					IOMasterMaterial.SetMasterMaterialType(FDatasmithFacadeMaterial.EMasterMaterialType.Opaque);
+
+					Color color = GetColorPropertyValue(RenderingAsset, MetallicPaint.MetallicpaintBaseColor, sourceMaterialColor);
+
+					// Control the Unreal material Base Color.
+					IOMasterMaterial.AddColor("DiffuseColor", color.Red / 255.0F, color.Green / 255.0F, color.Blue / 255.0F, 1.0F);
+
+					string diffuseMapPath = GetTexturePropertyPath(RenderingAsset, MetallicPaint.MetallicpaintBaseColor);
+
+					if (!string.IsNullOrEmpty(diffuseMapPath))
+					{
+						float diffuseMapUVOffsetX = GetTexturePropertyDistance(RenderingAsset, MetallicPaint.MetallicpaintBaseColor, UnifiedBitmap.TextureRealWorldOffsetX, 0.0F);
+						float diffuseMapUVOffsetY = GetTexturePropertyDistance(RenderingAsset, MetallicPaint.MetallicpaintBaseColor, UnifiedBitmap.TextureRealWorldOffsetY, 0.0F);
+						float diffuseMapUVScaleX  = 1.0F / GetTexturePropertyDistance(RenderingAsset, MetallicPaint.MetallicpaintBaseColor, UnifiedBitmap.TextureRealWorldScaleX, 1.0F);
+						float diffuseMapUVScaleY  = 1.0F / GetTexturePropertyDistance(RenderingAsset, MetallicPaint.MetallicpaintBaseColor, UnifiedBitmap.TextureRealWorldScaleY, 1.0F);
+						float diffuseMapUVWAngle  = GetTexturePropertyAngle(RenderingAsset, MetallicPaint.MetallicpaintBaseColor, UnifiedBitmap.TextureWAngle);
+
+						// Control the Unreal material Base Color.
+						IOMasterMaterial.AddFloat("DiffuseMapFading", 1.0F);                          
+						IOMasterMaterial.AddTexture("DiffuseMap", diffuseMapPath);                    
+						IOMasterMaterial.AddFloat("DiffuseMap_UVOffsetX", diffuseMapUVOffsetX);
+						IOMasterMaterial.AddFloat("DiffuseMap_UVOffsetY", diffuseMapUVOffsetY);
+						IOMasterMaterial.AddFloat("DiffuseMap_UVScaleX",  diffuseMapUVScaleX);
+						IOMasterMaterial.AddFloat("DiffuseMap_UVScaleY",  diffuseMapUVScaleY);
+						IOMasterMaterial.AddFloat("DiffuseMap_UVWAngle",  diffuseMapUVWAngle);
+					}
+					else
+					{
+						IOMasterMaterial.AddFloat("DiffuseMapFading", 0.0F);                          
+					}
+
+					bool  tintEnabled = GetBooleanPropertyValue(RenderingAsset, MetallicPaint.CommonTintToggle, false);
+					Color tintColor   = GetColorPropertyValue(RenderingAsset, MetallicPaint.CommonTintColor, sourceMaterialColor);
+
+					// Control the Unreal material Base Color.
+					IOMasterMaterial.AddBoolean("TintEnabled", tintEnabled);
+					IOMasterMaterial.AddColor("TintColor", tintColor.Red / 255.0F, tintColor.Green / 255.0F, tintColor.Blue / 255.0F, 1.0F);
+
+					// Control the Unreal material Metallic.
+					IOMasterMaterial.AddBoolean("IsMetal", true);
+
+					// Control the Unreal material Roughness.
+					IOMasterMaterial.AddFloat("Glossiness", 0.9F);
+
+					// Control the Unreal material Emissive Color.
+					IOMasterMaterial.AddFloat("SelfIlluminationLuminance", 0.0F);
+
+					// Control the Unreal material Normal.
+					IOMasterMaterial.AddFloat("BumpAmount", 0.0F);
+				}
+				break;
+
+				case "Mirror":
+				{
+					IOMasterMaterial.SetMasterMaterialType(FDatasmithFacadeMaterial.EMasterMaterialType.Opaque);
+
+					Color color = GetColorPropertyValue(RenderingAsset, Mirror.MirrorTintcolor, sourceMaterialColor);
+
+					// Control the Unreal material Base Color.
+					IOMasterMaterial.AddColor("DiffuseColor", color.Red / 255.0F, color.Green / 255.0F, color.Blue / 255.0F, 1.0F);
+					IOMasterMaterial.AddFloat("DiffuseMapFading", 0.0F);
+
+					bool  tintEnabled = GetBooleanPropertyValue(RenderingAsset, Mirror.CommonTintToggle, false);
+					Color tintColor   = GetColorPropertyValue(RenderingAsset, Mirror.CommonTintColor, sourceMaterialColor);
+
+					// Control the Unreal material Base Color.
+					IOMasterMaterial.AddBoolean("TintEnabled", tintEnabled);
+					IOMasterMaterial.AddColor("TintColor", tintColor.Red / 255.0F, tintColor.Green / 255.0F, tintColor.Blue / 255.0F, 1.0F);
+
+					// Control the Unreal material Metallic.
+					IOMasterMaterial.AddBoolean("IsMetal", true);
+
+					// Control the Unreal material Roughness.
+					IOMasterMaterial.AddFloat("Glossiness", 1.0F);
+
+					// Control the Unreal material Emissive Color.
+					IOMasterMaterial.AddFloat("SelfIlluminationLuminance", 0.0F);
+
+					// Control the Unreal material Normal.
+					IOMasterMaterial.AddFloat("BumpAmount", 0.0F);
+				}
+				break;
+
+				case "PlasticVinyl":
+				{
+					Color color = GetColorPropertyValue(RenderingAsset, PlasticVinyl.PlasticvinylColor, sourceMaterialColor);
+
+					// Control the Unreal material Base Color.
+					IOMasterMaterial.AddColor("DiffuseColor", color.Red / 255.0F, color.Green / 255.0F, color.Blue / 255.0F, 1.0F);
+
+					string diffuseMapPath = GetTexturePropertyPath(RenderingAsset, PlasticVinyl.PlasticvinylColor);
+
+					if (!string.IsNullOrEmpty(diffuseMapPath))
+					{
+						float diffuseMapUVOffsetX = GetTexturePropertyDistance(RenderingAsset, PlasticVinyl.PlasticvinylColor, UnifiedBitmap.TextureRealWorldOffsetX, 0.0F);
+						float diffuseMapUVOffsetY = GetTexturePropertyDistance(RenderingAsset, PlasticVinyl.PlasticvinylColor, UnifiedBitmap.TextureRealWorldOffsetY, 0.0F);
+						float diffuseMapUVScaleX  = 1.0F / GetTexturePropertyDistance(RenderingAsset, PlasticVinyl.PlasticvinylColor, UnifiedBitmap.TextureRealWorldScaleX, 1.0F);
+						float diffuseMapUVScaleY  = 1.0F / GetTexturePropertyDistance(RenderingAsset, PlasticVinyl.PlasticvinylColor, UnifiedBitmap.TextureRealWorldScaleY, 1.0F);
+						float diffuseMapUVWAngle  = GetTexturePropertyAngle(RenderingAsset, PlasticVinyl.PlasticvinylColor, UnifiedBitmap.TextureWAngle);
+
+						// Control the Unreal material Base Color.
+						IOMasterMaterial.AddFloat("DiffuseMapFading", 1.0F);                          
+						IOMasterMaterial.AddTexture("DiffuseMap", diffuseMapPath);                    
+						IOMasterMaterial.AddFloat("DiffuseMap_UVOffsetX", diffuseMapUVOffsetX);
+						IOMasterMaterial.AddFloat("DiffuseMap_UVOffsetY", diffuseMapUVOffsetY);
+						IOMasterMaterial.AddFloat("DiffuseMap_UVScaleX",  diffuseMapUVScaleX);
+						IOMasterMaterial.AddFloat("DiffuseMap_UVScaleY",  diffuseMapUVScaleY);
+						IOMasterMaterial.AddFloat("DiffuseMap_UVWAngle",  diffuseMapUVWAngle);
+					}
+					else
+					{
+						IOMasterMaterial.AddFloat("DiffuseMapFading", 0.0F);                          
+					}
+
+					bool  tintEnabled = GetBooleanPropertyValue(RenderingAsset, PlasticVinyl.CommonTintToggle, false);
+					Color tintColor   = GetColorPropertyValue(RenderingAsset, PlasticVinyl.CommonTintColor, sourceMaterialColor);
+
+					// Control the Unreal material Base Color.
+					IOMasterMaterial.AddBoolean("TintEnabled", tintEnabled);
+					IOMasterMaterial.AddColor("TintColor", tintColor.Red / 255.0F, tintColor.Green / 255.0F, tintColor.Blue / 255.0F, 1.0F);
+
+					// Control the Unreal material Emissive Color.
+					IOMasterMaterial.AddFloat("SelfIlluminationLuminance", 0.0F);
+
+					if ((PlasticvinylBumpType) GetIntegerPropertyValue(RenderingAsset, PlasticVinyl.PlasticvinylBump, (int) PlasticvinylBumpType.None) == PlasticvinylBumpType.Custom)
+					{
+						string bumpMapPath = GetTexturePropertyPath(RenderingAsset, PlasticVinyl.PlasticvinylBumpMap);
+
+						if (!string.IsNullOrEmpty(bumpMapPath))
+						{
+							float bumpAmount       = GetFloatPropertyValue(RenderingAsset, PlasticVinyl.PlasticvinylBumpAmount, 0.0F);
+							float bumpMapUVOffsetX = GetTexturePropertyDistance(RenderingAsset, PlasticVinyl.PlasticvinylBumpMap, UnifiedBitmap.TextureRealWorldOffsetX, 0.0F);
+							float bumpMapUVOffsetY = GetTexturePropertyDistance(RenderingAsset, PlasticVinyl.PlasticvinylBumpMap, UnifiedBitmap.TextureRealWorldOffsetY, 0.0F);
+							float bumpMapUVScaleX  = 1.0F / GetTexturePropertyDistance(RenderingAsset, PlasticVinyl.PlasticvinylBumpMap, UnifiedBitmap.TextureRealWorldScaleX, 1.0F);
+							float bumpMapUVScaleY  = 1.0F / GetTexturePropertyDistance(RenderingAsset, PlasticVinyl.PlasticvinylBumpMap, UnifiedBitmap.TextureRealWorldScaleY, 1.0F);
+							float bumpMapUVWAngle  = GetTexturePropertyAngle(RenderingAsset, PlasticVinyl.PlasticvinylBumpMap, UnifiedBitmap.TextureWAngle);
+
+							// Control the Unreal material Normal.
+							IOMasterMaterial.AddFloat("BumpAmount", bumpAmount);
+							IOMasterMaterial.AddTexture("BumpMap", bumpMapPath);
+							IOMasterMaterial.AddFloat("BumpMap_UVOffsetX", bumpMapUVOffsetX);
+							IOMasterMaterial.AddFloat("BumpMap_UVOffsetY", bumpMapUVOffsetY);
+							IOMasterMaterial.AddFloat("BumpMap_UVScaleX",  bumpMapUVScaleX);
+							IOMasterMaterial.AddFloat("BumpMap_UVScaleY",  bumpMapUVScaleY);
+							IOMasterMaterial.AddFloat("BumpMap_UVWAngle",  bumpMapUVWAngle);
+						}
+						else
+						{
+							IOMasterMaterial.AddFloat("BumpAmount", 0.0F);
+						}
+					}
+					else
+					{
+						IOMasterMaterial.AddFloat("BumpAmount", 0.0F);
+					}
+
+					if ((PlasticvinylType) GetIntegerPropertyValue(RenderingAsset, PlasticVinyl.PlasticvinylType, (int) PlasticvinylType.Plasticsolid) == PlasticvinylType.Plastictransparent)
+					{
+						IOMasterMaterial.SetMasterMaterialType(FDatasmithFacadeMaterial.EMasterMaterialType.Transparent);
+
+						// Control the Unreal material Opacity.
+						IOMasterMaterial.AddFloat("Transparency", 0.5F);
+						IOMasterMaterial.AddFloat("TransparencyMapFading", 0.0F);
+
+						// Control the Unreal material Refraction.
+						IOMasterMaterial.AddFloat("RefractionIndex", 1.0F);
+					}
+					else
+					{
+						IOMasterMaterial.SetMasterMaterialType(FDatasmithFacadeMaterial.EMasterMaterialType.Opaque);
+
+						// Control the Unreal material Metallic.
+						IOMasterMaterial.AddBoolean("IsMetal", false);
+
+						float glossiness = 0.9F; // PlasticvinylApplicationType.Polished
+						switch ((PlasticvinylApplicationType) GetIntegerPropertyValue(RenderingAsset, PlasticVinyl.PlasticvinylApplication, (int) PlasticvinylApplicationType.Polished))
+						{
+							case PlasticvinylApplicationType.Glossy:
+								glossiness = 1.0F;
+								break;
+							case PlasticvinylApplicationType.Polished:
+								glossiness = 0.9F;
+								break;
+							case PlasticvinylApplicationType.Matte:
+								glossiness = 0.75F;
+								break;
+						}                        
+
+						// Control the Unreal material Roughness.
+						IOMasterMaterial.AddFloat("Glossiness", glossiness);
+					}
+				}
+				break;
+
+				case "SolidGlass":
+				{
+					IOMasterMaterial.SetMasterMaterialType(FDatasmithFacadeMaterial.EMasterMaterialType.Transparent);
+
+					// TODO: Should use the SolidGlass.SolidglassTransmittance to select a predefined color value.
+					Color color           = sourceMaterialColor;
+					string diffuseMapPath = null;
+					if ((SolidglassTransmittanceType) GetIntegerPropertyValue(RenderingAsset, SolidGlass.SolidglassTransmittance, (int) SolidglassTransmittanceType.Clear) == SolidglassTransmittanceType.CustomColor)
+					{
+						color          = GetColorPropertyValue(RenderingAsset, SolidGlass.SolidglassTransmittanceCustomColor, sourceMaterialColor);
+						diffuseMapPath = GetTexturePropertyPath(RenderingAsset, SolidGlass.SolidglassTransmittanceCustomColor);
+					}
+
+					// Control the Unreal material Base Color.
+					IOMasterMaterial.AddColor("DiffuseColor", color.Red / 255.0F, color.Green / 255.0F, color.Blue / 255.0F, 1.0F);
+
+					if (!string.IsNullOrEmpty(diffuseMapPath))
+					{
+						float diffuseMapUVOffsetX = GetTexturePropertyDistance(RenderingAsset, SolidGlass.SolidglassTransmittanceCustomColor, UnifiedBitmap.TextureRealWorldOffsetX, 0.0F);
+						float diffuseMapUVOffsetY = GetTexturePropertyDistance(RenderingAsset, SolidGlass.SolidglassTransmittanceCustomColor, UnifiedBitmap.TextureRealWorldOffsetY, 0.0F);
+						float diffuseMapUVScaleX  = 1.0F / GetTexturePropertyDistance(RenderingAsset, SolidGlass.SolidglassTransmittanceCustomColor, UnifiedBitmap.TextureRealWorldScaleX, 1.0F);
+						float diffuseMapUVScaleY  = 1.0F / GetTexturePropertyDistance(RenderingAsset, SolidGlass.SolidglassTransmittanceCustomColor, UnifiedBitmap.TextureRealWorldScaleY, 1.0F);
+						float diffuseMapUVWAngle  = GetTexturePropertyAngle(RenderingAsset, SolidGlass.SolidglassTransmittanceCustomColor, UnifiedBitmap.TextureWAngle);
+
+						// Control the Unreal material Base Color.
+						IOMasterMaterial.AddFloat("DiffuseMapFading", 1.0F);                          
+						IOMasterMaterial.AddTexture("DiffuseMap", diffuseMapPath);                    
+						IOMasterMaterial.AddFloat("DiffuseMap_UVOffsetX", diffuseMapUVOffsetX);
+						IOMasterMaterial.AddFloat("DiffuseMap_UVOffsetY", diffuseMapUVOffsetY);
+						IOMasterMaterial.AddFloat("DiffuseMap_UVScaleX",  diffuseMapUVScaleX);
+						IOMasterMaterial.AddFloat("DiffuseMap_UVScaleY",  diffuseMapUVScaleY);
+						IOMasterMaterial.AddFloat("DiffuseMap_UVWAngle",  diffuseMapUVWAngle);
+					}
+					else
+					{
+						IOMasterMaterial.AddFloat("DiffuseMapFading", 0.0F);                          
+					}
+
+					bool  tintEnabled = GetBooleanPropertyValue(RenderingAsset, SolidGlass.CommonTintToggle, false);
+					Color tintColor   = GetColorPropertyValue(RenderingAsset, SolidGlass.CommonTintColor, sourceMaterialColor);
+
+					// Control the Unreal material Base Color.
+					IOMasterMaterial.AddBoolean("TintEnabled", tintEnabled);
+					IOMasterMaterial.AddColor("TintColor", tintColor.Red / 255.0F, tintColor.Green / 255.0F, tintColor.Blue / 255.0F, 1.0F);
+
+					// Control the Unreal material Emissive Color.
+					IOMasterMaterial.AddFloat("SelfIlluminationLuminance", 0.0F);
+
+					float transparency = 1.0F - GetFloatPropertyValue(RenderingAsset, SolidGlass.SolidglassReflectance, 0.0F);
+
+					// Control the Unreal material Opacity.
+					IOMasterMaterial.AddFloat("Transparency", transparency);
+					IOMasterMaterial.AddFloat("TransparencyMapFading", 0.0F);
+
+					if ((SolidglassBumpEnableType) GetIntegerPropertyValue(RenderingAsset, SolidGlass.SolidglassBumpEnable, (int) SolidglassBumpEnableType.None) == SolidglassBumpEnableType.Custom)
+					{
+						string bumpMapPath = GetTexturePropertyPath(RenderingAsset, SolidGlass.SolidglassBumpMap);
+
+						if (!string.IsNullOrEmpty(bumpMapPath))
+						{
+							float bumpAmount       = GetFloatPropertyValue(RenderingAsset, SolidGlass.SolidglassBumpAmount, 0.0F);
+							float bumpMapUVOffsetX = GetTexturePropertyDistance(RenderingAsset, SolidGlass.SolidglassBumpMap, UnifiedBitmap.TextureRealWorldOffsetX, 0.0F);
+							float bumpMapUVOffsetY = GetTexturePropertyDistance(RenderingAsset, SolidGlass.SolidglassBumpMap, UnifiedBitmap.TextureRealWorldOffsetY, 0.0F);
+							float bumpMapUVScaleX  = 1.0F / GetTexturePropertyDistance(RenderingAsset, SolidGlass.SolidglassBumpMap, UnifiedBitmap.TextureRealWorldScaleX, 1.0F);
+							float bumpMapUVScaleY  = 1.0F / GetTexturePropertyDistance(RenderingAsset, SolidGlass.SolidglassBumpMap, UnifiedBitmap.TextureRealWorldScaleY, 1.0F);
+							float bumpMapUVWAngle  = GetTexturePropertyAngle(RenderingAsset, SolidGlass.SolidglassBumpMap, UnifiedBitmap.TextureWAngle);
+
+							// Control the Unreal material Normal.
+							IOMasterMaterial.AddFloat("BumpAmount", bumpAmount);
+							IOMasterMaterial.AddTexture("BumpMap", bumpMapPath);
+							IOMasterMaterial.AddFloat("BumpMap_UVOffsetX", bumpMapUVOffsetX);
+							IOMasterMaterial.AddFloat("BumpMap_UVOffsetY", bumpMapUVOffsetY);
+							IOMasterMaterial.AddFloat("BumpMap_UVScaleX",  bumpMapUVScaleX);
+							IOMasterMaterial.AddFloat("BumpMap_UVScaleY",  bumpMapUVScaleY);
+							IOMasterMaterial.AddFloat("BumpMap_UVWAngle",  bumpMapUVWAngle);
+						}
+						else
+						{
+							IOMasterMaterial.AddFloat("BumpAmount", 0.0F);
+						}
+					}
+					else
+					{
+						IOMasterMaterial.AddFloat("BumpAmount", 0.0F);
+					}
+
+					// Control the Unreal material Refraction.
+					IOMasterMaterial.AddFloat("RefractionIndex", 1.0F);
+				}
+				break;
+				
+				case "Stone":
+				{
+					IOMasterMaterial.SetMasterMaterialType(FDatasmithFacadeMaterial.EMasterMaterialType.Opaque);
+
+					// Control the Unreal material Base Color.
+					IOMasterMaterial.AddColor("DiffuseColor", sourceMaterialColor.Red / 255.0F, sourceMaterialColor.Green / 255.0F, sourceMaterialColor.Blue / 255.0F, 1.0F);
+
+					string diffuseMapPath = GetTexturePropertyPath(RenderingAsset, Stone.StoneColor);
+
+					if (!string.IsNullOrEmpty(diffuseMapPath))
+					{
+						float diffuseMapUVOffsetX = GetTexturePropertyDistance(RenderingAsset, Stone.StoneColor, UnifiedBitmap.TextureRealWorldOffsetX, 0.0F);
+						float diffuseMapUVOffsetY = GetTexturePropertyDistance(RenderingAsset, Stone.StoneColor, UnifiedBitmap.TextureRealWorldOffsetY, 0.0F);
+						float diffuseMapUVScaleX  = 1.0F / GetTexturePropertyDistance(RenderingAsset, Stone.StoneColor, UnifiedBitmap.TextureRealWorldScaleX, 1.0F);
+						float diffuseMapUVScaleY  = 1.0F / GetTexturePropertyDistance(RenderingAsset, Stone.StoneColor, UnifiedBitmap.TextureRealWorldScaleY, 1.0F);
+						float diffuseMapUVWAngle  = GetTexturePropertyAngle(RenderingAsset, Stone.StoneColor, UnifiedBitmap.TextureWAngle);
+
+						// Control the Unreal material Base Color.
+						IOMasterMaterial.AddFloat("DiffuseMapFading", 1.0F);                          
+						IOMasterMaterial.AddTexture("DiffuseMap", diffuseMapPath);                    
+						IOMasterMaterial.AddFloat("DiffuseMap_UVOffsetX", diffuseMapUVOffsetX);
+						IOMasterMaterial.AddFloat("DiffuseMap_UVOffsetY", diffuseMapUVOffsetY);
+						IOMasterMaterial.AddFloat("DiffuseMap_UVScaleX",  diffuseMapUVScaleX);
+						IOMasterMaterial.AddFloat("DiffuseMap_UVScaleY",  diffuseMapUVScaleY);
+						IOMasterMaterial.AddFloat("DiffuseMap_UVWAngle",  diffuseMapUVWAngle);
+					}
+					else
+					{
+						IOMasterMaterial.AddFloat("DiffuseMapFading", 0.0F);                          
+					}
+
+					bool  tintEnabled = GetBooleanPropertyValue(RenderingAsset, Stone.CommonTintToggle, false);
+					Color tintColor   = GetColorPropertyValue(RenderingAsset, Stone.CommonTintColor, sourceMaterialColor);
+
+					// Control the Unreal material Base Color.
+					IOMasterMaterial.AddBoolean("TintEnabled", tintEnabled);
+					IOMasterMaterial.AddColor("TintColor", tintColor.Red / 255.0F, tintColor.Green / 255.0F, tintColor.Blue / 255.0F, 1.0F);
+
+					// Control the Unreal material Metallic.
+					IOMasterMaterial.AddBoolean("IsMetal", false);
+
+					float glossiness = 0.35F; // StoneApplicationType.Matte
+					switch ((StoneApplicationType) GetIntegerPropertyValue(RenderingAsset, Stone.StoneApplication, (int) StoneApplicationType.Matte))
+					{
+						case StoneApplicationType.Polished:
+							glossiness = 1.0F;
+							break;
+						case StoneApplicationType.Glossy:
+							glossiness = 0.8F;
+							break;
+						case StoneApplicationType.Matte:
+							glossiness = 0.35F;
+							break;
+						case StoneApplicationType.Unfinished:
+							glossiness = 0.25F;
+							break;
+					}                        
+
+					// Control the Unreal material Roughness.
+					IOMasterMaterial.AddFloat("Glossiness", glossiness);
+
+					// Control the Unreal material Emissive Color.
+					IOMasterMaterial.AddFloat("SelfIlluminationLuminance", 0.0F);
+
+					if ((StonePatternType) GetIntegerPropertyValue(RenderingAsset, Stone.StonePattern, (int) StonePatternType.None) == StonePatternType.Custom)
+					{
+						string bumpMapPath = GetTexturePropertyPath(RenderingAsset, Stone.StonePatternMap);
+
+						if (!string.IsNullOrEmpty(bumpMapPath))
+						{
+							float bumpAmount       = GetFloatPropertyValue(RenderingAsset, Stone.StonePatternAmount, 0.0F);
+							float bumpMapUVOffsetX = GetTexturePropertyDistance(RenderingAsset, Stone.StonePatternMap, UnifiedBitmap.TextureRealWorldOffsetX, 0.0F);
+							float bumpMapUVOffsetY = GetTexturePropertyDistance(RenderingAsset, Stone.StonePatternMap, UnifiedBitmap.TextureRealWorldOffsetY, 0.0F);
+							float bumpMapUVScaleX  = 1.0F / GetTexturePropertyDistance(RenderingAsset, Stone.StonePatternMap, UnifiedBitmap.TextureRealWorldScaleX, 1.0F);
+							float bumpMapUVScaleY  = 1.0F / GetTexturePropertyDistance(RenderingAsset, Stone.StonePatternMap, UnifiedBitmap.TextureRealWorldScaleY, 1.0F);
+							float bumpMapUVWAngle  = GetTexturePropertyAngle(RenderingAsset, Stone.StonePatternMap, UnifiedBitmap.TextureWAngle);
+
+							// Control the Unreal material Normal.
+							IOMasterMaterial.AddFloat("BumpAmount", bumpAmount);
+							IOMasterMaterial.AddTexture("BumpMap", bumpMapPath);
+							IOMasterMaterial.AddFloat("BumpMap_UVOffsetX", bumpMapUVOffsetX);
+							IOMasterMaterial.AddFloat("BumpMap_UVOffsetY", bumpMapUVOffsetY);
+							IOMasterMaterial.AddFloat("BumpMap_UVScaleX",  bumpMapUVScaleX);
+							IOMasterMaterial.AddFloat("BumpMap_UVScaleY",  bumpMapUVScaleY);
+							IOMasterMaterial.AddFloat("BumpMap_UVWAngle",  bumpMapUVWAngle);
+						}
+						else
+						{
+							IOMasterMaterial.AddFloat("BumpAmount", 0.0F);
+						}
+					}
+					else
+					{
+						IOMasterMaterial.AddFloat("BumpAmount", 0.0F);
+					}
+				}
+				break;
+
+				case "WallPaint":
+				{
+					IOMasterMaterial.SetMasterMaterialType(FDatasmithFacadeMaterial.EMasterMaterialType.Opaque);
+
+					Color color = GetColorPropertyValue(RenderingAsset, WallPaint.WallpaintColor, sourceMaterialColor);
+
+					// Control the Unreal material Base Color.
+					IOMasterMaterial.AddColor("DiffuseColor", color.Red / 255.0F, color.Green / 255.0F, color.Blue / 255.0F, 1.0F);
+					IOMasterMaterial.AddFloat("DiffuseMapFading", 0.0F);
+
+					bool  tintEnabled = GetBooleanPropertyValue(RenderingAsset, WallPaint.CommonTintToggle, false);
+					Color tintColor   = GetColorPropertyValue(RenderingAsset, WallPaint.CommonTintColor, sourceMaterialColor);
+
+					// Control the Unreal material Base Color.
+					IOMasterMaterial.AddBoolean("TintEnabled", tintEnabled);
+					IOMasterMaterial.AddColor("TintColor", tintColor.Red / 255.0F, tintColor.Green / 255.0F, tintColor.Blue / 255.0F, 1.0F);
+
+					// Control the Unreal material Metallic.
+					IOMasterMaterial.AddBoolean("IsMetal", false);
+
+					float glossiness = 0.6F; // WallpaintFinishType.Pearl
+					switch ((WallpaintFinishType) GetIntegerPropertyValue(RenderingAsset, WallPaint.WallpaintFinish, (int) WallpaintFinishType.Pearl))
+					{
+						case WallpaintFinishType.Gloss:
+							glossiness = 0.75F;
+							break;
+						case WallpaintFinishType.Semigloss:
+							glossiness = 0.7F;
+							break;
+						case WallpaintFinishType.Pearl:
+							glossiness = 0.6F;
+							break;
+						case WallpaintFinishType.Platinum:
+							glossiness = 0.55F;
+							break;
+						case WallpaintFinishType.Eggshell:
+							glossiness = 0.5F;
+							break;
+						case WallpaintFinishType.Flat:
+							glossiness = 0.4F;
+							break;
+					}                        
+
+					// Control the Unreal material Roughness.
+					IOMasterMaterial.AddFloat("Glossiness", glossiness);
+
+					// Control the Unreal material Emissive Color.
+					IOMasterMaterial.AddFloat("SelfIlluminationLuminance", 0.0F);
+
+					// Control the Unreal material Normal.
+					IOMasterMaterial.AddFloat("BumpAmount", 0.0F);
+				}
+				break;
+
+				case "Generic":
+				{
+					Color color = GetColorPropertyValue(RenderingAsset, Generic.GenericDiffuse, sourceMaterialColor);
+
+					// Control the Unreal material Base Color.
+					IOMasterMaterial.AddColor("DiffuseColor", color.Red / 255.0F, color.Green / 255.0F, color.Blue / 255.0F, 1.0F);
+
+					string diffuseMapPath = GetTexturePropertyPath(RenderingAsset, Generic.GenericDiffuse);
+
+					if (!string.IsNullOrEmpty(diffuseMapPath))
+					{
+						float diffuseImageFade    = GetFloatPropertyValue(RenderingAsset, Generic.GenericDiffuseImageFade, 0.0F);
+						float diffuseMapUVOffsetX = GetTexturePropertyDistance(RenderingAsset, Generic.GenericDiffuse, UnifiedBitmap.TextureRealWorldOffsetX, 0.0F);
+						float diffuseMapUVOffsetY = GetTexturePropertyDistance(RenderingAsset, Generic.GenericDiffuse, UnifiedBitmap.TextureRealWorldOffsetY, 0.0F);
+						float diffuseMapUVScaleX  = 1.0F / GetTexturePropertyDistance(RenderingAsset, Generic.GenericDiffuse, UnifiedBitmap.TextureRealWorldScaleX, 1.0F);
+						float diffuseMapUVScaleY  = 1.0F / GetTexturePropertyDistance(RenderingAsset, Generic.GenericDiffuse, UnifiedBitmap.TextureRealWorldScaleY, 1.0F);
+						float diffuseMapUVWAngle  = GetTexturePropertyAngle(RenderingAsset, Generic.GenericDiffuse, UnifiedBitmap.TextureWAngle);
+
+						// Control the Unreal material Base Color.
+						IOMasterMaterial.AddFloat("DiffuseMapFading", diffuseImageFade);                          
+						IOMasterMaterial.AddTexture("DiffuseMap", diffuseMapPath);                    
+						IOMasterMaterial.AddFloat("DiffuseMap_UVOffsetX", diffuseMapUVOffsetX);
+						IOMasterMaterial.AddFloat("DiffuseMap_UVOffsetY", diffuseMapUVOffsetY);
+						IOMasterMaterial.AddFloat("DiffuseMap_UVScaleX",  diffuseMapUVScaleX);
+						IOMasterMaterial.AddFloat("DiffuseMap_UVScaleY",  diffuseMapUVScaleY);
+						IOMasterMaterial.AddFloat("DiffuseMap_UVWAngle",  diffuseMapUVWAngle);
+					}
+					else
+					{
+						IOMasterMaterial.AddFloat("DiffuseMapFading", 0.0F);                          
+					}
+
+					bool  tintEnabled = GetBooleanPropertyValue(RenderingAsset, Generic.CommonTintToggle, false);
+					Color tintColor   = GetColorPropertyValue(RenderingAsset, Generic.CommonTintColor, sourceMaterialColor);
+
+					// Control the Unreal material Base Color.
+					IOMasterMaterial.AddBoolean("TintEnabled", tintEnabled);
+					IOMasterMaterial.AddColor("TintColor", tintColor.Red / 255.0F, tintColor.Green / 255.0F, tintColor.Blue / 255.0F, 1.0F);
+
+					float selfIlluminationLuminance   = GetFloatPropertyValue(RenderingAsset, Generic.GenericSelfIllumLuminance, 0.0F);
+					Color selfIlluminationFilterColor = GetColorPropertyValue(RenderingAsset, Generic.GenericSelfIllumFilterMap, new Color(255, 255, 255));
+
+					// Control the Unreal material Emissive Color.
+					IOMasterMaterial.AddFloat("SelfIlluminationLuminance", selfIlluminationLuminance);
+					IOMasterMaterial.AddColor("SelfIlluminationFilter", selfIlluminationFilterColor.Red / 255.0F, selfIlluminationFilterColor.Green / 255.0F, selfIlluminationFilterColor.Blue / 255.0F, 1.0F);
+
+					string selfIlluminationMapPath = GetTexturePropertyPath(RenderingAsset, Generic.GenericSelfIllumFilterMap);
+
+					if (!string.IsNullOrEmpty(selfIlluminationMapPath))
+					{
+						float selfIlluminationMapUVOffsetX = GetTexturePropertyDistance(RenderingAsset, Generic.GenericSelfIllumFilterMap, UnifiedBitmap.TextureRealWorldOffsetX, 0.0F);
+						float selfIlluminationMapUVOffsetY = GetTexturePropertyDistance(RenderingAsset, Generic.GenericSelfIllumFilterMap, UnifiedBitmap.TextureRealWorldOffsetY, 0.0F);
+						float selfIlluminationMapUVScaleX  = 1.0F / GetTexturePropertyDistance(RenderingAsset, Generic.GenericSelfIllumFilterMap, UnifiedBitmap.TextureRealWorldScaleX, 1.0F);
+						float selfIlluminationMapUVScaleY  = 1.0F / GetTexturePropertyDistance(RenderingAsset, Generic.GenericSelfIllumFilterMap, UnifiedBitmap.TextureRealWorldScaleY, 1.0F);
+						float selfIlluminationMapUVWAngle  = GetTexturePropertyAngle(RenderingAsset, Generic.GenericSelfIllumFilterMap, UnifiedBitmap.TextureWAngle);
+
+						// Control the Unreal material Emissive Color.
+						IOMasterMaterial.AddBoolean("SelfIlluminationMapEnable", true);
+						IOMasterMaterial.AddTexture("SelfIlluminationMap", selfIlluminationMapPath);
+						IOMasterMaterial.AddFloat("SelfIlluminationMap_UVOffsetX", selfIlluminationMapUVOffsetX);
+						IOMasterMaterial.AddFloat("SelfIlluminationMap_UVOffsetY", selfIlluminationMapUVOffsetY);
+						IOMasterMaterial.AddFloat("SelfIlluminationMap_UVScaleX",  selfIlluminationMapUVScaleX);
+						IOMasterMaterial.AddFloat("SelfIlluminationMap_UVScaleY",  selfIlluminationMapUVScaleY);
+						IOMasterMaterial.AddFloat("SelfIlluminationMap_UVWAngle",  selfIlluminationMapUVWAngle);
+					}
+					else
+					{
+						IOMasterMaterial.AddBoolean("SelfIlluminationMapEnable", false);
+					}
+
+					string bumpMapPath = GetTexturePropertyPath(RenderingAsset, Generic.GenericBumpMap);
+
+					if (!string.IsNullOrEmpty(bumpMapPath))
+					{
+						float bumpAmount       = GetFloatPropertyValue(RenderingAsset, Generic.GenericBumpAmount, 0.0F);
+						float bumpMapUVOffsetX = GetTexturePropertyDistance(RenderingAsset, Generic.GenericBumpMap, UnifiedBitmap.TextureRealWorldOffsetX, 0.0F);
+						float bumpMapUVOffsetY = GetTexturePropertyDistance(RenderingAsset, Generic.GenericBumpMap, UnifiedBitmap.TextureRealWorldOffsetY, 0.0F);
+						float bumpMapUVScaleX  = 1.0F / GetTexturePropertyDistance(RenderingAsset, Generic.GenericBumpMap, UnifiedBitmap.TextureRealWorldScaleX, 1.0F);
+						float bumpMapUVScaleY  = 1.0F / GetTexturePropertyDistance(RenderingAsset, Generic.GenericBumpMap, UnifiedBitmap.TextureRealWorldScaleY, 1.0F);
+						float bumpMapUVWAngle  = GetTexturePropertyAngle(RenderingAsset, Generic.GenericBumpMap, UnifiedBitmap.TextureWAngle);
+
+						// Control the Unreal material Normal.
+						IOMasterMaterial.AddFloat("BumpAmount", bumpAmount);
+						IOMasterMaterial.AddTexture("BumpMap", bumpMapPath);
+						IOMasterMaterial.AddFloat("BumpMap_UVOffsetX", bumpMapUVOffsetX);
+						IOMasterMaterial.AddFloat("BumpMap_UVOffsetY", bumpMapUVOffsetY);
+						IOMasterMaterial.AddFloat("BumpMap_UVScaleX",  bumpMapUVScaleX);
+						IOMasterMaterial.AddFloat("BumpMap_UVScaleY",  bumpMapUVScaleY);
+						IOMasterMaterial.AddFloat("BumpMap_UVWAngle",  bumpMapUVWAngle);
+					}
+					else
+					{
+						IOMasterMaterial.AddFloat("BumpAmount", 0.0F);
+					}
+
+					float transparency = GetFloatPropertyValue(RenderingAsset, Generic.GenericTransparency, 0.0F);
+
+					if (transparency > 0.0F)
+					{
+						IOMasterMaterial.SetMasterMaterialType(FDatasmithFacadeMaterial.EMasterMaterialType.Transparent);
+
+						// Control the Unreal material Opacity.
+						IOMasterMaterial.AddFloat("Transparency", transparency);
+
+						string transparencyMapPath = GetTexturePropertyPath(RenderingAsset, Generic.GenericTransparency);
+
+						if (!string.IsNullOrEmpty(transparencyMapPath))
+						{
+							float transparencyImageFade    = GetFloatPropertyValue(RenderingAsset, Generic.GenericTransparencyImageFade, 0.0F);
+							float transparencyMapUVOffsetX = GetTexturePropertyDistance(RenderingAsset, Generic.GenericTransparency, UnifiedBitmap.TextureRealWorldOffsetX, 0.0F);
+							float transparencyMapUVOffsetY = GetTexturePropertyDistance(RenderingAsset, Generic.GenericTransparency, UnifiedBitmap.TextureRealWorldOffsetY, 0.0F);
+							float transparencyMapUVScaleX  = 1.0F / GetTexturePropertyDistance(RenderingAsset, Generic.GenericTransparency, UnifiedBitmap.TextureRealWorldScaleX, 1.0F);
+							float transparencyMapUVScaleY  = 1.0F / GetTexturePropertyDistance(RenderingAsset, Generic.GenericTransparency, UnifiedBitmap.TextureRealWorldScaleY, 1.0F);
+							float transparencyMapUVWAngle  = GetTexturePropertyAngle(RenderingAsset, Generic.GenericTransparency, UnifiedBitmap.TextureWAngle);
+
+							// Control the Unreal material Opacity.
+							IOMasterMaterial.AddFloat("TransparencyMapFading", transparencyImageFade);
+							IOMasterMaterial.AddTexture("TransparencyMap", transparencyMapPath);
+							IOMasterMaterial.AddFloat("TransparencyMap_UVOffsetX", transparencyMapUVOffsetX);
+							IOMasterMaterial.AddFloat("TransparencyMap_UVOffsetY", transparencyMapUVOffsetY);
+							IOMasterMaterial.AddFloat("TransparencyMap_UVScaleX",  transparencyMapUVScaleX);
+							IOMasterMaterial.AddFloat("TransparencyMap_UVScaleY",  transparencyMapUVScaleY);
+							IOMasterMaterial.AddFloat("TransparencyMap_UVWAngle",  transparencyMapUVWAngle);
+						}
+						else
+						{
+							IOMasterMaterial.AddFloat("TransparencyMapFading", 0.0F);
+						}
+
+						float refractionIndex = GetFloatPropertyValue(RenderingAsset, Generic.GenericRefractionIndex, 1.0F);
+
+						// Control the Unreal material Refraction.
+						IOMasterMaterial.AddFloat("RefractionIndex", refractionIndex);
+					}
+					else
+					{
+						IOMasterMaterial.SetMasterMaterialType(FDatasmithFacadeMaterial.EMasterMaterialType.Opaque);
+
+						bool isMetal = GetBooleanPropertyValue(RenderingAsset, Generic.GenericIsMetal, false);
+
+						// Control the Unreal material Metallic.
+						IOMasterMaterial.AddBoolean("IsMetal", isMetal);
+
+						float glossiness = GetFloatPropertyValue(RenderingAsset, Generic.GenericGlossiness, InMaterial.Smoothness / 100.0F);
+
+						// Control the Unreal material Roughness.
+						IOMasterMaterial.AddFloat("Glossiness", glossiness);
+
+						string cutoutMapPath = GetTexturePropertyPath(RenderingAsset, Generic.GenericCutoutOpacity);
+
+						if (!string.IsNullOrEmpty(cutoutMapPath))
+						{
+							float cutoutMapUVOffsetX = GetTexturePropertyDistance(RenderingAsset, Generic.GenericCutoutOpacity, UnifiedBitmap.TextureRealWorldOffsetX, 0.0F);
+							float cutoutMapUVOffsetY = GetTexturePropertyDistance(RenderingAsset, Generic.GenericCutoutOpacity, UnifiedBitmap.TextureRealWorldOffsetY, 0.0F);
+							float cutoutMapUVScaleX  = 1.0F / GetTexturePropertyDistance(RenderingAsset, Generic.GenericCutoutOpacity, UnifiedBitmap.TextureRealWorldScaleX, 1.0F);
+							float cutoutMapUVScaleY  = 1.0F / GetTexturePropertyDistance(RenderingAsset, Generic.GenericCutoutOpacity, UnifiedBitmap.TextureRealWorldScaleY, 1.0F);
+							float cutoutMapUVWAngle  = GetTexturePropertyAngle(RenderingAsset, Generic.GenericCutoutOpacity, UnifiedBitmap.TextureWAngle);
+
+							IOMasterMaterial.SetMasterMaterialType(FDatasmithFacadeMaterial.EMasterMaterialType.CutOut);
+
+							// Control the Unreal material Opacity Mask.
+							IOMasterMaterial.AddTexture("CutoutOpacityMap",  cutoutMapPath);
+							IOMasterMaterial.AddFloat("CutoutMap_UVOffsetX", cutoutMapUVOffsetX);
+							IOMasterMaterial.AddFloat("CutoutMap_UVOffsetY", cutoutMapUVOffsetY);
+							IOMasterMaterial.AddFloat("CutoutMap_UVScaleX",  cutoutMapUVScaleX);
+							IOMasterMaterial.AddFloat("CutoutMap_UVScaleY",  cutoutMapUVScaleY);
+							IOMasterMaterial.AddFloat("CutoutMap_UVWAngle",  cutoutMapUVWAngle);
+						}
+					}
+				}
+				break;
+
+				default:
+				{
+					// The properties of the Datasmith master material cannot be set.
+					return false;
+				}
+			}
+
+			// The properties of the Datasmith master material are set.
+			return true;
+        }
+
+		private void SetFallbackMaterial(
+			Color                    InMaterialColor,
+			float                    InMaterialTransparency, // in range 0.0-1.0
+			float                    InMaterialGlossiness,   // in range 0.0-1.0
+			FDatasmithFacadeMaterial IOMasterMaterial
+		)
+		{
+			// Control the Unreal material Base Color.
+			Color MaterialColor = InMaterialColor.IsValid ? InMaterialColor : new Color(255, 255, 255);
+			IOMasterMaterial.AddColor("DiffuseColor", MaterialColor.Red / 255.0F, MaterialColor.Green / 255.0F, MaterialColor.Blue / 255.0F, 1.0F);
+			IOMasterMaterial.AddFloat("DiffuseMapFading", 0.0F);
+			IOMasterMaterial.AddBoolean("TintEnabled", false);
+
+			// Control the Unreal material Emissive Color.
+			IOMasterMaterial.AddFloat("SelfIlluminationLuminance", 0.0F);
+
+			// Control the Unreal material Normal.
+			IOMasterMaterial.AddFloat("BumpAmount", 0.0F);
+
+			if (InMaterialTransparency > 0.0F)
+			{
+				IOMasterMaterial.SetMasterMaterialType(FDatasmithFacadeMaterial.EMasterMaterialType.Transparent);
+
+				// Control the Unreal material Opacity.
+				IOMasterMaterial.AddFloat("Transparency", InMaterialTransparency);
+				IOMasterMaterial.AddFloat("TransparencyMapFading", 0.0F);
+
+				// Control the Unreal material Refraction.
+				IOMasterMaterial.AddFloat("RefractionIndex", 1.0F);
+			}
+			else
+			{
+				IOMasterMaterial.SetMasterMaterialType(FDatasmithFacadeMaterial.EMasterMaterialType.Opaque);
+
+				// Control the Unreal material Metallic.
+				IOMasterMaterial.AddBoolean("IsMetal", false);
+
+				// Control the Unreal material Roughness.
+				IOMasterMaterial.AddFloat("Glossiness", InMaterialGlossiness);
+			}
+		}
+
+		private bool GetBooleanPropertyValue(
+			Asset  in_asset,
+			string in_propertyName,
+			bool   in_defaultValue
+		)
+        {
+			// DebugLog.AddLine($"Boolean Property {in_propertyName}");
+
+			AssetProperty booleanProperty = in_asset.FindByName(in_propertyName);
+
+			return (booleanProperty != null) ? (booleanProperty as AssetPropertyBoolean).Value : in_defaultValue;
+        }
+
+		private int GetIntegerPropertyValue(
+			Asset  in_asset,
+			string in_propertyName,
+			int    in_defaultValue
+		)
+        {
+			// DebugLog.AddLine($"Integer Property {in_propertyName}");
+
+			AssetProperty integerProperty = in_asset.FindByName(in_propertyName);
+
+			if (integerProperty != null)
+			{
+				if (integerProperty.Type == AssetPropertyType.Enumeration)
+				{
+					// Handle the spurious case of an enumerated type value not stored in an integer asset property.
+					return (integerProperty as AssetPropertyEnum).Value;
+				}
+				else
+				{
+					return (integerProperty as AssetPropertyInteger).Value;
+				}
+			}
+
+			return in_defaultValue;
+        }
+
+		private float GetFloatPropertyValue(
+			Asset  in_asset,
+			string in_propertyName,
+			float  in_defaultValue
+		)
+        {
+			// DebugLog.AddLine($"Float Property {in_propertyName}");
+
+			AssetProperty doubleProperty = in_asset.FindByName(in_propertyName);
+
+			if (doubleProperty != null)
+			{
+				if (doubleProperty.Type == AssetPropertyType.Float)
+				{
+					// Handle the spurious case of a value not stored in a double asset property.
+					return (doubleProperty as AssetPropertyFloat).Value;
+				}
+				else
+				{
+					return (float) (doubleProperty as AssetPropertyDouble).Value;
+				}
+			}
+
+			return in_defaultValue;
+        }
+
+		private Color GetColorPropertyValue(
+			Asset  in_asset,
+			string in_propertyName,
+			Color  in_defaultValue
+		)
+        {
+			// DebugLog.AddLine($"Color Property {in_propertyName}");
+
+			AssetProperty colorProperty = in_asset.FindByName(in_propertyName);
+
+			if (colorProperty != null)
+			{
+				Color color = (colorProperty as AssetPropertyDoubleArray4d).GetValueAsColor();
+
+				if (color.IsValid)
+				{
+					return color;
+				}
+			}
+
+			return in_defaultValue;
+		}
+
+		private string GetTexturePropertyPath(
+			Asset  in_asset,
+			string in_propertyName
+		)
+        {
+			// DebugLog.AddLine($"Texture Property {in_propertyName}");
+
+			AssetProperty textureProperty = in_asset.FindByName(in_propertyName);
+
+			if (textureProperty != null)
+			{
+				Asset unifiedBitmapAsset = textureProperty.GetSingleConnectedAsset();
+
+				if (unifiedBitmapAsset != null)
+				{
+					AssetProperty sourceProperty = unifiedBitmapAsset.FindByName(UnifiedBitmap.UnifiedbitmapBitmap);
+
+					if (sourceProperty != null)
+					{
+						string sourcePath = (sourceProperty as AssetPropertyString).Value;
+
+						// DebugLog.AddLine($"Texture path: {sourcePath}");
+
+						if (!string.IsNullOrEmpty(sourcePath))
+						{
+							if (sourcePath.Contains("|"))
+							{
+								sourcePath = sourcePath.Split('|')[0];
+							}
+
+							if (!string.IsNullOrEmpty(sourcePath))
+							{
+								// TODO: Better handle relative paths.
+								const string RevitTextureFolder1 = "C:\\Program Files (x86)\\Common Files\\Autodesk Shared\\Materials\\Textures\\";
+								const string RevitTextureFolder2 = "C:\\Program Files (x86)\\Common Files\\Autodesk Shared\\Materials\\Textures\\1\\Mats\\";
+
+								if (sourcePath[0] == '1')
+								{
+									sourcePath =  $"{RevitTextureFolder1}{sourcePath}";
+								}
+
+								if (sourcePath.Contains("Materials\\Generic\\Presets\\"))
+								{
+									sourcePath = sourcePath.Replace("Materials\\Generic\\Presets\\", RevitTextureFolder2).Replace(".jpg", ".png");
+								}
+
+								if (!Path.IsPathRooted(sourcePath))
+								{
+									string rootedSourcePath = Path.Combine(RevitTextureFolder2, sourcePath);
+
+									if (File.Exists(rootedSourcePath))
+									{
+										sourcePath = rootedSourcePath;
+									}
+								}
+
+								// Also search a relative path in the extra texture paths.
+								if (!Path.IsPathRooted(sourcePath))
+								{
+									foreach (string extraTexturePath in ExtraTexturePaths)
+									{
+										string extraSourcePath = Path.Combine(extraTexturePath, sourcePath);
+
+										if (Path.IsPathRooted(extraSourcePath) && File.Exists(extraSourcePath))
+										{
+											sourcePath = extraSourcePath;
+											break;
+										}
+									}
+								}
+								
+								if (!File.Exists(sourcePath))
+								{
+									MessageList.Add($"Warning - Material \"{MaterialLabel}\": Cannot find texture file {sourcePath}");
+								}
+
+								return sourcePath;
+							}
+						}
+					}
+				}
+			}
+
+            return "";
+        }
+
+		private float GetTexturePropertyDistance(
+			Asset  in_asset,
+			string in_propertyName,
+			string in_distanceName,
+			float  in_defaultValue
+		)
+        {
+			// DebugLog.AddLine($"Texture Distance Property {in_propertyName}");
+
+			AssetProperty textureProperty = in_asset.FindByName(in_propertyName);
+
+			if (textureProperty != null)
+			{
+				Asset unifiedBitmapAsset = textureProperty.GetSingleConnectedAsset();
+
+				if (unifiedBitmapAsset != null)
+				{
+					float scalingFactor = 1.0F;
+
+					AssetProperty versionProperty = unifiedBitmapAsset.FindByName("version");
+
+					if (versionProperty != null)
+					{
+						// DebugLog.AddLine($"Texture version: {(versionProperty as AssetPropertyInteger).Value}");
+
+						scalingFactor = ((versionProperty as AssetPropertyInteger).Value == 4) ? 12.0F /* inches per foot */ : 1.0F;
+					}
+
+					AssetProperty distanceProperty = unifiedBitmapAsset.FindByName(in_distanceName);
+
+					if (distanceProperty != null)
+					{
+						double distance = (distanceProperty as AssetPropertyDistance).Value;
+
+						return (distance != 0.0) ? (float) distance / scalingFactor : in_defaultValue;
+					}
+				}
+			}
+
+            return in_defaultValue;
+        }
+
+		private float GetTexturePropertyAngle(
+			Asset  in_asset,
+			string in_propertyName,
+			string in_angleName
+		)
+        {
+			// DebugLog.AddLine($"Texture Angle Property {in_propertyName}");
+
+			AssetProperty textureProperty = in_asset.FindByName(in_propertyName);
+
+			if (textureProperty != null)
+			{
+				Asset unifiedBitmapAsset = textureProperty.GetSingleConnectedAsset();
+
+				if (unifiedBitmapAsset != null)
+				{
+					AssetProperty angleProperty = unifiedBitmapAsset.FindByName(in_angleName);
+
+					if (angleProperty != null)
+					{
+						return (float) (angleProperty as AssetPropertyDouble).Value / 360.0F;
+					}
+				}
+			}
+
+            return 0.0F;
+        }
+	}
+}
diff --git a/Engine/Source/Programs/Enterprise/Datasmith/DatasmithSketchUpExporter/DatasmithSketchUp2018.Build.cs b/Engine/Source/Programs/Enterprise/Datasmith/DatasmithSketchUpExporter/DatasmithSketchUp2018.Build.cs
new file mode 100644
--- /dev/null
+++ b/Engine/Source/Programs/Enterprise/Datasmith/DatasmithSketchUpExporter/DatasmithSketchUp2018.Build.cs
@@ -0,0 +1,26 @@
+// Copyright 1998-2019 Epic Games, Inc. All Rights Reserved.
+
+using System.IO;
+
+namespace UnrealBuildTool.Rules
+{
+	[SupportedPlatforms("Win64")]
+	public class DatasmithSketchUp2018 : DatasmithSketchUpBase
+	{
+		public DatasmithSketchUp2018(ReadOnlyTargetRules Target)
+			: base(Target)
+		{
+			PrivateDefinitions.Add("SKP_SDK_2018");
+		}
+
+		public override string GetSketchUpSDKFolder()
+		{
+			return "SDK_Win_x64_18-0-18664";
+		}
+
+		public override string GetSketchUpEnvVar()
+		{
+			return "SKP_SDK_2018";
+		}
+	}
+}
