Merge from Dev-Enterprise: Revit Datasmith exporter fails to export System elements

#rnx

diff --git a/Engine/Source/Programs/Enterprise/Datasmith/DatasmithRevitExporter/Private/DatasmithRevitExportContext.cs b/Engine/Source/Programs/Enterprise/Datasmith/DatasmithRevitExporter/Private/DatasmithRevitExportContext.cs
--- a/Engine/Source/Programs/Enterprise/Datasmith/DatasmithRevitExporter/Private/DatasmithRevitExportContext.cs
+++ b/Engine/Source/Programs/Enterprise/Datasmith/DatasmithRevitExporter/Private/DatasmithRevitExportContext.cs
@@ -1539,7 +1539,7 @@ namespace DatasmithRevitExporter
 			SiteLocation InSiteLocation
 		)
 		{
-			if (!InSiteLocation.IsValidObject)
+			if (InSiteLocation == null || !InSiteLocation.IsValidObject)
 			{
 				return;
 			}
