Build fix for int to bool implicit conversion

diff --git a/Engine/Source/Programs/Enterprise/Datasmith/DatasmithMaxExporter/Private/DatasmithMaxScanline.cpp b/Engine/Source/Programs/Enterprise/Datasmith/DatasmithMaxExporter/Private/DatasmithMaxScanline.cpp
--- a/Engine/Source/Programs/Enterprise/Datasmith/DatasmithMaxExporter/Private/DatasmithMaxScanline.cpp
+++ b/Engine/Source/Programs/Enterprise/Datasmith/DatasmithMaxExporter/Private/DatasmithMaxScanline.cpp
@@ -427,7 +427,7 @@ FString FDatasmithMaxMatWriter::DumpAutodeskBitmap(TSharedPtr<IDatasmithComposit
 	}
 
 	bool bInvertTexture = false;
-	if (AutodeskBitmapParameters.InvertImage)
+	if (AutodeskBitmapParameters.bInvertImage)
 	{
 		bInvertTexture = true;
 	}
diff --git a/Engine/Source/Programs/Enterprise/Datasmith/DatasmithMaxExporter/Private/DatasmithMaxTexmapParser.cpp b/Engine/Source/Programs/Enterprise/Datasmith/DatasmithMaxExporter/Private/DatasmithMaxTexmapParser.cpp
--- a/Engine/Source/Programs/Enterprise/Datasmith/DatasmithMaxExporter/Private/DatasmithMaxTexmapParser.cpp
+++ b/Engine/Source/Programs/Enterprise/Datasmith/DatasmithMaxExporter/Private/DatasmithMaxTexmapParser.cpp
@@ -246,7 +246,7 @@ namespace DatasmithMaxTexmapParser
 				}
 				else if (FCString::Stricmp(ParamDefinition.int_name, TEXT("Parameters_Invert_Image")) == 0)
 				{
-					AutodeskBitmapParameters.InvertImage = (ParamBlock2->GetInt(ParamDefinition.ID, CurrentTime));
+					AutodeskBitmapParameters.bInvertImage = (ParamBlock2->GetInt(ParamDefinition.ID, CurrentTime) != 0);
 				}
 				else if (FCString::Stricmp(ParamDefinition.int_name, TEXT("Position_X")) == 0)
 				{
@@ -270,11 +270,11 @@ namespace DatasmithMaxTexmapParser
 				}
 				else if (FCString::Stricmp(ParamDefinition.int_name, TEXT("Repeat_Horizontal")) == 0)
 				{
-					AutodeskBitmapParameters.RepeatHorizontal = (ParamBlock2->GetInt(ParamDefinition.ID, CurrentTime));
+					AutodeskBitmapParameters.bRepeatHorizontal = (ParamBlock2->GetInt(ParamDefinition.ID, CurrentTime) != 0);
 				}
 				else if (FCString::Stricmp(ParamDefinition.int_name, TEXT("Repeat_Vertical")) == 0)
 				{
-					AutodeskBitmapParameters.RepeatVertical = (ParamBlock2->GetInt(ParamDefinition.ID, CurrentTime));
+					AutodeskBitmapParameters.bRepeatVertical = (ParamBlock2->GetInt(ParamDefinition.ID, CurrentTime) != 0);
 				}
 				else if (FCString::Stricmp(ParamDefinition.int_name, TEXT("Advanced_Parameters_Blur")) == 0)
 				{
diff --git a/Engine/Source/Programs/Enterprise/Datasmith/DatasmithMaxExporter/Private/DatasmithMaxTexmapParser.h b/Engine/Source/Programs/Enterprise/Datasmith/DatasmithMaxExporter/Private/DatasmithMaxTexmapParser.h
--- a/Engine/Source/Programs/Enterprise/Datasmith/DatasmithMaxExporter/Private/DatasmithMaxTexmapParser.h
+++ b/Engine/Source/Programs/Enterprise/Datasmith/DatasmithMaxExporter/Private/DatasmithMaxTexmapParser.h
@@ -74,12 +74,12 @@ namespace DatasmithMaxTexmapParser
 	{
 		PBBitmap* SourceFile = nullptr;
 		float Brightness = 1;
-		bool InvertImage = false;
+		bool bInvertImage = false;
 		FVector2D Position = FVector2D(0,0);
 		float Rotation = 0;
 		FVector2D Scale = FVector2D(1, 1);
-		bool RepeatHorizontal = true;
-		bool RepeatVertical = true;
+		bool bRepeatHorizontal = true;
+		bool bRepeatVertical = true;
 		float BlurValue = 0;
 		float BlurOffset = 0;
 		float FilteringValue = 0;
diff --git a/Engine/Source/Programs/Enterprise/Datasmith/DatasmithMaxExporter/Private/MaxMaterialsToUEPbr/DatasmithMaxTexmapToUEPbr.cpp b/Engine/Source/Programs/Enterprise/Datasmith/DatasmithMaxExporter/Private/MaxMaterialsToUEPbr/DatasmithMaxTexmapToUEPbr.cpp
--- a/Engine/Source/Programs/Enterprise/Datasmith/DatasmithMaxExporter/Private/MaxMaterialsToUEPbr/DatasmithMaxTexmapToUEPbr.cpp
+++ b/Engine/Source/Programs/Enterprise/Datasmith/DatasmithMaxExporter/Private/MaxMaterialsToUEPbr/DatasmithMaxTexmapToUEPbr.cpp
@@ -825,7 +825,7 @@ IDatasmithMaterialExpression* FDatasmithMaxAutodeskBitmapToUEPbr::Convert(FDatas
 	DatasmithMaxTexmapParser::FAutodeskBitmapParameters AutodeskBitmapParameters = DatasmithMaxTexmapParser::ParseAutodeskBitmap(InTexmap);
 	FScopedBitMapPtr ActualBitmap(AutodeskBitmapParameters.SourceFile->bi, AutodeskBitmapParameters.SourceFile->bm);
 	FString ActualBitmapName = FDatasmithMaxMatWriter::GetActualBitmapName(&ActualBitmap.MapInfo);
-	bool bUseAlphaAsMono = ActualBitmap.Map->HasAlpha();
+	bool bUseAlphaAsMono = (ActualBitmap.Map->HasAlpha() != 0);
 	bool bIsSRGB = FDatasmithMaxMatHelper::IsSRGB(*ActualBitmap.Map);
 
 	return FDatasmithMaxTexmapToUEPbrUtils::ConvertBitMap(MaxMaterialToUEPbr, InTexmap, ActualBitmapName, bUseAlphaAsMono, bIsSRGB);
