Datasmith 3dsMax Exporter: More int to bool conversion fixes for operator== that returns an int.

diff --git a/Engine/Source/Programs/Enterprise/Datasmith/DatasmithMaxExporter/Private/MaxMaterialsToUEPbr/DatasmithMaxTexmapToUEPbr.cpp b/Engine/Source/Programs/Enterprise/Datasmith/DatasmithMaxExporter/Private/MaxMaterialsToUEPbr/DatasmithMaxTexmapToUEPbr.cpp
--- a/Engine/Source/Programs/Enterprise/Datasmith/DatasmithMaxExporter/Private/MaxMaterialsToUEPbr/DatasmithMaxTexmapToUEPbr.cpp
+++ b/Engine/Source/Programs/Enterprise/Datasmith/DatasmithMaxExporter/Private/MaxMaterialsToUEPbr/DatasmithMaxTexmapToUEPbr.cpp
@@ -975,7 +975,7 @@ IDatasmithMaterialExpression* FDatasmithMaxNormalToUEPbr::Convert( FDatasmithMax
 
 bool FDatasmithMaxRGBMultiplyToUEPbr::IsSupported( const FDatasmithMaxMaterialsToUEPbr* MaxMaterialToUEPbr, Texmap* InTexmap ) const
 {
-	return InTexmap ? InTexmap->ClassID() == RGBMULTIPLYCLASS : false;
+	return InTexmap ? (bool)( InTexmap->ClassID() == RGBMULTIPLYCLASS ) : false;
 }
 
 IDatasmithMaterialExpression* FDatasmithMaxRGBMultiplyToUEPbr::Convert( FDatasmithMaxMaterialsToUEPbr* MaxMaterialToUEPbr, Texmap* InTexmap )
@@ -1003,7 +1003,7 @@ IDatasmithMaterialExpression* FDatasmithMaxRGBMultiplyToUEPbr::Convert( FDatasmi
 
 bool FDatasmithMaxRGBTintToUEPbr::IsSupported( const FDatasmithMaxMaterialsToUEPbr* MaxMaterialToUEPbr, Texmap* InTexmap ) const
 {
-	return InTexmap ? InTexmap->ClassID() == RGBTINTCLASS : false;
+	return InTexmap ? (bool)( InTexmap->ClassID() == RGBTINTCLASS ) : false;
 }
 
 IDatasmithMaterialExpression* FDatasmithMaxRGBTintToUEPbr::Convert( FDatasmithMaxMaterialsToUEPbr* MaxMaterialToUEPbr, Texmap* InTexmap )
diff --git a/Engine/Source/Programs/Enterprise/Datasmith/DatasmithMaxExporter/Private/MaxMaterialsToUEPbr/DatasmithMaxVRayTexmapToUEPbr.cpp b/Engine/Source/Programs/Enterprise/Datasmith/DatasmithMaxExporter/Private/MaxMaterialsToUEPbr/DatasmithMaxVRayTexmapToUEPbr.cpp
--- a/Engine/Source/Programs/Enterprise/Datasmith/DatasmithMaxExporter/Private/MaxMaterialsToUEPbr/DatasmithMaxVRayTexmapToUEPbr.cpp
+++ b/Engine/Source/Programs/Enterprise/Datasmith/DatasmithMaxExporter/Private/MaxMaterialsToUEPbr/DatasmithMaxVRayTexmapToUEPbr.cpp
@@ -139,7 +139,7 @@ namespace DatasmithMaxVRayTexmapToUEPbrImpl
 
 bool FDatasmithMaxVRayColorTexmapToUEPbr::IsSupported( const FDatasmithMaxMaterialsToUEPbr* MaxMaterialToUEPbr, Texmap* InTexmap ) const
 {
-	return InTexmap ? InTexmap->ClassID() == VRAYCOLORCLASS : false;
+	return InTexmap ? (bool)( InTexmap->ClassID() == VRAYCOLORCLASS ) : false;
 }
 
 IDatasmithMaterialExpression* FDatasmithMaxVRayColorTexmapToUEPbr::Convert( FDatasmithMaxMaterialsToUEPbr* MaxMaterialToUEPbr, Texmap* InTexmap )
@@ -155,7 +155,7 @@ IDatasmithMaterialExpression* FDatasmithMaxVRayColorTexmapToUEPbr::Convert( FDat
 
 bool FDatasmithMaxVRayHDRITexmapToUEPbr::IsSupported( const FDatasmithMaxMaterialsToUEPbr* MaxMaterialToUEPbr, Texmap* InTexmap ) const
 {
-	return InTexmap ? InTexmap->ClassID() == VRAYHDRICLASS : false;
+	return InTexmap ? (bool)( InTexmap->ClassID() == VRAYHDRICLASS ) : false;
 }
 
 IDatasmithMaterialExpression* FDatasmithMaxVRayHDRITexmapToUEPbr::Convert( FDatasmithMaxMaterialsToUEPbr* MaxMaterialToUEPbr, Texmap* InTexmap )
@@ -214,7 +214,7 @@ IDatasmithMaterialExpression* FDatasmithMaxVRayHDRITexmapToUEPbr::Convert( FData
 
 bool FDatasmithMaxVRayDirtTexmapToUEPbr::IsSupported( const FDatasmithMaxMaterialsToUEPbr* MaxMaterialToUEPbr, Texmap* InTexmap ) const
 {
-	return InTexmap ? InTexmap->ClassID() == VRAYDIRTCLASS : false;
+	return InTexmap ? (bool)( InTexmap->ClassID() == VRAYDIRTCLASS ) : false;
 }
 
 IDatasmithMaterialExpression* FDatasmithMaxVRayDirtTexmapToUEPbr::Convert( FDatasmithMaxMaterialsToUEPbr* MaxMaterialToUEPbr, Texmap* InTexmap )
