Wrong Normals orientation on wall entities exported from Revit Datasmith

Integrated from Dev-Enterprise (CL 11194428)

diff --git a/Engine/Source/Programs/Enterprise/Datasmith/DatasmithRevitExporter/Private/DatasmithRevitExportContext.cs b/Engine/Source/Programs/Enterprise/Datasmith/DatasmithRevitExporter/Private/DatasmithRevitExportContext.cs
--- a/Engine/Source/Programs/Enterprise/Datasmith/DatasmithRevitExporter/Private/DatasmithRevitExportContext.cs
+++ b/Engine/Source/Programs/Enterprise/Datasmith/DatasmithRevitExporter/Private/DatasmithRevitExportContext.cs
@@ -419,21 +419,44 @@ namespace DatasmithRevitExporter
 				case DistributionOfNormals.AtEachPoint:
 				{
 					IList<XYZ> normals = InPolymeshNode.GetNormals();
-					foreach (PolymeshFacet facet in InPolymeshNode.GetFacets())
+					if (MeshPointsTransform != null)
 					{
-						XYZ normal1 = normals[facet.V1];
-						XYZ normal2 = normals[facet.V2];
-						XYZ normal3 = normals[facet.V3];
+						foreach (PolymeshFacet facet in InPolymeshNode.GetFacets())
+						{
+							XYZ normal1 = MeshPointsTransform.OfVector(normals[facet.V1]);
+							XYZ normal2 = MeshPointsTransform.OfVector(normals[facet.V2]);
+							XYZ normal3 = MeshPointsTransform.OfVector(normals[facet.V3]);
+
+							CurrentMesh.AddNormal((float)normal1.X, (float)normal1.Y, (float)normal1.Z);
+							CurrentMesh.AddNormal((float)normal2.X, (float)normal2.Y, (float)normal2.Z);
+							CurrentMesh.AddNormal((float)normal3.X, (float)normal3.Y, (float)normal3.Z);
+						}
+					}
+					else
+					{
+						foreach (PolymeshFacet facet in InPolymeshNode.GetFacets())
+						{
+							XYZ normal1 = normals[facet.V1];
+							XYZ normal2 = normals[facet.V2];
+							XYZ normal3 = normals[facet.V3];
 
-						CurrentMesh.AddNormal((float) normal1.X, (float) normal1.Y, (float) normal1.Z);
-						CurrentMesh.AddNormal((float) normal2.X, (float) normal2.Y, (float) normal2.Z);
-						CurrentMesh.AddNormal((float) normal3.X, (float) normal3.Y, (float) normal3.Z);
+							CurrentMesh.AddNormal((float)normal1.X, (float)normal1.Y, (float)normal1.Z);
+							CurrentMesh.AddNormal((float)normal2.X, (float)normal2.Y, (float)normal2.Z);
+							CurrentMesh.AddNormal((float)normal3.X, (float)normal3.Y, (float)normal3.Z);
+						}
 					}
+
 					break;
 				}
 				case DistributionOfNormals.OnePerFace:
 				{
 					XYZ normal = InPolymeshNode.GetNormals()[0];
+
+					if (MeshPointsTransform != null)
+					{
+						normal = MeshPointsTransform.OfVector(normal);
+					}
+
 					for (int i = 0; i < 3 * InPolymeshNode.NumberOfFacets; i++)
 					{
 						CurrentMesh.AddNormal((float) normal.X, (float) normal.Y, (float) normal.Z);
@@ -442,11 +465,24 @@ namespace DatasmithRevitExporter
 				}
 				case DistributionOfNormals.OnEachFacet:
 				{
-					foreach (XYZ normal in InPolymeshNode.GetNormals())
+					if (MeshPointsTransform != null)
 					{
-						CurrentMesh.AddNormal((float) normal.X, (float) normal.Y, (float) normal.Z);
-						CurrentMesh.AddNormal((float) normal.X, (float) normal.Y, (float) normal.Z);
-						CurrentMesh.AddNormal((float) normal.X, (float) normal.Y, (float) normal.Z);
+						foreach (XYZ normal in InPolymeshNode.GetNormals())
+						{
+							XYZ FinalNormal = MeshPointsTransform.OfVector(normal);
+							CurrentMesh.AddNormal((float)FinalNormal.X, (float)FinalNormal.Y, (float)FinalNormal.Z);
+							CurrentMesh.AddNormal((float)FinalNormal.X, (float)FinalNormal.Y, (float)FinalNormal.Z);
+							CurrentMesh.AddNormal((float)FinalNormal.X, (float)FinalNormal.Y, (float)FinalNormal.Z);
+						}
+					}
+					else
+					{
+						foreach (XYZ normal in InPolymeshNode.GetNormals())
+						{
+							CurrentMesh.AddNormal((float)normal.X, (float)normal.Y, (float)normal.Z);
+							CurrentMesh.AddNormal((float)normal.X, (float)normal.Y, (float)normal.Z);
+							CurrentMesh.AddNormal((float)normal.X, (float)normal.Y, (float)normal.Z);
+						}
 					}
 					break;
 				}
