Integrating CL# 11427051 from Dev-Enterprise
Fixing 3dsMax exporter crash

diff --git a/Engine/Source/Programs/Enterprise/Datasmith/DatasmithMaxExporter/Private/DatasmithMaxMeshExporter.cpp b/Engine/Source/Programs/Enterprise/Datasmith/DatasmithMaxExporter/Private/DatasmithMaxMeshExporter.cpp
--- a/Engine/Source/Programs/Enterprise/Datasmith/DatasmithMaxExporter/Private/DatasmithMaxMeshExporter.cpp
+++ b/Engine/Source/Programs/Enterprise/Datasmith/DatasmithMaxExporter/Private/DatasmithMaxMeshExporter.cpp
@@ -596,10 +596,10 @@ void FDatasmithMaxMeshExporter::FillDatasmithMeshFromMaxMesh(FDatasmithMesh& Dat
 
 	for ( int32 i = 1; i <= MaxMesh.getNumMaps(); ++i )
 	{
-		if ( MaxMesh.mapSupport(i) == BOOL(true) )
+		if ( MaxMesh.mapSupport(i) == BOOL(true) && MaxMesh.getNumMapVerts(i) > 0)
 		{
 			DatasmithMesh.AddUVChannel();
-			int32 UVChannelIndex = DatasmithMesh.GetUVChannelsCount() - 1;
+			const int32 UVChannelIndex = DatasmithMesh.GetUVChannelsCount() - 1;
 			const int32 UVsCount = MaxMesh.getNumMapVerts(i);
 
 			DatasmithMesh.SetUVCount( UVChannelIndex, UVsCount );
