Integrating CL# 11292240 from Dev-Enterprise
Fixing crash in 3ds Max Exporter, autodesk noise map are not supported.

Julien.StJean

diff --git a/Engine/Source/Programs/Enterprise/Datasmith/DatasmithMaxExporter/Private/DatasmithMaxAttributes.cpp b/Engine/Source/Programs/Enterprise/Datasmith/DatasmithMaxExporter/Private/DatasmithMaxAttributes.cpp
--- a/Engine/Source/Programs/Enterprise/Datasmith/DatasmithMaxExporter/Private/DatasmithMaxAttributes.cpp
+++ b/Engine/Source/Programs/Enterprise/Datasmith/DatasmithMaxExporter/Private/DatasmithMaxAttributes.cpp
@@ -115,6 +115,8 @@ TOptional<FDatasmithMaxStaticMeshAttributes> FDatasmithMaxStaticMeshAttributes::
 					}
 				}
 			}
+
+			ParamBlock2->ReleaseDesc();
 		}
 	}
 	
diff --git a/Engine/Source/Programs/Enterprise/Datasmith/DatasmithMaxExporter/Private/DatasmithMaxCameraExporter.cpp b/Engine/Source/Programs/Enterprise/Datasmith/DatasmithMaxExporter/Private/DatasmithMaxCameraExporter.cpp
--- a/Engine/Source/Programs/Enterprise/Datasmith/DatasmithMaxExporter/Private/DatasmithMaxCameraExporter.cpp
+++ b/Engine/Source/Programs/Enterprise/Datasmith/DatasmithMaxExporter/Private/DatasmithMaxCameraExporter.cpp
@@ -535,6 +535,8 @@ bool FDatasmithMaxCameraExporter::ExportPhysicalExposureControl( ToneOperator& T
 				GlobalEV = ParamBlock2->GetFloat( ParamDefinition.ID, GetCOREInterface()->GetTime(), ValidInterval );
 			}
 		}
+
+		ParamBlock2->ReleaseDesc();
 	}
 
 	// Compute the shutter speed based on the EV, F-Stop and ISO
diff --git a/Engine/Source/Programs/Enterprise/Datasmith/DatasmithMaxExporter/Private/DatasmithMaxMeshExporter.cpp b/Engine/Source/Programs/Enterprise/Datasmith/DatasmithMaxExporter/Private/DatasmithMaxMeshExporter.cpp
--- a/Engine/Source/Programs/Enterprise/Datasmith/DatasmithMaxExporter/Private/DatasmithMaxMeshExporter.cpp
+++ b/Engine/Source/Programs/Enterprise/Datasmith/DatasmithMaxExporter/Private/DatasmithMaxMeshExporter.cpp
@@ -94,6 +94,7 @@ namespace
 					break;
 				}
 			}
+			ParamBlock2->ReleaseDesc();
 		}
 
 		return PrevDisplayValue;
diff --git a/Engine/Source/Programs/Enterprise/Datasmith/DatasmithMaxExporter/Private/DatasmithMaxSceneExporter.cpp b/Engine/Source/Programs/Enterprise/Datasmith/DatasmithMaxExporter/Private/DatasmithMaxSceneExporter.cpp
--- a/Engine/Source/Programs/Enterprise/Datasmith/DatasmithMaxExporter/Private/DatasmithMaxSceneExporter.cpp
+++ b/Engine/Source/Programs/Enterprise/Datasmith/DatasmithMaxExporter/Private/DatasmithMaxSceneExporter.cpp
@@ -198,6 +198,7 @@ int FDatasmithMaxSceneExporter::GetSeedFromMaterial(Mtl* Material)
 				Seed = ParamBlock2->GetInt(ParamDefinition.ID);
 			}
 		}
+		ParamBlock2->ReleaseDesc();
 	}
 	return Seed;
 }
@@ -707,6 +708,8 @@ TSharedPtr< IDatasmithLightActorElement > FDatasmithMaxSceneExporter::CreateLigh
 				}
 			}
 		}
+
+		ParamBlock2->ReleaseDesc();
 	}
 
 	TSharedPtr< IDatasmithElement > Element = FDatasmithSceneFactory::CreateElement( LightType, Name );
@@ -1213,6 +1216,8 @@ bool FDatasmithMaxSceneExporter::ParseCoronaLight(LightObject& Light, TSharedRef
 				}
 			}
 		}
+
+		ParamBlock2->ReleaseDesc();
 	}
 
 	if ( !bLightShapeVisible)
@@ -1644,6 +1649,8 @@ bool FDatasmithMaxSceneExporter::ParseVRayLightIES(LightObject& Light, TSharedRe
 				PointLightElement->SetTemperature( ParamBlock2->GetFloat( ParamDefinition.ID, GetCOREInterface()->GetTime() ) );
 			}
 		}
+
+		ParamBlock2->ReleaseDesc();
 	}
 
 	if ( PointLightElement->GetUseTemperature() )
@@ -1719,6 +1726,8 @@ bool FDatasmithMaxSceneExporter::ParseVRayLightIES(LightObject& Light, TSharedRe
 					}
 				}
 			}
+
+			ParamBlock2->ReleaseDesc();
 		}
 
 		if ( bUseLightShape  )
@@ -1847,6 +1856,8 @@ bool FDatasmithMaxSceneExporter::ParseLightParameters(EMaxLightClass LightClass,
 				}
 			}
 		}
+
+		ParamBlock2->ReleaseDesc();
 	}
 
 	if ( LightElement->GetUseTemperature() )
diff --git a/Engine/Source/Programs/Enterprise/Datasmith/DatasmithMaxExporter/Private/DatasmithMaxSceneParser.cpp b/Engine/Source/Programs/Enterprise/Datasmith/DatasmithMaxExporter/Private/DatasmithMaxSceneParser.cpp
--- a/Engine/Source/Programs/Enterprise/Datasmith/DatasmithMaxExporter/Private/DatasmithMaxSceneParser.cpp
+++ b/Engine/Source/Programs/Enterprise/Datasmith/DatasmithMaxExporter/Private/DatasmithMaxSceneParser.cpp
@@ -520,6 +520,8 @@ void FDatasmithMaxSceneParser::LightTexEnum(INode* Node)
 					}
 				}
 			}
+
+			ParamBlock2->ReleaseDesc();
 		}
 	}
 }
diff --git a/Engine/Source/Programs/Enterprise/Datasmith/DatasmithMaxExporter/Private/DatasmithMaxTexmapParser.cpp b/Engine/Source/Programs/Enterprise/Datasmith/DatasmithMaxExporter/Private/DatasmithMaxTexmapParser.cpp
--- a/Engine/Source/Programs/Enterprise/Datasmith/DatasmithMaxExporter/Private/DatasmithMaxTexmapParser.cpp
+++ b/Engine/Source/Programs/Enterprise/Datasmith/DatasmithMaxExporter/Private/DatasmithMaxTexmapParser.cpp
@@ -215,6 +215,8 @@ namespace DatasmithMaxTexmapParser
 					NormalMapParameters.BumpMap.Weight = ParamBlock2->GetFloat( ParamDefinition.ID, GetCOREInterface()->GetTime() );
 				}
 			}
+
+			ParamBlock2->ReleaseDesc();
 		}
 
 		return NormalMapParameters;
@@ -293,6 +295,8 @@ namespace DatasmithMaxTexmapParser
 					AutodeskBitmapParameters.MapChannel = (ParamBlock2->GetInt(ParamDefinition.ID, CurrentTime));
 				}
 			}
+
+			ParamBlock2->ReleaseDesc();
 		}
 
 		return AutodeskBitmapParameters;
diff --git a/Engine/Source/Programs/Enterprise/Datasmith/DatasmithMaxExporter/Private/DatasmithMaxWriter.cpp b/Engine/Source/Programs/Enterprise/Datasmith/DatasmithMaxExporter/Private/DatasmithMaxWriter.cpp
--- a/Engine/Source/Programs/Enterprise/Datasmith/DatasmithMaxExporter/Private/DatasmithMaxWriter.cpp
+++ b/Engine/Source/Programs/Enterprise/Datasmith/DatasmithMaxExporter/Private/DatasmithMaxWriter.cpp
@@ -67,11 +67,15 @@ EDSBitmapType FDatasmithMaxMatHelper::GetTextureClass(Texmap* InTexMap)
 
 	if (InTexMap != nullptr)
 	{
+		MSTR ClassName;
+		InTexMap->GetClassName(ClassName);
+
 		if (InTexMap->ClassID() == RBITMAPCLASS)
 		{
 			Type = EDSBitmapType::RegularBitmap;
 		}
-		else if (InTexMap->ClassID() == AUTODESKBITMAPCLASS)
+		//Somehow, there are multiple autodesk map classes using the same ClassID, we only support Autodesk Bitmap.
+		else if (InTexMap->ClassID() == AUTODESKBITMAPCLASS && FCString::Stricmp(ClassName, TEXT("Autodesk Bitmap")) == 0)
 		{
 			Type = EDSBitmapType::AutodeskBitmap;
 		}
diff --git a/Engine/Source/Programs/Enterprise/Datasmith/DatasmithMaxExporter/Private/MaxMaterialsToUEPbr/DatasmithMaxTexmapToUEPbr.cpp b/Engine/Source/Programs/Enterprise/Datasmith/DatasmithMaxExporter/Private/MaxMaterialsToUEPbr/DatasmithMaxTexmapToUEPbr.cpp
--- a/Engine/Source/Programs/Enterprise/Datasmith/DatasmithMaxExporter/Private/MaxMaterialsToUEPbr/DatasmithMaxTexmapToUEPbr.cpp
+++ b/Engine/Source/Programs/Enterprise/Datasmith/DatasmithMaxExporter/Private/MaxMaterialsToUEPbr/DatasmithMaxTexmapToUEPbr.cpp
@@ -128,6 +128,8 @@ namespace DatasmithMaxBitmapToUEPbrImpl
 					RGBTintParameters.Map1.bEnabled = ( ParamBlock2->GetInt( ParamDefinition.ID, GetCOREInterface()->GetTime() ) != 0 );
 				}
 			}
+
+			ParamBlock2->ReleaseDesc();
 		}
 
 		return RGBTintParameters;
@@ -173,6 +175,8 @@ namespace DatasmithMaxBitmapToUEPbrImpl
 					MixParameters.MaskMap.bEnabled = ( ParamBlock2->GetInt(ParamDefinition.ID, GetCOREInterface()->GetTime()) != 0 );
 				}
 			}
+
+			ParamBlock2->ReleaseDesc();
 		}
 
 		return MixParameters;
@@ -260,6 +264,8 @@ namespace DatasmithMaxBitmapToUEPbrImpl
 					NoiseParameters.Size = ParamBlock2->GetFloat(ParamDefinition.ID, GetCOREInterface()->GetTime());
 				}
 			}
+
+			ParamBlock2->ReleaseDesc();
 		}
 
 		return NoiseParameters;
@@ -812,9 +818,15 @@ IDatasmithMaterialExpression* FDatasmithMaxBitmapToUEPbr::Convert( FDatasmithMax
 
 bool FDatasmithMaxAutodeskBitmapToUEPbr::IsSupported(const FDatasmithMaxMaterialsToUEPbr* MaxMaterialToUEPbr, Texmap* InTexmap) const
 {
-	if (InTexmap && InTexmap->ClassID() == AUTODESKBITMAPCLASS)
+	if (InTexmap)
 	{
-		return true;
+		MSTR ClassName;
+		InTexmap->GetClassName(ClassName);
+		//Somehow, there are multiple autodesk map classes using the same ClassID, we only support Autodesk Bitmap.
+		if (InTexmap->ClassID() == AUTODESKBITMAPCLASS && FCString::Stricmp(ClassName, TEXT("Autodesk Bitmap")) == 0)
+		{
+			return true;
+		}
 	}
 
 	return false;
@@ -823,12 +835,16 @@ bool FDatasmithMaxAutodeskBitmapToUEPbr::IsSupported(const FDatasmithMaxMaterial
 IDatasmithMaterialExpression* FDatasmithMaxAutodeskBitmapToUEPbr::Convert(FDatasmithMaxMaterialsToUEPbr* MaxMaterialToUEPbr, Texmap* InTexmap)
 {
 	DatasmithMaxTexmapParser::FAutodeskBitmapParameters AutodeskBitmapParameters = DatasmithMaxTexmapParser::ParseAutodeskBitmap(InTexmap);
-	FScopedBitMapPtr ActualBitmap(AutodeskBitmapParameters.SourceFile->bi, AutodeskBitmapParameters.SourceFile->bm);
-	FString ActualBitmapName = FDatasmithMaxMatWriter::GetActualBitmapName(&ActualBitmap.MapInfo);
-	bool bUseAlphaAsMono = (ActualBitmap.Map->HasAlpha() != 0);
-	bool bIsSRGB = FDatasmithMaxMatHelper::IsSRGB(*ActualBitmap.Map);
+	if (AutodeskBitmapParameters.SourceFile)
+	{
+		FScopedBitMapPtr ActualBitmap(AutodeskBitmapParameters.SourceFile->bi, AutodeskBitmapParameters.SourceFile->bm);
+		FString ActualBitmapName = FDatasmithMaxMatWriter::GetActualBitmapName(&ActualBitmap.MapInfo);
+		bool bUseAlphaAsMono = (ActualBitmap.Map->HasAlpha() != 0);
+		bool bIsSRGB = FDatasmithMaxMatHelper::IsSRGB(*ActualBitmap.Map);
 
-	return FDatasmithMaxTexmapToUEPbrUtils::ConvertBitMap(MaxMaterialToUEPbr, InTexmap, ActualBitmapName, bUseAlphaAsMono, bIsSRGB);
+		return FDatasmithMaxTexmapToUEPbrUtils::ConvertBitMap(MaxMaterialToUEPbr, InTexmap, ActualBitmapName, bUseAlphaAsMono, bIsSRGB);
+	}
+	return nullptr;
 }
 
 bool FDatasmithMaxNormalToUEPbr::IsSupported( const FDatasmithMaxMaterialsToUEPbr* MaxMaterialToUEPbr, Texmap* InTexmap ) const
diff --git a/Engine/Source/Programs/Enterprise/Datasmith/DatasmithMaxExporter/Private/MaxMaterialsToUEPbr/DatasmithMaxVRayTexmapToUEPbr.cpp b/Engine/Source/Programs/Enterprise/Datasmith/DatasmithMaxExporter/Private/MaxMaterialsToUEPbr/DatasmithMaxVRayTexmapToUEPbr.cpp
--- a/Engine/Source/Programs/Enterprise/Datasmith/DatasmithMaxExporter/Private/MaxMaterialsToUEPbr/DatasmithMaxVRayTexmapToUEPbr.cpp
+++ b/Engine/Source/Programs/Enterprise/Datasmith/DatasmithMaxExporter/Private/MaxMaterialsToUEPbr/DatasmithMaxVRayTexmapToUEPbr.cpp
@@ -131,6 +131,8 @@ namespace DatasmithMaxVRayTexmapToUEPbrImpl
 					VRayDirtParameters.UnoccludedMap.Weight = ParamBlock2->GetFloat( ParamDefinition.ID, GetCOREInterface()->GetTime() ) / 100.f;
 				}
 			}
+
+			ParamBlock2->ReleaseDesc();
 		}
 
 		return VRayDirtParameters;
