Scene Hierarchy from a Revit Datasmith export has duplicates

diff --git a/Engine/Source/Programs/Enterprise/Datasmith/DatasmithRevitExporter/Private/DatasmithRevitExportContext.cs b/Engine/Source/Programs/Enterprise/Datasmith/DatasmithRevitExporter/Private/DatasmithRevitExportContext.cs
--- a/Engine/Source/Programs/Enterprise/Datasmith/DatasmithRevitExporter/Private/DatasmithRevitExportContext.cs
+++ b/Engine/Source/Programs/Enterprise/Datasmith/DatasmithRevitExporter/Private/DatasmithRevitExportContext.cs
@@ -288,7 +288,7 @@ namespace DatasmithRevitExporter
 			if (LinkedDocument != null)
 			{
 				// Keep track of the linked document being processed.
-				PushDocument(LinkedDocument);
+				PushDocument(LinkedDocument, false);
 			}
 
 			return (CurrentInstanceType != null && LinkedDocument != null) ? RenderNodeAction.Proceed : RenderNodeAction.Skip;
@@ -544,12 +544,16 @@ namespace DatasmithRevitExporter
 		}
 
 		private void PushDocument(
-			Document InDocument
+			Document InDocument,
+			bool bInAddLocationActors = true
 		)
 		{
 			DocumentDataStack.Push(new FDocumentData(InDocument, ref MessageList));
 
-			DocumentDataStack.Peek().AddLocationActors(WorldTransformStack.Peek());
+			if (bInAddLocationActors)
+			{
+				DocumentDataStack.Peek().AddLocationActors(WorldTransformStack.Peek());
+			}
 		}
 
 		private void PopDocument()
