Added support for Sketchup 2020 in Datasmith SketchUp exporter

diff --git a/Engine/Source/Programs/Enterprise/Datasmith/DatasmithSketchUpExporter/DatasmithSketchUp2020.Build.cs b/Engine/Source/Programs/Enterprise/Datasmith/DatasmithSketchUpExporter/DatasmithSketchUp2020.Build.cs
new file mode 100644
--- /dev/null
+++ b/Engine/Source/Programs/Enterprise/Datasmith/DatasmithSketchUpExporter/DatasmithSketchUp2020.Build.cs
@@ -0,0 +1,26 @@
+// Copyright Epic Games, Inc. All Rights Reserved.
+
+using System.IO;
+
+namespace UnrealBuildTool.Rules
+{
+	[SupportedPlatforms("Win64")]
+	public class DatasmithSketchUp2020 : DatasmithSketchUpBase
+	{
+		public DatasmithSketchUp2020(ReadOnlyTargetRules Target)
+			: base(Target)
+		{
+			PrivateDefinitions.Add("SKP_SDK_2020");
+		}
+
+		public override string GetSketchUpSDKFolder()
+		{
+			return "SDK_WIN_x64_2020-0-363";
+		}
+
+		public override string GetSketchUpEnvVar()
+		{
+			return "SKP_SDK_2020";
+		}
+	}
+}
diff --git a/Engine/Source/Programs/Enterprise/Datasmith/DatasmithSketchUpExporter/DatasmithSketchUp2020.Target.cs b/Engine/Source/Programs/Enterprise/Datasmith/DatasmithSketchUpExporter/DatasmithSketchUp2020.Target.cs
new file mode 100644
--- /dev/null
+++ b/Engine/Source/Programs/Enterprise/Datasmith/DatasmithSketchUpExporter/DatasmithSketchUp2020.Target.cs
@@ -0,0 +1,16 @@
+// Copyright Epic Games, Inc. All Rights Reserved.
+
+using UnrealBuildTool;
+using System.IO;
+
+public class DatasmithSketchUp2020Target : DatasmithSketchUpBaseTarget
+{
+	public DatasmithSketchUp2020Target(TargetInfo Target)
+		: base(Target)
+	{
+		LaunchModuleName = "DatasmithSketchUp2020";
+		ExeBinariesSubFolder = Path.Combine("SketchUp", "2020");
+		
+		AddCopyPostBuildStep(Target, "skp2udatasmith2020");
+	}
+}
