Copy of 11271391 from //UE4/Dev-Enterprise:

Revit: Fix export time field in exported datasmith scene

diff --git a/Engine/Source/Developer/Datasmith/DatasmithFacade/Private/DatasmithFacadeScene.cpp b/Engine/Source/Developer/Datasmith/DatasmithFacade/Private/DatasmithFacadeScene.cpp
--- a/Engine/Source/Developer/Datasmith/DatasmithFacade/Private/DatasmithFacadeScene.cpp
+++ b/Engine/Source/Developer/Datasmith/DatasmithFacade/Private/DatasmithFacadeScene.cpp
@@ -106,20 +106,28 @@ void FDatasmithFacadeScene::BuildScene(
 	FDatasmithFacadeMaterial::ClearBuiltTextureSet();
 }
 
-void FDatasmithFacadeScene::ExportScene(
-	const TCHAR* InOutputPath
-)
+void FDatasmithFacadeScene::PreExport()
 {
-	FString OutputPath = InOutputPath;
-
 	// Initialize the Datasmith exporter module.
 	FDatasmithExporterManager::Initialize();
 
 	// Create a Datasmith scene exporter.
-	TSharedRef<FDatasmithSceneExporter> SceneExporterRef = MakeShared<FDatasmithSceneExporter>();
+	SceneExporterRef = MakeShared<FDatasmithSceneExporter>();
 
 	// Start measuring the time taken to export the scene.
 	SceneExporterRef->PreExport();
+}
+
+void FDatasmithFacadeScene::ExportScene(
+	const TCHAR* InOutputPath
+)
+{
+	if (!SceneExporterRef.IsValid())
+	{
+		return;
+	}
+
+	FString OutputPath = InOutputPath;
 
 	// Set the name of the scene to export and let Datasmith sanitize it when required.
 	FString SceneName = FPaths::GetBaseFilename(OutputPath);
diff --git a/Engine/Source/Developer/Datasmith/DatasmithFacade/Public/DatasmithFacadeScene.h b/Engine/Source/Developer/Datasmith/DatasmithFacade/Public/DatasmithFacadeScene.h
--- a/Engine/Source/Developer/Datasmith/DatasmithFacade/Public/DatasmithFacadeScene.h
+++ b/Engine/Source/Developer/Datasmith/DatasmithFacade/Public/DatasmithFacadeScene.h
@@ -4,6 +4,7 @@
 
 // Datasmith SDK.
 #include "DatasmithSceneFactory.h"
+#include "DatasmithSceneExporter.h"
 
 // Datasmith facade classes.
 class FDatasmithFacadeElement;
@@ -42,6 +43,9 @@ public:
 		const TCHAR* InSceneName // Datasmith scene name
 	);
 
+	// Instanciate an exporter and register export start time
+	void PreExport();
+
 	// Build and export a Datasmith scene instance and its scene element assets.
 	void ExportScene(
 		const TCHAR* InOutputPath // Datasmith scene output file path
@@ -78,4 +82,7 @@ private:
 
 	// Datasmith scene instance built with the collected elements.
 	TSharedRef<IDatasmithScene> SceneRef;
+
+	// Datasmith scene exporter
+	TSharedPtr<FDatasmithSceneExporter> SceneExporterRef;
 };
diff --git a/Engine/Source/Programs/Enterprise/Datasmith/DatasmithFacadeCSharp/Private/DatasmithFacadeCSharp.cpp b/Engine/Source/Programs/Enterprise/Datasmith/DatasmithFacadeCSharp/Private/DatasmithFacadeCSharp.cpp
--- a/Engine/Source/Programs/Enterprise/Datasmith/DatasmithFacadeCSharp/Private/DatasmithFacadeCSharp.cpp
+++ b/Engine/Source/Programs/Enterprise/Datasmith/DatasmithFacadeCSharp/Private/DatasmithFacadeCSharp.cpp
@@ -1358,6 +1358,14 @@ SWIGEXPORT void SWIGSTDCALL CSharp_FDatasmithFacadeScene_BuildScene(void * jarg1
 }
 
 
+SWIGEXPORT void SWIGSTDCALL CSharp_FDatasmithFacadeScene_PreExport(void * jarg1) {
+  FDatasmithFacadeScene *arg1 = (FDatasmithFacadeScene *) 0 ;
+  
+  arg1 = (FDatasmithFacadeScene *)jarg1; 
+  (arg1)->PreExport();
+}
+
+
 SWIGEXPORT void SWIGSTDCALL CSharp_FDatasmithFacadeScene_ExportScene(void * jarg1, wchar_t * jarg2) {
   FDatasmithFacadeScene *arg1 = (FDatasmithFacadeScene *) 0 ;
   TCHAR *arg2 = (TCHAR *) 0 ;
diff --git a/Engine/Source/Programs/Enterprise/Datasmith/DatasmithFacadeCSharp/Public/DatasmithFacadeCSharpPINVOKE.cs b/Engine/Source/Programs/Enterprise/Datasmith/DatasmithFacadeCSharp/Public/DatasmithFacadeCSharpPINVOKE.cs
--- a/Engine/Source/Programs/Enterprise/Datasmith/DatasmithFacadeCSharp/Public/DatasmithFacadeCSharpPINVOKE.cs
+++ b/Engine/Source/Programs/Enterprise/Datasmith/DatasmithFacadeCSharp/Public/DatasmithFacadeCSharpPINVOKE.cs
@@ -476,6 +476,9 @@ class DatasmithFacadeCSharpPINVOKE {
   [global::System.Runtime.InteropServices.DllImport("DatasmithFacadeCSharp", EntryPoint="CSharp_FDatasmithFacadeScene_BuildScene")]
   public static extern void FDatasmithFacadeScene_BuildScene(global::System.Runtime.InteropServices.HandleRef jarg1, [global::System.Runtime.InteropServices.MarshalAs(global::System.Runtime.InteropServices.UnmanagedType.LPWStr)]string jarg2);
 
+  [global::System.Runtime.InteropServices.DllImport("DatasmithFacadeCSharp", EntryPoint="CSharp_FDatasmithFacadeScene_PreExport")]
+  public static extern void FDatasmithFacadeScene_PreExport(global::System.Runtime.InteropServices.HandleRef jarg1);
+
   [global::System.Runtime.InteropServices.DllImport("DatasmithFacadeCSharp", EntryPoint="CSharp_FDatasmithFacadeScene_ExportScene")]
   public static extern void FDatasmithFacadeScene_ExportScene(global::System.Runtime.InteropServices.HandleRef jarg1, [global::System.Runtime.InteropServices.MarshalAs(global::System.Runtime.InteropServices.UnmanagedType.LPWStr)]string jarg2);
 
diff --git a/Engine/Source/Programs/Enterprise/Datasmith/DatasmithFacadeCSharp/Public/FDatasmithFacadeScene.cs b/Engine/Source/Programs/Enterprise/Datasmith/DatasmithFacadeCSharp/Public/FDatasmithFacadeScene.cs
--- a/Engine/Source/Programs/Enterprise/Datasmith/DatasmithFacadeCSharp/Public/FDatasmithFacadeScene.cs
+++ b/Engine/Source/Programs/Enterprise/Datasmith/DatasmithFacadeCSharp/Public/FDatasmithFacadeScene.cs
@@ -64,6 +64,10 @@ public class FDatasmithFacadeScene : global::System.IDisposable {
     DatasmithFacadeCSharpPINVOKE.FDatasmithFacadeScene_BuildScene(swigCPtr, InSceneName);
   }
 
+  public void PreExport() {
+    DatasmithFacadeCSharpPINVOKE.FDatasmithFacadeScene_PreExport(swigCPtr);
+  }
+
   public void ExportScene(string InOutputPath) {
     DatasmithFacadeCSharpPINVOKE.FDatasmithFacadeScene_ExportScene(swigCPtr, InOutputPath);
   }
diff --git a/Engine/Source/Programs/Enterprise/Datasmith/DatasmithRevitExporter/Private/DatasmithRevitExportContext.cs b/Engine/Source/Programs/Enterprise/Datasmith/DatasmithRevitExporter/Private/DatasmithRevitExportContext.cs
--- a/Engine/Source/Programs/Enterprise/Datasmith/DatasmithRevitExporter/Private/DatasmithRevitExportContext.cs
+++ b/Engine/Source/Programs/Enterprise/Datasmith/DatasmithRevitExporter/Private/DatasmithRevitExportContext.cs
@@ -15,7 +15,7 @@ using Autodesk.Revit.DB.Visual;
 
 namespace DatasmithRevitExporter
 {
-	// Custom export context for command Export to Unreal Datasmith. 
+	// Custom export context for command Export to Unreal Datasmith.
 	public class FDatasmithRevitExportContext : IPhotoRenderContext
 	{
 		// Revit application information for Datasmith.
@@ -52,7 +52,7 @@ namespace DatasmithRevitExporter
 
 		// List of extra search paths for Revit texture files.
 		private IList<string> ExtraTexturePaths = new List<string>();
-		
+
 		// List of messages generated during the export process.
 		private List<string> MessageList = new List<string>();
 
@@ -140,6 +140,7 @@ namespace DatasmithRevitExporter
 
 			// Create an empty Datasmith scene.
 			DatasmithScene = new FDatasmithFacadeScene(HOST_NAME, VENDOR_NAME, PRODUCT_NAME, ProductVersion);
+			DatasmithScene.PreExport();
 
 			View3D ViewToExport = RevitDocument.GetElement(InViewNode.ViewId) as View3D;
 			if (!DatasmithFilePaths.TryGetValue(ViewToExport.Id, out CurrentDatasmithFilePath))
@@ -174,7 +175,7 @@ namespace DatasmithRevitExporter
 
 			// Dispose of the Datasmith scene.
 			DatasmithScene = null;
-			
+
 			// Forget the 3D view world transform.
 			WorldTransformStack.Pop();
 		}
@@ -190,7 +191,7 @@ namespace DatasmithRevitExporter
 			{
 				// Keep track of the element being processed.
 				PushElement(CurrentElement, WorldTransformStack.Peek(), "Element Begin");
-			
+
 				// We want to export the element.
 				return RenderNodeAction.Proceed;
 			}
@@ -331,7 +332,7 @@ namespace DatasmithRevitExporter
 			// Forget the current light world transform.
 			WorldTransformStack.Pop();
 		}
-		
+
 		// OnRPC marks the beginning of export of an RPC object.
 		// This method is only called for interface IPhotoRenderContext.
 		public void OnRPC(
@@ -406,12 +407,13 @@ namespace DatasmithRevitExporter
 				CurrentMesh.AddUV(0, (float) uv.U, (float) -uv.V);
 			}
 
+			IList<PolymeshFacet> Facets = InPolymeshNode.GetFacets();
 			// Add the triangle vertex indexes to the Datasmith mesh.
-			foreach (PolymeshFacet facet in InPolymeshNode.GetFacets())
+			foreach (PolymeshFacet facet in Facets)
 			{
 				CurrentMesh.AddTriangle(initialVertexCount + facet.V1, initialVertexCount + facet.V2, initialVertexCount + facet.V3, CurrentMaterialIndex);
 			}
-			
+
 			// Add the triangle vertex normals (in right-handed Z-up coordinates) to the Datasmith mesh.
 			// Normals can be associated with either points or facets of the polymesh.
 			switch (InPolymeshNode.DistributionOfNormals)
@@ -421,7 +423,7 @@ namespace DatasmithRevitExporter
 					IList<XYZ> normals = InPolymeshNode.GetNormals();
 					if (MeshPointsTransform != null)
 					{
-						foreach (PolymeshFacet facet in InPolymeshNode.GetFacets())
+						foreach (PolymeshFacet facet in Facets)
 						{
 							XYZ normal1 = MeshPointsTransform.OfVector(normals[facet.V1]);
 							XYZ normal2 = MeshPointsTransform.OfVector(normals[facet.V2]);
@@ -434,7 +436,7 @@ namespace DatasmithRevitExporter
 					}
 					else
 					{
-						foreach (PolymeshFacet facet in InPolymeshNode.GetFacets())
+						foreach (PolymeshFacet facet in Facets)
 						{
 							XYZ normal1 = normals[facet.V1];
 							XYZ normal2 = normals[facet.V2];
@@ -773,8 +775,8 @@ namespace DatasmithRevitExporter
 					{
 						InWorldTransform = PivotTransform;
 					}
-					
-					if (CurrentElement.GetType() == typeof(Wall) 
+
+					if (CurrentElement.GetType() == typeof(Wall)
 						|| CurrentElement.GetType() == typeof(ModelText)
 						|| CurrentElement.GetType().IsSubclassOf(typeof(MEPCurve)))
 					{
@@ -785,7 +787,7 @@ namespace DatasmithRevitExporter
 				CreateMeshActor(InWorldTransform, out ElementMesh, out ElementActor);
 			}
 
-			// Compute orthonormal basis, given the X vector. 
+			// Compute orthonormal basis, given the X vector.
 			static private void ComputeBasis(XYZ BasisX, ref XYZ BasisY, ref XYZ BasisZ)
 			{
 				BasisY = XYZ.BasisZ.CrossProduct(BasisX);
@@ -1409,7 +1411,7 @@ namespace DatasmithRevitExporter
 		{
 			ElementDataStack.Peek().AddLightActor(InWorldTransform, InLightAsset);
 		}
-		
+
 		public void AddRPCActor(
 			Transform InWorldTransform,
 			Asset     InRPCAsset
