Some entities in Revit exporter Datasmith scenes have an offset on their meshes

diff --git a/Engine/Source/Programs/Enterprise/Datasmith/DatasmithRevitExporter/Private/DatasmithRevitExportContext.cs b/Engine/Source/Programs/Enterprise/Datasmith/DatasmithRevitExporter/Private/DatasmithRevitExportContext.cs
--- a/Engine/Source/Programs/Enterprise/Datasmith/DatasmithRevitExporter/Private/DatasmithRevitExportContext.cs
+++ b/Engine/Source/Programs/Enterprise/Datasmith/DatasmithRevitExporter/Private/DatasmithRevitExportContext.cs
@@ -768,8 +768,7 @@ namespace DatasmithRevitExporter
 				{
 					if (!InWorldTransform.IsIdentity)
 					{
-						PivotTransform = PivotTransform * InWorldTransform;
-						InWorldTransform = PivotTransform;
+						InWorldTransform = PivotTransform * InWorldTransform;
 					}
 					else
 					{
