Some entities in Revit exporter Datasmith scenes have an offset on their meshes

diff --git a/Engine/Source/Programs/Enterprise/Datasmith/DatasmithRevitExporter/Private/DatasmithRevitExportContext.cs b/Engine/Source/Programs/Enterprise/Datasmith/DatasmithRevitExporter/Private/DatasmithRevitExportContext.cs
--- a/Engine/Source/Programs/Enterprise/Datasmith/DatasmithRevitExporter/Private/DatasmithRevitExportContext.cs
+++ b/Engine/Source/Programs/Enterprise/Datasmith/DatasmithRevitExporter/Private/DatasmithRevitExportContext.cs
@@ -768,7 +768,7 @@ namespace DatasmithRevitExporter
 				{
 					if (!InWorldTransform.IsIdentity)
 					{
-						InWorldTransform = PivotTransform * InWorldTransform;
+						InWorldTransform = InWorldTransform * PivotTransform;
 					}
 					else
 					{
