Revit Structural Steel Connections are not exported with Datasmith

diff --git a/Engine/Source/Programs/Enterprise/Datasmith/DatasmithRevitExporter/Private/DatasmithRevitExportContext.cs b/Engine/Source/Programs/Enterprise/Datasmith/DatasmithRevitExporter/Private/DatasmithRevitExportContext.cs
--- a/Engine/Source/Programs/Enterprise/Datasmith/DatasmithRevitExporter/Private/DatasmithRevitExportContext.cs
+++ b/Engine/Source/Programs/Enterprise/Datasmith/DatasmithRevitExporter/Private/DatasmithRevitExportContext.cs
@@ -235,12 +235,11 @@ namespace DatasmithRevitExporter
 				{
 					PushElement(CurrentInstance, WorldTransformStack.Peek(), "Symbol Begin");
 				}
-
-				// We want to process the instance.
-				return RenderNodeAction.Proceed;
 			}
 
-            return RenderNodeAction.Skip;
+			// We always wanna proceed, because in certain cases where InInstanceNode is valid but CurrentInstance is not, 
+			// what follows is valid geometry related to the instance previously exported.
+            return RenderNodeAction.Proceed;
 		}
 
 		// OnInstanceEnd marks the end of a family instance being exported.
