- fix DatasmithCADWorker launching without CAD_INTERFACE: avoid crash, log a message, return a specific code error
 - fix FDatasmithWorkerHandler::StartWorkerProcess: Test if the external process and extenal application are running (FProcHandle::IsValide is not enough to check if a process is well launch)

diff --git a/Engine/Plugins/Enterprise/DatasmithCADImporter/Source/DatasmithDispatcher/Private/DatasmithWorkerHandler.cpp b/Engine/Plugins/Enterprise/DatasmithCADImporter/Source/DatasmithDispatcher/Private/DatasmithWorkerHandler.cpp
--- a/Engine/Plugins/Enterprise/DatasmithCADImporter/Source/DatasmithDispatcher/Private/DatasmithWorkerHandler.cpp
+++ b/Engine/Plugins/Enterprise/DatasmithCADImporter/Source/DatasmithDispatcher/Private/DatasmithWorkerHandler.cpp
@@ -58,6 +58,8 @@ FDatasmithWorkerHandler::~FDatasmithWorkerHandler()
 
 void FDatasmithWorkerHandler::StartWorkerProcess()
 {
+	uint32 WorkerProcessId = 0;
+
 	ensure(ErrorState == EWorkerErrorState::Ok);
 	FString ProcessorPath = GetWorkerExecutablePath();
 	if (FPaths::FileExists(ProcessorPath))
@@ -83,11 +85,10 @@ void FDatasmithWorkerHandler::StartWorkerProcess()
 		CommandToProcess += TEXT(" -EnginePluginsDir \"") + FPaths::EnginePluginsDir() + TEXT('"');
 		UE_LOG(LogDatasmithDispatcher, Verbose, TEXT("CommandToProcess: %s"), *CommandToProcess);
 
-		uint32 WorkerProcessId = 0;
 		WorkerHandle = FPlatformProcess::CreateProc(*ProcessorPath, *CommandToProcess, true, false, false, &WorkerProcessId, 0, nullptr, nullptr);
 	}
 
-	if (!WorkerHandle.IsValid())
+	if (!WorkerHandle.IsValid() || !FPlatformProcess::IsProcRunning(WorkerHandle) || !FPlatformProcess::IsApplicationRunning(WorkerProcessId))
 	{
 		ErrorState = EWorkerErrorState::WorkerProcess_CantCreate;
 		return;
diff --git a/Engine/Source/Programs/Enterprise/Datasmith/DatasmithCADWorker/Private/DatasmithCADWorker.cpp b/Engine/Source/Programs/Enterprise/Datasmith/DatasmithCADWorker/Private/DatasmithCADWorker.cpp
--- a/Engine/Source/Programs/Enterprise/Datasmith/DatasmithCADWorker/Private/DatasmithCADWorker.cpp
+++ b/Engine/Source/Programs/Enterprise/Datasmith/DatasmithCADWorker/Private/DatasmithCADWorker.cpp
@@ -9,6 +9,8 @@
 IMPLEMENT_APPLICATION(DatasmithCADWorker, "DatasmithCADWorker");
 DEFINE_LOG_CATEGORY(LogDatasmithCADWorker);
 
+#define EXIT_MISSING_CORETECH_MODULE 2
+
 
 void GetParameter(int32 Argc, TCHAR* Argv[], const FString& InParam, FString& OutValue)
 {
@@ -41,6 +43,12 @@ int32 Main(int32 Argc, TCHAR * Argv[])
 {
 	UE_SET_LOG_VERBOSITY(LogDatasmithCADWorker, Verbose);
 
+#ifndef CAD_INTERFACE
+	UE_LOG(LogDatasmithCADWorker, Error, TEXT("Missing CoreTech module. DatasmithCADWorker is not functional."));
+	return EXIT_MISSING_CORETECH_MODULE;
+#endif // CAD_INTERFACE
+
+
 	FString ServerPID, ServerPort, CacheDirectory, CacheVersion, EnginePluginsPath;
 	GetParameter(Argc, Argv, "-ServerPID", ServerPID);
 	GetParameter(Argc, Argv, "-ServerPort", ServerPort);
@@ -69,7 +77,6 @@ int32 Filter(uint32 Code, struct _EXCEPTION_POINTERS *Ep)
 
 INT32_MAIN_INT32_ARGC_TCHAR_ARGV()
 {
-#ifdef CAD_INTERFACE
 	GEngineLoop.PreInit(ArgC, ArgV);
 
 	SetErrorMode(SEM_FAILCRITICALERRORS | SEM_NOGPFAULTERRORBOX);
@@ -83,7 +90,4 @@ INT32_MAIN_INT32_ARGC_TCHAR_ARGV()
 		return EXIT_FAILURE;
 	}
 	return EXIT_SUCCESS;
-#else
-	return EXIT_FAILURE;
-#endif // CAD_INTERFACE
 }
