Datasmith Max Exporter fix: Disable strict conformance since the 3dsmax sdk isn't compliant.

diff --git a/Engine/Source/Programs/Enterprise/Datasmith/DatasmithMaxExporter/DatasmithMax2017.Target.cs b/Engine/Source/Programs/Enterprise/Datasmith/DatasmithMaxExporter/DatasmithMax2017.Target.cs
--- a/Engine/Source/Programs/Enterprise/Datasmith/DatasmithMaxExporter/DatasmithMax2017.Target.cs
+++ b/Engine/Source/Programs/Enterprise/Datasmith/DatasmithMaxExporter/DatasmithMax2017.Target.cs
@@ -17,6 +17,7 @@ public abstract class DatasmithMaxBaseTarget : TargetRules
 		LinkType = TargetLinkType.Monolithic;
 
 		WindowsPlatform.ModuleDefinitionFile = "Programs/Enterprise/Datasmith/DatasmithMaxExporter/DatasmithMaxExporter.def";
+		WindowsPlatform.bStrictConformanceMode = false;
 
 		bBuildDeveloperTools = false;
 		bUseMallocProfiler = false;
