Renamed UE4.ico to UE.ico for the Datasmith 3dsmax exporter

diff --git a/Engine/Source/Programs/Enterprise/Datasmith/DatasmithMaxExporter/Resources/Windows/Resource.rc b/Engine/Source/Programs/Enterprise/Datasmith/DatasmithMaxExporter/Resources/Windows/Resource.rc
--- a/Engine/Source/Programs/Enterprise/Datasmith/DatasmithMaxExporter/Resources/Windows/Resource.rc
+++ b/Engine/Source/Programs/Enterprise/Datasmith/DatasmithMaxExporter/Resources/Windows/Resource.rc
@@ -127,7 +127,7 @@ END
 
 // Icon with lowest ID value placed first to ensure application icon
 // remains consistent on all systems.
-IDI_ICONUE              ICON                    "UE4.ico"
+IDI_ICONUE              ICON                    "UE.ico"
 
 #endif    // English (United States) resources
 /////////////////////////////////////////////////////////////////////////////
