Fix compile error in Datasmith Revit Exporter

diff --git a/Engine/Source/Programs/Enterprise/Datasmith/DatasmithRevitExporter/Private/DatasmithRevitDocumentData.cs b/Engine/Source/Programs/Enterprise/Datasmith/DatasmithRevitExporter/Private/DatasmithRevitDocumentData.cs
--- a/Engine/Source/Programs/Enterprise/Datasmith/DatasmithRevitExporter/Private/DatasmithRevitDocumentData.cs
+++ b/Engine/Source/Programs/Enterprise/Datasmith/DatasmithRevitExporter/Private/DatasmithRevitDocumentData.cs
@@ -450,6 +450,11 @@ namespace DatasmithRevitExporter
 				return Instance;
 			}
 
+			public FDatasmithFacadeMeshElement GetCurrentMeshElement()
+			{
+				return DatasmithMeshElement;
+			}
+
 			public void AddLightActor(
 				Transform InWorldTransform,
 				Asset InLightAsset
@@ -1291,9 +1296,9 @@ namespace DatasmithRevitExporter
 				// For mesh to be reused, it must not be cutoff by a section box.
 				FBaseElementData CurrentInstance = ElementDataStack.Peek().PeekInstance();
 
-				if (CurrentInstance != null && CurrentInstance.bAllowMeshInstancing && CurrentInstance.ElementMesh != null)
+				if (CurrentInstance != null && CurrentInstance.bAllowMeshInstancing && CurrentInstance.DatasmithMeshElement != null)
 				{
-					bIgnore = MeshMap.ContainsKey(CurrentInstance.ElementMesh.GetName());
+					bIgnore = MeshMap.ContainsKey(CurrentInstance.DatasmithMeshElement.GetName());
 				}
 			}
 
