Fixed UE4 references

diff --git a/Engine/Source/Programs/Enterprise/Datasmith/DatasmithMaxExporter/Private/DatasmithMaxExporter.cpp b/Engine/Source/Programs/Enterprise/Datasmith/DatasmithMaxExporter/Private/DatasmithMaxExporter.cpp
--- a/Engine/Source/Programs/Enterprise/Datasmith/DatasmithMaxExporter/Private/DatasmithMaxExporter.cpp
+++ b/Engine/Source/Programs/Enterprise/Datasmith/DatasmithMaxExporter/Private/DatasmithMaxExporter.cpp
@@ -42,7 +42,7 @@ THIRD_PARTY_INCLUDES_START
 #include <locale.h>
 THIRD_PARTY_INCLUDES_END
 
-static FString OriginalLocale( _wsetlocale(LC_NUMERIC, nullptr) ); // Cache LC_NUMERIC locale before initialization of UE4
+static FString OriginalLocale( _wsetlocale(LC_NUMERIC, nullptr) ); // Cache LC_NUMERIC locale before initialization of UnrealEditor
 static FString NewLocale = _wsetlocale(LC_NUMERIC, TEXT("C"));
 
 HINSTANCE FDatasmithMaxExporter::HInstanceMax;
@@ -76,7 +76,7 @@ __declspec( dllexport ) int LibInitialize(void)
 {
 	bool bResult = FDatasmithExporterManager::Initialize();
 
-	// Restore LC_NUMERIC locale after initialization of UE4
+	// Restore LC_NUMERIC locale after initialization of UnrealEditor
 	_wsetlocale(LC_NUMERIC, *OriginalLocale);
 
 	return bResult;
diff --git a/Engine/Source/Programs/Enterprise/Datasmith/DatasmithMaxExporter/Public/DatasmithMaxSceneParser.h b/Engine/Source/Programs/Enterprise/Datasmith/DatasmithMaxExporter/Public/DatasmithMaxSceneParser.h
--- a/Engine/Source/Programs/Enterprise/Datasmith/DatasmithMaxExporter/Public/DatasmithMaxSceneParser.h
+++ b/Engine/Source/Programs/Enterprise/Datasmith/DatasmithMaxExporter/Public/DatasmithMaxSceneParser.h
@@ -74,7 +74,7 @@ class FDatasmithMaxSceneParser
 public:
 	static EMaxLightClass GetLightClass(INode* InNode);
 	static bool CanBeTriMesh(Object* Obj);
-	static const TArray< FString, TInlineAllocator< 4 > > CollisionNodesPrefixes; // List of supported mesh prefixes for UE4 collision
+	static const TArray< FString, TInlineAllocator< 4 > > CollisionNodesPrefixes; // List of supported mesh prefixes for Unreal collision
 
 	FDatasmithMaxSceneParser();
 	~FDatasmithMaxSceneParser();
diff --git a/Engine/Source/Programs/Enterprise/Datasmith/DatasmithMaxExporter/Resources/Windows/Resource.rc b/Engine/Source/Programs/Enterprise/Datasmith/DatasmithMaxExporter/Resources/Windows/Resource.rc
--- a/Engine/Source/Programs/Enterprise/Datasmith/DatasmithMaxExporter/Resources/Windows/Resource.rc
+++ b/Engine/Source/Programs/Enterprise/Datasmith/DatasmithMaxExporter/Resources/Windows/Resource.rc
@@ -169,7 +169,7 @@ BEGIN
             VALUE "LegalCopyright", "Copyright Epic Games, Inc. All Rights Reserved."
             VALUE "OriginalFilename", "DatasmithMaxExporter.dle"
             VALUE "ProductName", "Unreal Engine"
-            VALUE "ProductVersion", "4.21.0-0+UE4"
+            VALUE "ProductVersion", "5.0"
         END
     END
     BLOCK "VarFileInfo"
