Merge from UE4/Main (15724219): PR #7820: Make sure ExtraTexturePaths contain only absolute paths. (Contributed by [REDACTED_PERSON])

diff --git a/Engine/Source/Programs/Enterprise/Datasmith/DatasmithRevitExporter/Private/DatasmithRevitExportContext.cs b/Engine/Source/Programs/Enterprise/Datasmith/DatasmithRevitExporter/Private/DatasmithRevitExportContext.cs
--- a/Engine/Source/Programs/Enterprise/Datasmith/DatasmithRevitExporter/Private/DatasmithRevitExportContext.cs
+++ b/Engine/Source/Programs/Enterprise/Datasmith/DatasmithRevitExporter/Private/DatasmithRevitExportContext.cs
@@ -567,7 +567,7 @@ namespace DatasmithRevitExporter
 
 									foreach (string TexturePath in TexturePaths)
 									{
-										ExtraTexturePaths.Add(TexturePath);
+										ExtraTexturePaths.Add(Path.GetFullPath(Path.Combine(AppDomain.CurrentDomain.BaseDirectory, TexturePath)));
 									}
 
 									break;
