Merge from UE4/Main (15746169): DirectLink - Datasmith Revit 2020 Exporter - Missing parts of multi-story stairs

diff --git a/Engine/Source/Programs/Enterprise/Datasmith/DatasmithRevitExporter/Private/DatasmithRevitDocumentData.cs b/Engine/Source/Programs/Enterprise/Datasmith/DatasmithRevitExporter/Private/DatasmithRevitDocumentData.cs
--- a/Engine/Source/Programs/Enterprise/Datasmith/DatasmithRevitExporter/Private/DatasmithRevitDocumentData.cs
+++ b/Engine/Source/Programs/Enterprise/Datasmith/DatasmithRevitExporter/Private/DatasmithRevitDocumentData.cs
@@ -124,7 +124,14 @@ namespace DatasmithRevitExporter
 				{
 					foreach (FBaseElementData CurrentChild in ChildElements)
 					{
-						CurrentChild.AddToScene(InScene, this, false, (ThisElement == null) && bIsModified);
+						// Stairs get special treatment: elements of stairs (strings, landings etc.) can be duplicated,
+						// meaning that the same element id can exist multiple times under the same parent.
+						bool bIsStairsElement = (ThisElement != null) && (ThisElement.GetType() == typeof(Stairs));
+
+						bool bIsInstance = (ThisElement == null);
+						bool bForceAdd = (bIsInstance && bIsModified) || (bIsStairsElement && bIsModified);
+
+						CurrentChild.AddToScene(InScene, this, false, bForceAdd);
 					}
 				}
 
@@ -1124,7 +1131,6 @@ namespace DatasmithRevitExporter
 				{
 					// Add the element mesh actor to the Datasmith actor hierarchy.
 					ElementDataStack.Peek().AddChildActor(ElementData);
-					ActorMap[ElemId] = ElementData;
 				}
 			}
 		}
