Merge from UE4/Main (15859039): Datasmith Revit: Deleting and Undoing Revit objects syncs correctly but material don't sync correctly

diff --git a/Engine/Source/Programs/Enterprise/Datasmith/DatasmithRevitExporter/Private/DatasmithRevitDirectLink.cs b/Engine/Source/Programs/Enterprise/Datasmith/DatasmithRevitExporter/Private/DatasmithRevitDirectLink.cs
--- a/Engine/Source/Programs/Enterprise/Datasmith/DatasmithRevitExporter/Private/DatasmithRevitDirectLink.cs
+++ b/Engine/Source/Programs/Enterprise/Datasmith/DatasmithRevitExporter/Private/DatasmithRevitDirectLink.cs
@@ -451,6 +451,30 @@ namespace DatasmithRevitExporter
 
 			bool bUpdateOk = DatasmithDirectLink.UpdateScene(DatasmithScene);
 
+			// Sync materials: DatasmithScene.CleanUp() might have deleted some materials that are not referenced by meshes anymore, so we need to update our map.
+			HashSet<string> SceneMaterials = new HashSet<string>();
+			for (int MaterialIndex = 0; MaterialIndex < DatasmithScene.GetMaterialsCount(); ++MaterialIndex)
+			{
+				FDatasmithFacadeBaseMaterial Material = DatasmithScene.GetMaterial(MaterialIndex);
+				SceneMaterials.Add(Material.GetName());
+			}
+
+			List<string> MaterialsToDelete = new List<string>();
+
+			foreach (var MaterialKV in MaterialDataMap)
+			{
+				string MaterialName = MaterialKV.Key;
+
+				if (!SceneMaterials.Contains(MaterialName))
+				{
+					MaterialsToDelete.Add(MaterialName);
+				}
+			}
+			foreach (string MaterialName in MaterialsToDelete)
+			{
+				MaterialDataMap.Remove(MaterialName);
+			}
+
 			SyncCount++;
 
 			// Control metadata export via env var REVIT_DIRECTLINK_WITH_METADATA.
