Merge from UE4/Main (15887641): Datasmith: Revit Section Box not reflected in Direct Link sync

diff --git a/Engine/Source/Programs/Enterprise/Datasmith/DatasmithRevitExporter/Private/DatasmithRevitDirectLink.cs b/Engine/Source/Programs/Enterprise/Datasmith/DatasmithRevitExporter/Private/DatasmithRevitDirectLink.cs
--- a/Engine/Source/Programs/Enterprise/Datasmith/DatasmithRevitExporter/Private/DatasmithRevitDirectLink.cs
+++ b/Engine/Source/Programs/Enterprise/Datasmith/DatasmithRevitExporter/Private/DatasmithRevitDirectLink.cs
@@ -9,6 +9,7 @@ using Autodesk.Revit.ApplicationServices;
 using Autodesk.Revit.DB.Events;
 using System.Threading.Tasks;
 using System.Threading;
+using System.Linq;
 
 namespace DatasmithRevitExporter
 {
@@ -89,6 +90,14 @@ namespace DatasmithRevitExporter
 			}
 		};
 
+		struct SectionBoxInfo
+		{
+			public Element SectionBox;
+			// Store bounding box because removed section boxes lose their bounding box 
+			// and we can't query it anymore
+			public BoundingBoxXYZ SectionBoxBounds;
+		};
+
 		public FDatasmithFacadeScene									DatasmithScene { get; private set; }
 
 		private FCachedDocumentData										RootCache = null;
@@ -98,6 +107,8 @@ namespace DatasmithRevitExporter
 		private HashSet<ElementId>										ExportedLinkedDocuments = new HashSet<ElementId>();
 		private Stack<FCachedDocumentData>								CacheStack = new Stack<FCachedDocumentData>();
 
+		private IList<SectionBoxInfo>									PrevSectionBoxes = new List<SectionBoxInfo>();
+
 		// Sets of elements related to current sync.
 		private HashSet<ElementId>										DeletedElements = new HashSet<ElementId>();
 		
@@ -366,27 +377,80 @@ namespace DatasmithRevitExporter
 
 			// Handle section boxes.
 			FilteredElementCollector Collector = new FilteredElementCollector(RootCache.SourceDocument, RootCache.SourceDocument.ActiveView.Id);
-			IList<Element> SectionBoxes = Collector.OfCategory(BuiltInCategory.OST_SectionBox).ToElements();
+			List<SectionBoxInfo> CurrentSectionBoxes =  new List<SectionBoxInfo>();
 
-			foreach(var SectionBox in SectionBoxes)
+			foreach (Element SectionBox in Collector.OfCategory(BuiltInCategory.OST_SectionBox).ToElements())
 			{
-				if (!RootCache.ModifiedElements.Contains(SectionBox.Id))
+				SectionBoxInfo Info = new SectionBoxInfo();
+				Info.SectionBox = SectionBox;
+				Info.SectionBoxBounds = SectionBox.get_BoundingBox(RootCache.SourceDocument.ActiveView);
+				CurrentSectionBoxes.Add(Info);
+			}
+
+			List<SectionBoxInfo> ModifiedSectionBoxes = new List<SectionBoxInfo>();
+
+			foreach(SectionBoxInfo CurrentSectionBoxInfo in CurrentSectionBoxes)
+			{
+				if (!RootCache.ModifiedElements.Contains(CurrentSectionBoxInfo.SectionBox.Id))
 				{
 					continue;
 				}
 
-				// This section box was modified, need to make all elements it intersects dirty, so they 
-				// can be re-exported.
-				BoundingBoxXYZ BBox = SectionBox.get_BoundingBox(RootCache.SourceDocument.ActiveView);
-				ElementFilter IntersectFilterStart = new BoundingBoxIntersectsFilter(new Outline(BBox.Min, BBox.Max));
-				ICollection<ElementId> IntersectedElements = new FilteredElementCollector(RootCache.SourceDocument).WherePasses(IntersectFilterStart).ToElementIds();
+				ModifiedSectionBoxes.Add(CurrentSectionBoxInfo);
+			}
 
-				foreach (var ElemId in IntersectedElements)
+			// Check for old section boxes that were disabled since last sync.
+			foreach (SectionBoxInfo PrevSectionBoxInfo in PrevSectionBoxes)
+			{
+				bool bSectionBoxWasDisabled = !CurrentSectionBoxes.Any(Info => Info.SectionBox == PrevSectionBoxInfo.SectionBox);
+
+				if (bSectionBoxWasDisabled)
 				{
-					if (!RootCache.ModifiedElements.Contains(ElemId))
-					{
-						RootCache.ModifiedElements.Add(ElemId);
-					}
+					// Section box was removed, need to mark the elemets it intersected as modified
+					ModifiedSectionBoxes.Add(PrevSectionBoxInfo);
+				}
+			}
+
+			// Check all elements that need to be re-exported
+			foreach(var SectionBoxInfo in ModifiedSectionBoxes)
+			{
+				MarkIntersectedElementsAsModified(RootCache, SectionBoxInfo.SectionBox, SectionBoxInfo.SectionBoxBounds);
+			}
+
+			PrevSectionBoxes = CurrentSectionBoxes;
+		}
+
+		void MarkIntersectedElementsAsModified(FCachedDocumentData InData, Element InSectionBox, BoundingBoxXYZ InSectionBoxBounds)
+		{
+			ElementFilter IntersectFilter = new BoundingBoxIntersectsFilter(new Outline(InSectionBoxBounds.Min, InSectionBoxBounds.Max));
+			ICollection<ElementId> IntersectedElements = new FilteredElementCollector(InData.SourceDocument).WherePasses(IntersectFilter).ToElementIds();
+
+			ElementFilter InsideFilter = new BoundingBoxIsInsideFilter(new Outline(InSectionBoxBounds.Min, InSectionBoxBounds.Max));
+			ICollection<ElementId> InsideElements = new FilteredElementCollector(InData.SourceDocument).WherePasses(InsideFilter).ToElementIds();
+
+			// Elements that are fully inside the section box should not be marked modified to save export time
+			foreach (ElementId InsideElement in InsideElements)
+			{
+				IntersectedElements.Remove(InsideElement);
+			}
+			
+			foreach (var ElemId in IntersectedElements)
+			{
+				if (!InData.ModifiedElements.Contains(ElemId))
+				{
+					InData.ModifiedElements.Add(ElemId);
+				}
+			}
+
+			// Run the linked documents
+			foreach (var LinkedDoc in InData.LinkedDocumentsCache)
+			{
+				MarkIntersectedElementsAsModified(LinkedDoc.Value, InSectionBox, InSectionBoxBounds);
+
+				if (LinkedDoc.Value.ModifiedElements.Count > 0)
+				{
+					InData.ModifiedElements.Add(LinkedDoc.Key);
+					ModifiedLinkedDocuments.Add(LinkedDoc.Value.SourceDocument);
 				}
 			}
 		}
diff --git a/Engine/Source/Programs/Enterprise/Datasmith/DatasmithRevitExporter/Private/DatasmithRevitDocumentData.cs b/Engine/Source/Programs/Enterprise/Datasmith/DatasmithRevitExporter/Private/DatasmithRevitDocumentData.cs
--- a/Engine/Source/Programs/Enterprise/Datasmith/DatasmithRevitExporter/Private/DatasmithRevitDocumentData.cs
+++ b/Engine/Source/Programs/Enterprise/Datasmith/DatasmithRevitExporter/Private/DatasmithRevitDocumentData.cs
@@ -190,7 +190,7 @@ namespace DatasmithRevitExporter
 			public void UpdateMeshName()
 			{
 				FDatasmithFacadeActorMesh MeshActor = ElementActor as FDatasmithFacadeActorMesh;
-				MeshActor.SetMesh(DatasmithMeshElement.GetName());
+				MeshActor?.SetMesh(DatasmithMeshElement.GetName());
 				bOptimizeHierarchy = false;
 			}
 		}
@@ -1037,7 +1037,9 @@ namespace DatasmithRevitExporter
 				{
 					ElementData = (FElementData)DirectLink.GetCachedElement(InElement);
 
-					if (DirectLink.IsElementModified(InElement))
+					bool bIsLinkedDocument = InElement.GetType() == typeof(RevitLinkInstance);
+
+					if (bIsLinkedDocument || DirectLink.IsElementModified(InElement))
 					{
 						FDatasmithFacadeActor ExistingActor = ElementData.ElementActor;
 
