Merge from UE4/Main (15433810): Datasmith Revit: [Regression] Revit families exported with too many actors / components compared to 4.25

diff --git a/Engine/Source/Programs/Enterprise/Datasmith/DatasmithRevitExporter/Private/DatasmithRevitDocumentData.cs b/Engine/Source/Programs/Enterprise/Datasmith/DatasmithRevitExporter/Private/DatasmithRevitDocumentData.cs
--- a/Engine/Source/Programs/Enterprise/Datasmith/DatasmithRevitExporter/Private/DatasmithRevitDocumentData.cs
+++ b/Engine/Source/Programs/Enterprise/Datasmith/DatasmithRevitExporter/Private/DatasmithRevitDocumentData.cs
@@ -1353,6 +1353,56 @@ namespace DatasmithRevitExporter
 			return ElementDataStack.Count > 0 ? ElementDataStack.Peek().CurrentElement : null;
 		}
 
+		private FBaseElementData OptimizeElementRecursive(FBaseElementData InElementData, FDatasmithFacadeScene InDatasmithScene)
+		{
+			List<FDatasmithFacadeActor> RemoveChildren = new List<FDatasmithFacadeActor>();
+			List<FDatasmithFacadeActor> AddChildren = new List<FDatasmithFacadeActor>();
+
+			for (int ChildIndex = 0; ChildIndex < InElementData.ChildElements.Count; ChildIndex++)
+			{
+				FBaseElementData ChildElement = InElementData.ChildElements[ChildIndex];
+
+				// Optimize the Datasmith child actor.
+				FBaseElementData ResultElement = OptimizeElementRecursive(ChildElement, InDatasmithScene);
+
+				if (ChildElement != ResultElement)
+				{
+					RemoveChildren.Add(ChildElement.ElementActor);
+
+					if (ResultElement != null)
+					{
+						AddChildren.Add(ResultElement.ElementActor);
+					}
+				}
+			}
+
+			{
+				MeshMaterialsMap[MeshName] = new Dictionary<string, int>();
+			}
+
+			Dictionary<string, int> MaterialsMap = MeshMaterialsMap[MeshName];
+			if (!MaterialsMap.ContainsKey(InMaterialName))
+			{
+				int NewMaterialIndex = MaterialsMap.Count;
+				MaterialsMap[InMaterialName] = NewMaterialIndex;
+				// Add the current Datasmith master material name to the dictionary of material names utilized by the Datasmith mesh being processed.
+				InMesh.SetMaterial(InMaterialName, NewMaterialIndex);
+			}
+
+			return MaterialsMap[InMaterialName];
+		}
+
+		public int GetCurrentMaterialIndex()
+		{
+			FDatasmithFacadeMeshElement MeshElement = GetCurrentMeshElement();
+			return MeshElement != null ? GetMeshMaterialIndex(MeshElement, CurrentMaterialName) : -1;
+		}
+
+		public FBaseElementData GetCurrentActor()
+		{
+			return ElementDataStack.Peek().GetCurrentActor();
+		}
+
 		private FBaseElementData OptimizeElementRecursive(FBaseElementData InElementData, FDatasmithFacadeScene InDatasmithScene)
 		{
 			List<FDatasmithFacadeActor> RemoveChildren = new List<FDatasmithFacadeActor>();
