Merge from UE4/Main (15694933): Missing instances of linked Revit scenes with Datasmith exporter and Direct Link

diff --git a/Engine/Source/Programs/Enterprise/Datasmith/DatasmithRevitExporter/Private/DatasmithRevitDocumentData.cs b/Engine/Source/Programs/Enterprise/Datasmith/DatasmithRevitExporter/Private/DatasmithRevitDocumentData.cs
--- a/Engine/Source/Programs/Enterprise/Datasmith/DatasmithRevitExporter/Private/DatasmithRevitDocumentData.cs
+++ b/Engine/Source/Programs/Enterprise/Datasmith/DatasmithRevitExporter/Private/DatasmithRevitDocumentData.cs
@@ -1377,30 +1377,19 @@ namespace DatasmithRevitExporter
 			}
 
 			{
-				MeshMaterialsMap[MeshName] = new Dictionary<string, int>();
 			}
 
-			Dictionary<string, int> MaterialsMap = MeshMaterialsMap[MeshName];
-			if (!MaterialsMap.ContainsKey(InMaterialName))
-			{
-				int NewMaterialIndex = MaterialsMap.Count;
-				MaterialsMap[InMaterialName] = NewMaterialIndex;
-				// Add the current Datasmith master material name to the dictionary of material names utilized by the Datasmith mesh being processed.
-				InMesh.SetMaterial(InMaterialName, NewMaterialIndex);
-			}
-
-			return MaterialsMap[InMaterialName];
+			return CurrentElement.MeshMaterialsMap[CurrentMaterialName];
 		}
 
-		public int GetCurrentMaterialIndex()
+		public FBaseElementData GetCurrentActor()
 		{
-			FDatasmithFacadeMeshElement MeshElement = GetCurrentMeshElement();
-			return MeshElement != null ? GetMeshMaterialIndex(MeshElement, CurrentMaterialName) : -1;
+			return ElementDataStack.Peek().GetCurrentActor();
 		}
 
-		public FBaseElementData GetCurrentActor()
+		public Element GetCurrentElement()
 		{
-			return ElementDataStack.Peek().GetCurrentActor();
+			return ElementDataStack.Count > 0 ? ElementDataStack.Peek().CurrentElement : null;
 		}
 
 		private FBaseElementData OptimizeElementRecursive(FBaseElementData InElementData, FDatasmithFacadeScene InDatasmithScene)
