Fix Compilation of Datasmith Revit Exporter (Type 'FDocumentData' already defines a member called 'GetCurrentActor')

diff --git a/Engine/Source/Programs/Enterprise/Datasmith/DatasmithRevitExporter/Private/DatasmithRevitDocumentData.cs b/Engine/Source/Programs/Enterprise/Datasmith/DatasmithRevitExporter/Private/DatasmithRevitDocumentData.cs
--- a/Engine/Source/Programs/Enterprise/Datasmith/DatasmithRevitExporter/Private/DatasmithRevitDocumentData.cs
+++ b/Engine/Source/Programs/Enterprise/Datasmith/DatasmithRevitExporter/Private/DatasmithRevitDocumentData.cs
@@ -1353,45 +1353,6 @@ namespace DatasmithRevitExporter
 			return ElementDataStack.Count > 0 ? ElementDataStack.Peek().CurrentElement : null;
 		}
 
-		private FBaseElementData OptimizeElementRecursive(FBaseElementData InElementData, FDatasmithFacadeScene InDatasmithScene)
-		{
-			List<FDatasmithFacadeActor> RemoveChildren = new List<FDatasmithFacadeActor>();
-			List<FDatasmithFacadeActor> AddChildren = new List<FDatasmithFacadeActor>();
-
-			for (int ChildIndex = 0; ChildIndex < InElementData.ChildElements.Count; ChildIndex++)
-			{
-				FBaseElementData ChildElement = InElementData.ChildElements[ChildIndex];
-
-				// Optimize the Datasmith child actor.
-				FBaseElementData ResultElement = OptimizeElementRecursive(ChildElement, InDatasmithScene);
-
-				if (ChildElement != ResultElement)
-				{
-					RemoveChildren.Add(ChildElement.ElementActor);
-
-					if (ResultElement != null)
-					{
-						AddChildren.Add(ResultElement.ElementActor);
-					}
-				}
-			}
-
-			{
-			}
-
-			return CurrentElement.MeshMaterialsMap[CurrentMaterialName];
-		}
-
-		public FBaseElementData GetCurrentActor()
-		{
-			return ElementDataStack.Peek().GetCurrentActor();
-		}
-
-		public Element GetCurrentElement()
-		{
-			return ElementDataStack.Count > 0 ? ElementDataStack.Peek().CurrentElement : null;
-		}
-
 		private FBaseElementData OptimizeElementRecursive(FBaseElementData InElementData, FDatasmithFacadeScene InDatasmithScene)
 		{
 			List<FDatasmithFacadeActor> RemoveChildren = new List<FDatasmithFacadeActor>();
