Fixed warning from SketchUp 2019 SDK

preflight 6079e0e6bd324b0001bdcd94

diff --git a/Engine/Source/Programs/Enterprise/Datasmith/DatasmithSketchUpRubyExporter/Private/DatasmithSketchUpExporter.cpp b/Engine/Source/Programs/Enterprise/Datasmith/DatasmithSketchUpRubyExporter/Private/DatasmithSketchUpExporter.cpp
--- a/Engine/Source/Programs/Enterprise/Datasmith/DatasmithSketchUpRubyExporter/Private/DatasmithSketchUpExporter.cpp
+++ b/Engine/Source/Programs/Enterprise/Datasmith/DatasmithSketchUpRubyExporter/Private/DatasmithSketchUpExporter.cpp
@@ -35,6 +35,8 @@
 #pragma warning(disable: 4005)
 // disable(SU2021): 'reinterpret_cast': unsafe conversion from 'ruby::backward::cxxanyargs::void_type (__cdecl *)' to 'rb_gvar_setter_t (__cdecl *)'	
 #pragma warning(disable: 4191)
+// disable(SU2019 & SU2020): 'register' is no longer a supported storage class	
+#pragma warning(disable: 5033)
 #include <ruby.h>
 #pragma warning(pop)
 
