Disable debug output when generating project files.

diff --git a/Engine/Source/Programs/Enterprise/Datasmith/DatasmithSketchUpRubyExporter/DatasmithSketchUpRuby2020.Build.cs b/Engine/Source/Programs/Enterprise/Datasmith/DatasmithSketchUpRubyExporter/DatasmithSketchUpRuby2020.Build.cs
--- a/Engine/Source/Programs/Enterprise/Datasmith/DatasmithSketchUpRubyExporter/DatasmithSketchUpRuby2020.Build.cs
+++ b/Engine/Source/Programs/Enterprise/Datasmith/DatasmithSketchUpRubyExporter/DatasmithSketchUpRuby2020.Build.cs
@@ -51,10 +51,9 @@ namespace UnrealBuildTool.Rules
 					PublicAdditionalLibraries.Add(Path.Combine(SketchUpSDKLocation, "samples", "common", "ThirdParty", "ruby", "lib", "win", "x64", GetRubyLibName()));
 				}
 
-				if (!Directory.Exists(SketchUpSDKLocation))
+				if (!Directory.Exists(SketchUpSDKLocation) && !Target.bGenerateProjectFiles)
 				{
-					// XXX: remove
-					System.Console.WriteLine("SketchUp SDK directory doesn't exist: '" + SketchUpSDKLocation + "'");
+					throw new System.Exception("SketchUp SDK directory doesn't exist: '" + SketchUpSDKLocation + "'");
 				}
 			}
 		}
