Update UE4/Main version of file to match RES

diff --git a/Engine/Source/Programs/Enterprise/Datasmith/DatasmithSketchUpRubyExporter/DatasmithSketchUpRuby2020.Build.cs b/Engine/Source/Programs/Enterprise/Datasmith/DatasmithSketchUpRubyExporter/DatasmithSketchUpRuby2020.Build.cs
--- a/Engine/Source/Programs/Enterprise/Datasmith/DatasmithSketchUpRubyExporter/DatasmithSketchUpRuby2020.Build.cs
+++ b/Engine/Source/Programs/Enterprise/Datasmith/DatasmithSketchUpRubyExporter/DatasmithSketchUpRuby2020.Build.cs
@@ -54,7 +54,7 @@ namespace UnrealBuildTool.Rules
 
 				if (!Directory.Exists(SketchUpSDKLocation) && !Target.bGenerateProjectFiles)
 				{
-					throw new System.Exception("SketchUp SDK directory doesn't exist: '" + SketchUpSDKLocation + "'");
+					Log.TraceWarningOnce("Unable to find SketchUp SDK directory. SketchUp plugins will not compile");
 				}
 			}
 		}
