Fix capitalization of PlatformFileManager.h in files that are not currently compiled by linux which reports it

diff --git a/Engine/Plugins/Media/PixelStreaming/Source/PixelStreaming/Private/PCH.h b/Engine/Plugins/Media/PixelStreaming/Source/PixelStreaming/Private/PCH.h
--- a/Engine/Plugins/Media/PixelStreaming/Source/PixelStreaming/Private/PCH.h
+++ b/Engine/Plugins/Media/PixelStreaming/Source/PixelStreaming/Private/PCH.h
@@ -9,7 +9,7 @@
 #include "HAL/PlatformTime.h"
 #include "HAL/ThreadSafeBool.h"
 #include "HAL/ThreadSafeCounter.h"
-#include "HAL/PlatformFilemanager.h"
+#include "HAL/PlatformFileManager.h"
 #include "HAL/IConsoleManager.h"
 #include "HAL/CriticalSection.h"
 
diff --git a/Engine/Plugins/XGEController/Source/Private/XGEControllerModule.cpp b/Engine/Plugins/XGEController/Source/Private/XGEControllerModule.cpp
--- a/Engine/Plugins/XGEController/Source/Private/XGEControllerModule.cpp
+++ b/Engine/Plugins/XGEController/Source/Private/XGEControllerModule.cpp
@@ -3,7 +3,7 @@
 #include "Features/IModularFeatures.h"
 #include "HAL/FileManager.h"
 #include "HAL/PlatformNamedPipe.h"
-#include "HAL/PlatformFilemanager.h"
+#include "HAL/PlatformFileManager.h"
 #include "HAL/PlatformFile.h"
 #include "Misc/CommandLine.h"
 #include "Misc/ScopeLock.h"
diff --git a/Engine/Source/Programs/Enterprise/Datasmith/DatasmithSketchUpRubyExporter/Private/DatasmithSketchUpExporter.cpp b/Engine/Source/Programs/Enterprise/Datasmith/DatasmithSketchUpRubyExporter/Private/DatasmithSketchUpExporter.cpp
--- a/Engine/Source/Programs/Enterprise/Datasmith/DatasmithSketchUpRubyExporter/Private/DatasmithSketchUpExporter.cpp
+++ b/Engine/Source/Programs/Enterprise/Datasmith/DatasmithSketchUpRubyExporter/Private/DatasmithSketchUpExporter.cpp
@@ -50,7 +50,7 @@
 #include "Misc/Paths.h"
 
 #include "HAL/FileManager.h"
-#include "HAL/PlatformFilemanager.h"
+#include "HAL/PlatformFileManager.h"
 #include "DatasmithSceneXmlWriter.h"
 
 #include "DatasmithSceneFactory.h"
