Adds processing time control in CADWorker
The processing max time is based on the size of the file to process according to its format.
+ increase MaxRestartAllowed
+ add Neu format

diff --git a/Engine/Plugins/Enterprise/DatasmithCADImporter/Source/CADTools/Public/CADOptions.h b/Engine/Plugins/Enterprise/DatasmithCADImporter/Source/CADTools/Public/CADOptions.h
--- a/Engine/Plugins/Enterprise/DatasmithCADImporter/Source/CADTools/Public/CADOptions.h
+++ b/Engine/Plugins/Enterprise/DatasmithCADImporter/Source/CADTools/Public/CADOptions.h
@@ -44,6 +44,7 @@ namespace CADLibrary
 		bool bScaleUVMap = true;
 		bool bEnableCacheUsage = true;
 		bool bEnableKernelIOTessellation = true;
+		bool bEnableTimeControl = true;
 
 		uint32 GetHash() const
 		{
diff --git a/Engine/Plugins/Enterprise/DatasmithCADImporter/Source/DatasmithCADTranslator/Private/DatasmithCADTranslator.cpp b/Engine/Plugins/Enterprise/DatasmithCADImporter/Source/DatasmithCADTranslator/Private/DatasmithCADTranslator.cpp
--- a/Engine/Plugins/Enterprise/DatasmithCADImporter/Source/DatasmithCADTranslator/Private/DatasmithCADTranslator.cpp
+++ b/Engine/Plugins/Enterprise/DatasmithCADImporter/Source/DatasmithCADTranslator/Private/DatasmithCADTranslator.cpp
@@ -114,7 +114,7 @@ bool FDatasmithCADTranslator::LoadScene(TSharedRef<IDatasmithScene> DatasmithSce
 	}
 	else if (FileDescription.Extension == TEXT("sldprt") || FileDescription.Extension == TEXT("sldasm") || // Solidworks
 		FileDescription.Extension == TEXT("iam") || FileDescription.Extension == TEXT("ipt") || // Inventor
-		FileDescription.Extension.StartsWith(TEXT("asm")) || FileDescription.Extension.StartsWith(TEXT("creo")) || FileDescription.Extension.StartsWith(TEXT("prt")) // Creo
+		FileDescription.Extension.StartsWith(TEXT("asm")) || FileDescription.Extension.StartsWith(TEXT("creo")) || FileDescription.Extension.StartsWith(TEXT("prt")) || FileDescription.Extension.StartsWith(TEXT("neu")) // Creo
 		)
 	{
 		ImportParameters.ModelCoordSys = FDatasmithUtils::EModelCoordSystem::YUp_RightHanded;
diff --git a/Engine/Plugins/Enterprise/DatasmithCADImporter/Source/DatasmithDispatcher/Private/DatasmithDispatcher.cpp b/Engine/Plugins/Enterprise/DatasmithCADImporter/Source/DatasmithDispatcher/Private/DatasmithDispatcher.cpp
--- a/Engine/Plugins/Enterprise/DatasmithCADImporter/Source/DatasmithDispatcher/Private/DatasmithDispatcher.cpp
+++ b/Engine/Plugins/Enterprise/DatasmithCADImporter/Source/DatasmithDispatcher/Private/DatasmithDispatcher.cpp
@@ -20,6 +20,12 @@ static TAutoConsoleVariable<int32> CVarStaticCADTranslatorEnableCADCache(
 	ECVF_Default);
 #endif
 
+static TAutoConsoleVariable<int32> CVarStaticCADTranslatorEnableTimeControl(
+	TEXT("r.CADTranslator.EnableTimeControl"),
+	1,
+	TEXT("Enable the timer that kill the worker if the import time is unusually long. With this time control, the load of the corrupted file is canceled but the rest of the scene is imported.\n"),
+	ECVF_Default);
+
 namespace DatasmithDispatcher
 {
 
@@ -33,6 +39,7 @@ FDatasmithDispatcher::FDatasmithDispatcher(const CADLibrary::FImportParameters&
 	, NumberOfWorkers(InNumberOfWorkers)
 	, NextWorkerId(0)
 {
+	ImportParameters.bEnableTimeControl = (CVarStaticCADTranslatorEnableTimeControl.GetValueOnAnyThread() != 0);
 #ifndef DATASMITH_CAD_IGNORE_CACHE
 	ImportParameters.bEnableCacheUsage = (CVarStaticCADTranslatorEnableCADCache.GetValueOnAnyThread() != 0);
 
@@ -44,6 +51,8 @@ FDatasmithDispatcher::FDatasmithDispatcher(const CADLibrary::FImportParameters&
 		IFileManager::Get().MakeDirectory(*FPaths::Combine(ProcessCacheFolder, TEXT("mesh")), true);
 		IFileManager::Get().MakeDirectory(*FPaths::Combine(ProcessCacheFolder, TEXT("body")), true);
 	}
+#else
+	ImportParameters.bEnableCacheUsage = false;
 #endif
 }
 
diff --git a/Engine/Plugins/Enterprise/DatasmithCADImporter/Source/DatasmithDispatcher/Public/DatasmithDispatcherConfig.h b/Engine/Plugins/Enterprise/DatasmithCADImporter/Source/DatasmithDispatcher/Public/DatasmithDispatcherConfig.h
--- a/Engine/Plugins/Enterprise/DatasmithCADImporter/Source/DatasmithDispatcher/Public/DatasmithDispatcherConfig.h
+++ b/Engine/Plugins/Enterprise/DatasmithCADImporter/Source/DatasmithDispatcher/Public/DatasmithDispatcherConfig.h
@@ -32,6 +32,6 @@ namespace DatasmithDispatcher
 		constexpr double ProcessingLoopDelay = 0.5;
 
 		// Number of worker instance that can be restarted when a crash occurs in en external process
-		constexpr int32 MaxRestartAllowed = 5;
+		constexpr int32 MaxRestartAllowed = 1000;
 	};
 } // ns DatasmithDispatcher
diff --git a/Engine/Source/Programs/Enterprise/Datasmith/DatasmithCADWorker/Private/DatasmithCADWorkerImpl.cpp b/Engine/Source/Programs/Enterprise/Datasmith/DatasmithCADWorker/Private/DatasmithCADWorkerImpl.cpp
--- a/Engine/Source/Programs/Enterprise/Datasmith/DatasmithCADWorker/Private/DatasmithCADWorkerImpl.cpp
+++ b/Engine/Source/Programs/Enterprise/Datasmith/DatasmithCADWorker/Private/DatasmithCADWorkerImpl.cpp
@@ -7,6 +7,9 @@
 #include "DatasmithDispatcherConfig.h"
 
 #include "HAL/FileManager.h"
+#include "HAL/PlatformProcess.h"
+#include "HAL/Thread.h"
+#include "Misc/CoreDelegates.h"
 #include "Misc/FileHelper.h"
 #include "Misc/Paths.h"
 #include "SocketSubsystem.h"
@@ -117,6 +120,33 @@ void FDatasmithCADWorkerImpl::ProcessCommand(const FImportParametersCommand& Imp
 	ImportParameters = ImportParametersCommand.ImportParameters;
 }
 
+uint64 DefineMaximumAllowedDuration(const CADLibrary::FFileDescription& FileDescription)
+{
+	FFileStatData FileStatData = IFileManager::Get().GetStatData(*FileDescription.Path);
+	double MaxTimePerMb = 5e-6;
+	double SafetyCoeficient = 5;
+
+	if (FileDescription.Extension.StartsWith(TEXT("sld"))) // SW
+	{
+		MaxTimePerMb = 1e-5;
+	}
+	else if (FileDescription.Extension == TEXT("3dxml") || FileDescription.Extension == TEXT("3drep")) // Catia V5 CGR
+	{
+		MaxTimePerMb = 1e-5;
+	}
+	else if (FileDescription.Extension == TEXT("cgr"))
+	{
+		MaxTimePerMb = 5e-7;
+	}
+	else if (FileDescription.Extension.StartsWith(TEXT("ig"))) // Iges
+	{
+		MaxTimePerMb = 1e-6;
+	}
+
+	uint64 MaximumDuration = ((double)FileStatData.FileSize) * MaxTimePerMb * SafetyCoeficient;
+	return FMath::Max(MaximumDuration, (uint64)30);
+}
+
 void FDatasmithCADWorkerImpl::ProcessCommand(const FRunTaskCommand& RunTaskCommand)
 {
 	const CADLibrary::FFileDescription& FileToProcess = RunTaskCommand.JobFileDescription;
@@ -124,9 +154,16 @@ void FDatasmithCADWorkerImpl::ProcessCommand(const FRunTaskCommand& RunTaskComma
 
 	FCompletedTaskCommand CompletedTask;
 
+	bProcessIsRunning = true;
+	int64 MaxDuration = DefineMaximumAllowedDuration(FileToProcess);
+	FThread TimeCheckerThread = FThread(TEXT("TimeCheckerThread"), [&]() { CheckDuration(FileToProcess, MaxDuration); });
+
 	CADLibrary::FCoreTechFileParser FileParser(ImportParameters, EnginePluginsPath, CachePath);
 	CADLibrary::ECoreTechParsingResult ProcessResult = FileParser.ProcessFile(FileToProcess);
 
+	bProcessIsRunning = false;
+	TimeCheckerThread.Join();
+	
 	CompletedTask.ProcessResult = ProcessResult;
 
 	if (CompletedTask.ProcessResult == ETaskState::ProcessOk)
@@ -142,3 +179,25 @@ void FDatasmithCADWorkerImpl::ProcessCommand(const FRunTaskCommand& RunTaskComma
 	UE_LOG(LogDatasmithCADWorker, Verbose, TEXT("End of Process %s %s saved in %s"), *FileToProcess.Name, *FileToProcess.Configuration, *CompletedTask.GeomFileName);
 }
 
+void FDatasmithCADWorkerImpl::CheckDuration(const CADLibrary::FFileDescription& FileToProcess, const int64 MaxDuration)
+{
+	if (!ImportParameters.bEnableTimeControl)
+	{
+		return;
+	}
+
+	const uint64 StartTime = FPlatformTime::Cycles64();
+	const uint64 MaxCycles = MaxDuration / FPlatformTime::GetSecondsPerCycle64() + StartTime;
+
+	while (bProcessIsRunning)
+	{
+		FPlatformProcess::Sleep(0.1f);
+		if (FPlatformTime::Cycles64() > MaxCycles)
+		{
+			UE_LOG(LogDatasmithCADWorker, Verbose, TEXT("Time exceeded to process %s %s. The maximum allowed duration is %ld s"), *FileToProcess.Name, *FileToProcess.Configuration, MaxDuration);
+			FPlatformMisc::RequestExit(true);
+		}
+	}
+	double Duration = (FPlatformTime::Cycles64() - StartTime) * FPlatformTime::GetSecondsPerCycle64();
+	UE_LOG(LogDatasmithCADWorker, Verbose, TEXT("    Processing Time: %f s"), Duration);
+}
diff --git a/Engine/Source/Programs/Enterprise/Datasmith/DatasmithCADWorker/Private/DatasmithCADWorkerImpl.h b/Engine/Source/Programs/Enterprise/Datasmith/DatasmithCADWorker/Private/DatasmithCADWorkerImpl.h
--- a/Engine/Source/Programs/Enterprise/Datasmith/DatasmithCADWorker/Private/DatasmithCADWorkerImpl.h
+++ b/Engine/Source/Programs/Enterprise/Datasmith/DatasmithCADWorker/Private/DatasmithCADWorkerImpl.h
@@ -27,6 +27,11 @@ private:
 	void ProcessCommand(const DatasmithDispatcher::FImportParametersCommand& BackPingCommand);
 	void ProcessCommand(const DatasmithDispatcher::FRunTaskCommand& TerminateCommand);
 
+	/**
+	 * @param MaxDuration is defined in seconds.
+	 */
+	void CheckDuration(const CADLibrary::FFileDescription& FileToProcess, const int64 MaxDuration); 
+
 private:
 	DatasmithDispatcher::FNetworkClientNode NetworkInterface;
 	DatasmithDispatcher::FCommandQueue CommandIO;
@@ -37,4 +42,5 @@ private:
 	FString CachePath;
 	CADLibrary::FImportParameters ImportParameters;
 	uint64 PingStartCycle;
+	TAtomic<bool> bProcessIsRunning;
 };
