Fixed Horde issue issue #75908

diff --git a/Engine/Source/Programs/Enterprise/Datasmith/DatasmithARCHICADExporter/Private/GeometryUtil.h b/Engine/Source/Programs/Enterprise/Datasmith/DatasmithARCHICADExporter/Private/GeometryUtil.h
--- a/Engine/Source/Programs/Enterprise/Datasmith/DatasmithARCHICADExporter/Private/GeometryUtil.h
+++ b/Engine/Source/Programs/Enterprise/Datasmith/DatasmithARCHICADExporter/Private/GeometryUtil.h
@@ -9,7 +9,7 @@
 #include "Vertex.hpp"
 
 struct FQuat;
-DECLARE_LWC_TYPE(Vector);
+DECLARE_LWC_TYPE(Vector, 3);
 
 BEGIN_NAMESPACE_UE_AC
 
diff --git a/Engine/Source/Programs/Enterprise/Datasmith/DatasmithARCHICADExporter/Private/Utils/WarningsDisabler.h b/Engine/Source/Programs/Enterprise/Datasmith/DatasmithARCHICADExporter/Private/Utils/WarningsDisabler.h
--- a/Engine/Source/Programs/Enterprise/Datasmith/DatasmithARCHICADExporter/Private/Utils/WarningsDisabler.h
+++ b/Engine/Source/Programs/Enterprise/Datasmith/DatasmithARCHICADExporter/Private/Utils/WarningsDisabler.h
@@ -10,6 +10,7 @@
 	__pragma(warning(push)) \
 	__pragma(warning(disable: 4996)) \
 	__pragma(warning(disable: 5038)) \
+	__pragma(warning(disable: 5040)) \
 	__pragma(warning(disable: 5208)) \
 	__pragma(warning(disable: 26495)) \
 	__pragma(warning(disable: 26451)) \
