Fix calling into Win64 only c# functions when running on an non-win64 platform

diff --git a/Engine/Source/Programs/Enterprise/Datasmith/DatasmithSolidworksExporter/DatasmithSolidworks2021.Target.cs b/Engine/Source/Programs/Enterprise/Datasmith/DatasmithSolidworksExporter/DatasmithSolidworks2021.Target.cs
--- a/Engine/Source/Programs/Enterprise/Datasmith/DatasmithSolidworksExporter/DatasmithSolidworks2021.Target.cs
+++ b/Engine/Source/Programs/Enterprise/Datasmith/DatasmithSolidworksExporter/DatasmithSolidworks2021.Target.cs
@@ -63,28 +63,31 @@ public abstract class DatasmithSolidworksBaseTarget : TargetRules
 		string ProjectName = "DatasmithSolidworks";
 		LaunchModuleName = ProjectName + GetVersionShort();
 
-		string SolidworksSDKPath = CheckSolidworksInstalledSub(GetVersionLong());
-		if (!Directory.Exists(SolidworksSDKPath))
+		if (Target.Platform == UnrealTargetPlatform.Win64)
 		{
-			// Try with build machine setup
-			string SDKRootEnvVar = System.Environment.GetEnvironmentVariable("UE_SDKS_ROOT");
-			if (SDKRootEnvVar != null && SDKRootEnvVar != "")
+			string SolidworksSDKPath = CheckSolidworksInstalledSub(GetVersionLong());
+			if (!Directory.Exists(SolidworksSDKPath))
 			{
-				SolidworksSDKPath = Path.Combine(SDKRootEnvVar, "HostWin64", "Win64", "Solidworks", GetVersionShort());
+				// Try with build machine setup
+				string SDKRootEnvVar = System.Environment.GetEnvironmentVariable("UE_SDKS_ROOT");
+				if (SDKRootEnvVar != null && SDKRootEnvVar != "")
+				{
+					SolidworksSDKPath = Path.Combine(SDKRootEnvVar, "HostWin64", "Win64", "Solidworks", GetVersionShort());
+				}
 			}
-		}
 
-		// Define post-build step
-		// Since the Datasmith Solidworks Exporter is a C# project, build in batch the release configuration of the Visual Studio C# project file.
-		string Config = "Release";
-		string SolidworksExporterPath = @"$(EngineDir)\Source\Programs\Enterprise\Datasmith\DatasmithSolidworksExporter";
-		string ProjectFile = Path.Combine(SolidworksExporterPath, ProjectName, ProjectName+".csproj");
+			// Define post-build step
+			// Since the Datasmith Solidworks Exporter is a C# project, build in batch the release configuration of the Visual Studio C# project file.
+			string Config = "Release";
+			string SolidworksExporterPath = @"$(EngineDir)\Source\Programs\Enterprise\Datasmith\DatasmithSolidworksExporter";
+			string ProjectFile = Path.Combine(SolidworksExporterPath, ProjectName, ProjectName+".csproj");
 
-		string BuildCommand = string.Format(@"$(EngineDir)\Build\BatchFiles\MSBuild.bat /t:Build /p:Configuration={1} /p:EngineDir=""$(EngineDir)"" /p:ExternalAssemblies=""{2}"" ""{0}""",
-			ProjectFile, Config, SolidworksSDKPath);
+			string BuildCommand = string.Format(@"$(EngineDir)\Build\BatchFiles\MSBuild.bat /t:Build /p:Configuration={1} /p:EngineDir=""$(EngineDir)"" /p:ExternalAssemblies=""{2}"" ""{0}""",
+				ProjectFile, Config, SolidworksSDKPath);
 
-		PostBuildSteps.Add(string.Format(@"echo BuildCommand: {0}", BuildCommand));
-		PostBuildSteps.Add(BuildCommand);
+			PostBuildSteps.Add(string.Format(@"echo BuildCommand: {0}", BuildCommand));
+			PostBuildSteps.Add(BuildCommand);
+		}
 	}
 
 	public abstract string GetVersionLong();
