Datasmith Solidworks: Warning: Check CSProj Target Versions - DatasmithSolidworks.csproj and PreprocessWxs.csproj were not 4.5/4.6.2!

diff --git a/Engine/Source/Programs/Enterprise/Datasmith/DatasmithSolidworksExporter/DatasmithSolidworks/DatasmithSolidworks.csproj b/Engine/Source/Programs/Enterprise/Datasmith/DatasmithSolidworksExporter/DatasmithSolidworks/DatasmithSolidworks.csproj
--- a/Engine/Source/Programs/Enterprise/Datasmith/DatasmithSolidworksExporter/DatasmithSolidworks/DatasmithSolidworks.csproj
+++ b/Engine/Source/Programs/Enterprise/Datasmith/DatasmithSolidworksExporter/DatasmithSolidworks/DatasmithSolidworks.csproj
@@ -26,7 +26,7 @@
     </UpgradeBackupLocation>
     <StartProgram>$(ProgramW6432)\SOLIDWORKS Corp\SOLIDWORKS\SldWorks.exe</StartProgram>
     <StartAction>Program</StartAction>
-    <TargetFrameworkVersion>v4.0</TargetFrameworkVersion>
+    <TargetFrameworkVersion>v4.5</TargetFrameworkVersion>
     <OldToolsVersion>2.0</OldToolsVersion>
     <TargetFrameworkProfile />
   </PropertyGroup>
@@ -52,6 +52,7 @@
     <WarningLevel>4</WarningLevel>
     <DebugType>full</DebugType>
     <ErrorReport>prompt</ErrorReport>
+    <Prefer32Bit>false</Prefer32Bit>
   </PropertyGroup>
   <PropertyGroup Condition=" '$(Configuration)|$(Platform)' == 'Release|AnyCPU' ">
     <OutputPath>bin\Release\</OutputPath>
@@ -75,6 +76,7 @@
     <WarningLevel>4</WarningLevel>
     <DebugType>none</DebugType>
     <ErrorReport>prompt</ErrorReport>
+    <Prefer32Bit>false</Prefer32Bit>
   </PropertyGroup>
   <PropertyGroup Condition="'$(Configuration)|$(Platform)' == 'Debug|x64'">
     <DebugSymbols>true</DebugSymbols>
@@ -87,6 +89,7 @@
     <PlatformTarget>x64</PlatformTarget>
     <ErrorReport>prompt</ErrorReport>
     <CodeAnalysisRuleSet>MinimumRecommendedRules.ruleset</CodeAnalysisRuleSet>
+    <Prefer32Bit>false</Prefer32Bit>
   </PropertyGroup>
   <PropertyGroup Condition="'$(Configuration)|$(Platform)' == 'Release|x64'">
     <OutputPath>bin\x64\Release\</OutputPath>
@@ -98,6 +101,7 @@
     <PlatformTarget>x64</PlatformTarget>
     <ErrorReport>prompt</ErrorReport>
     <CodeAnalysisRuleSet>MinimumRecommendedRules.ruleset</CodeAnalysisRuleSet>
+    <Prefer32Bit>false</Prefer32Bit>
   </PropertyGroup>
   <PropertyGroup>
     <SignAssembly>true</SignAssembly>
@@ -203,7 +207,7 @@
   <PropertyGroup>
     <EngineDir>$(ProjectDir)..\..\..\..\..\..</EngineDir>
     <InstallerDir>$(ProjectDir)..\Installer</InstallerDir>
-     <PostBuildEvent>
+    <PostBuildEvent>
 "$(FrameworkSdkDir)\bin\NETFX 4.6 Tools\tlbexp.exe" "$(TargetDir)$(TargetName)$(TargetExt)" /out:"$(TargetDir)$(TargetName).tlb"
 "$(EngineDir)\Source\ThirdParty\WiX\3.8\heat.exe" file "$(TargetDir)$(TargetName)$(TargetExt)" -gg -o "$(TargetDir)asm_SolidworksToUnreal.wxs"
 "$(EngineDir)\Source\ThirdParty\WiX\3.8\heat.exe" file "$(TargetDir)$(TargetName).tlb" -gg -o "$(TargetDir)tlb_SolidworksToUnreal.wxs"
diff --git a/Engine/Source/Programs/Enterprise/Datasmith/DatasmithSolidworksExporter/PreprocessWxs/App.config b/Engine/Source/Programs/Enterprise/Datasmith/DatasmithSolidworksExporter/PreprocessWxs/App.config
--- a/Engine/Source/Programs/Enterprise/Datasmith/DatasmithSolidworksExporter/PreprocessWxs/App.config
+++ b/Engine/Source/Programs/Enterprise/Datasmith/DatasmithSolidworksExporter/PreprocessWxs/App.config
@@ -1,6 +1,6 @@
-﻿<?xml version="1.0" encoding="utf-8" ?>
+﻿<?xml version="1.0" encoding="utf-8"?>
 <configuration>
     <startup> 
-        <supportedRuntime version="v4.0" sku=".NETFramework,Version=v4.6.1" />
+        <supportedRuntime version="v4.0" sku=".NETFramework,Version=v4.5"/>
     </startup>
-</configuration>
\ No newline at end of file
+</configuration>
diff --git a/Engine/Source/Programs/Enterprise/Datasmith/DatasmithSolidworksExporter/PreprocessWxs/PreprocessWxs.csproj b/Engine/Source/Programs/Enterprise/Datasmith/DatasmithSolidworksExporter/PreprocessWxs/PreprocessWxs.csproj
--- a/Engine/Source/Programs/Enterprise/Datasmith/DatasmithSolidworksExporter/PreprocessWxs/PreprocessWxs.csproj
+++ b/Engine/Source/Programs/Enterprise/Datasmith/DatasmithSolidworksExporter/PreprocessWxs/PreprocessWxs.csproj
@@ -8,10 +8,11 @@
     <OutputType>Exe</OutputType>
     <RootNamespace>PreprocessWxs</RootNamespace>
     <AssemblyName>PreprocessWxs</AssemblyName>
-    <TargetFrameworkVersion>v4.6.1</TargetFrameworkVersion>
+    <TargetFrameworkVersion>v4.5</TargetFrameworkVersion>
     <FileAlignment>512</FileAlignment>
     <AutoGenerateBindingRedirects>true</AutoGenerateBindingRedirects>
     <Deterministic>true</Deterministic>
+    <TargetFrameworkProfile />
   </PropertyGroup>
   <PropertyGroup Condition=" '$(Configuration)|$(Platform)' == 'Debug|AnyCPU' ">
     <PlatformTarget>AnyCPU</PlatformTarget>
