Change Datasmith Exporters pool to  use standard Win64 agents
+ fix NETFX version path to use .NetCore 3.1 version

#rnx

diff --git a/Engine/Source/Programs/Enterprise/Datasmith/DatasmithSolidworksExporter/DatasmithSolidworks/DatasmithSolidworks.csproj b/Engine/Source/Programs/Enterprise/Datasmith/DatasmithSolidworksExporter/DatasmithSolidworks/DatasmithSolidworks.csproj
--- a/Engine/Source/Programs/Enterprise/Datasmith/DatasmithSolidworksExporter/DatasmithSolidworks/DatasmithSolidworks.csproj
+++ b/Engine/Source/Programs/Enterprise/Datasmith/DatasmithSolidworksExporter/DatasmithSolidworks/DatasmithSolidworks.csproj
@@ -208,7 +208,7 @@
     <EngineDir>$(ProjectDir)..\..\..\..\..\..</EngineDir>
     <InstallerDir>$(ProjectDir)..\Installer</InstallerDir>
     <PostBuildEvent>
-"$(FrameworkSdkDir)\bin\NETFX 4.6 Tools\tlbexp.exe" "$(TargetDir)$(TargetName)$(TargetExt)" /out:"$(TargetDir)$(TargetName).tlb"
+"$(FrameworkSdkDir)\bin\NETFX 4.8 Tools\tlbexp.exe" "$(TargetDir)$(TargetName)$(TargetExt)" /out:"$(TargetDir)$(TargetName).tlb"
 "$(EngineDir)\Source\ThirdParty\WiX\3.8\heat.exe" file "$(TargetDir)$(TargetName)$(TargetExt)" -gg -o "$(TargetDir)asm_SolidworksToUnreal.wxs"
 "$(EngineDir)\Source\ThirdParty\WiX\3.8\heat.exe" file "$(TargetDir)$(TargetName).tlb" -gg -o "$(TargetDir)tlb_SolidworksToUnreal.wxs"
 "$(ProjectDir)..\bin\$(Configuration)\PreprocessWxs.exe" "$(TargetDir)$(TargetName)$(TargetExt)" "$(TargetDir)asm_SolidworksToUnreal.wxs" "$(TargetDir)tlb_SolidworksToUnreal.wxs" "$(InstallerDir)" "$(EngineDir)"
