Fixed Mac build for ArchiCAD and SketchUp exporters
 - Replaced C++14 option with c++17 one
 - Workaround permissions issue with Renaming command with two sh scripts performing the same

diff --git a/Engine/Source/Programs/Enterprise/Datasmith/DatasmithARCHICADExporter/Build/Config/app_to_bundle.sh b/Engine/Source/Programs/Enterprise/Datasmith/DatasmithARCHICADExporter/Build/Config/app_to_bundle.sh
new file mode 100644
--- /dev/null
+++ b/Engine/Source/Programs/Enterprise/Datasmith/DatasmithARCHICADExporter/Build/Config/app_to_bundle.sh
@@ -0,0 +1,14 @@
+#!/bin/sh
+
+ConfigPath=`dirname "$0"`
+projectPath=$ConfigPath/..
+
+RelativeEnginePath=$projectPath/../../../../../../../Engine
+EnginePath=`python -c "import os; print(os.path.realpath('$RelativeEnginePath'))"`
+
+echo "Renaming *.app to *.bundle"
+pushd $EnginePath/Binaries/Mac/DatasmithARCHICADExporter
+mv DatasmithARCHICAD23Exporter.app DatasmithARCHICAD23Exporter.bundle
+mv DatasmithARCHICAD24Exporter.app DatasmithARCHICAD24Exporter.bundle
+mv DatasmithARCHICAD25Exporter.app DatasmithARCHICAD25Exporter.bundle
+popd
diff --git a/Engine/Source/Programs/Enterprise/Datasmith/DatasmithARCHICADExporter/Build/Config/bundle_to_app.sh b/Engine/Source/Programs/Enterprise/Datasmith/DatasmithARCHICADExporter/Build/Config/bundle_to_app.sh
new file mode 100644
--- /dev/null
+++ b/Engine/Source/Programs/Enterprise/Datasmith/DatasmithARCHICADExporter/Build/Config/bundle_to_app.sh
@@ -0,0 +1,14 @@
+#!/bin/sh
+
+ConfigPath=`dirname "$0"`
+projectPath=$ConfigPath/..
+
+RelativeEnginePath=$projectPath/../../../../../../../Engine
+EnginePath=`python -c "import os; print(os.path.realpath('$RelativeEnginePath'))"`
+
+echo "Renaming *.bundle to *.app"
+pushd $EnginePath/Binaries/Mac/DatasmithARCHICADExporter
+mv DatasmithARCHICAD23Exporter.bundle DatasmithARCHICAD23Exporter.app
+mv DatasmithARCHICAD24Exporter.bundle DatasmithARCHICAD24Exporter.app
+mv DatasmithARCHICAD25Exporter.bundle DatasmithARCHICAD25Exporter.app
+popd
diff --git a/Engine/Source/Programs/Enterprise/Datasmith/DatasmithSketchUpRubyExporter/Build/build_mac.py b/Engine/Source/Programs/Enterprise/Datasmith/DatasmithSketchUpRubyExporter/Build/build_mac.py
--- a/Engine/Source/Programs/Enterprise/Datasmith/DatasmithSketchUpRubyExporter/Build/build_mac.py
+++ b/Engine/Source/Programs/Enterprise/Datasmith/DatasmithSketchUpRubyExporter/Build/build_mac.py
@@ -122,7 +122,7 @@ UE_INCLUDES_PATH = f'''"{UE_Core}" "{UE_Core}/Internationalization" "{UE_Core}/A
 CXX = ['xcrun', 'clang', 
     '-x', 'objective-c++', 
     '-arch', 'x86_64',
-    '-std=c++14',
+    '-std=c++17',
     '-stdlib=libc++',
 #    '-isysroot', '/Applications/Xcode.app/Contents/Developer/Platforms/MacOSX.platform/Developer/SDKs/MacOSX11.1.sdk',
     '-target', 'x86_64-apple-macos10.14',
@@ -147,7 +147,7 @@ CXX_FLAGS += [
 
 include_paths = []
 include_paths += parse_paths(UE_INCLUDES_PATH)
-include_paths.append(f'{UE_OpenEXR}/include')
+include_paths.append(f'{UE_OpenEXR}/include/openexr')
 include_paths.append(f'{SU_SDK_PATH}/samples/common/ThirdParty/ruby/include/mac')
 include_paths.append(f'{SU_SDK_PATH}/samples/common/ThirdParty/ruby/include/mac/x86_64-darwin') # todo: might not be needed - sample uses x86_64-darwin17 which doesn't exist
 
@@ -156,7 +156,7 @@ framework_search_paths = []
 framework_search_paths.append(f'{SU_SDK_PATH}')
 framework_search_paths.append(f'{SU_SDK_PATH}/samples/common/ThirdParty/ruby/lib/mac')
 
-definitions = [f'SKP_SDK_{SU_VERSION}', 'macintosh=1']
+definitions = [f'SKP_SDK_{SU_VERSION}', 'macintosh=1', '__register_unsupported__']
 UE_PREPROCESSOR_DEFINITIONS = 'UE_BUILD_DEVELOPMENT=1 UE_BUILD_MINIMAL=1 WITH_EDITOR=0 WITH_EDITORONLY_DATA=0 WITH_SERVER_CODE=1 WITH_ENGINE=0 WITH_UNREAL_DEVELOPER_TOOLS=0 WITH_PLUGIN_SUPPORT=0 IS_MONOLITHIC=1 IS_PROGRAM=1 PLATFORM_MAC=1 PLATFORM_APPLE=1 UE_BUILD_DEVELOPMENT_WITH_DEBUGGAME=0 UBT_COMPILED_PLATFORM=Mac CORE_API=DLLIMPORT COREUOBJECT_API=DLLIMPORT DATASMITHEXPORTER_API=DLLIMPORT DATASMITHCORE_API=DLLIMPORT DIRECTLINK_API=DLLIMPORT DATASMITHEXPORTERUI_API=DLLIMPORT'
 definitions += UE_PREPROCESSOR_DEFINITIONS.split()
 
