Fixed build warning when compiling SketchUp exporter on Mac

diff --git a/Engine/Source/Programs/Enterprise/Datasmith/DatasmithSketchUpRubyExporter/Build/build_mac.py b/Engine/Source/Programs/Enterprise/Datasmith/DatasmithSketchUpRubyExporter/Build/build_mac.py
--- a/Engine/Source/Programs/Enterprise/Datasmith/DatasmithSketchUpRubyExporter/Build/build_mac.py
+++ b/Engine/Source/Programs/Enterprise/Datasmith/DatasmithSketchUpRubyExporter/Build/build_mac.py
@@ -125,7 +125,7 @@ CXX = ['xcrun', 'clang',
     '-std=c++17',
     '-stdlib=libc++',
 #    '-isysroot', '/Applications/Xcode.app/Contents/Developer/Platforms/MacOSX.platform/Developer/SDKs/MacOSX11.1.sdk',
-    '-target', 'x86_64-apple-macos10.14',
+    '-target', 'x86_64-apple-macos10.15',
     ]
 
 CXX_FLAGS = []
@@ -191,7 +191,7 @@ link_cmd = [
     'xcrun',
     'clang++', 
 #    '-isysroot', '/Applications/Xcode.app/Contents/Developer/Platforms/MacOSX.platform/Developer/SDKs/MacOSX11.1.sdk',
-    '-target', 'x86_64-apple-macos10.14',
+    '-target', 'x86_64-apple-macos10.15',
     '-Xlinker', '-export_dynamic',
     '-Xlinker', '-no_deduplicate',
     '-stdlib=libc++',
