Fixed build of ArchiCAD exporter on Windows - LWC disabled for the time being

diff --git a/Engine/Source/Programs/Enterprise/Datasmith/DatasmithARCHICADExporter/DatsmithUE4ArchiCAD/DatasmithUE4ArchiCAD.Target.cs b/Engine/Source/Programs/Enterprise/Datasmith/DatasmithARCHICADExporter/DatsmithUE4ArchiCAD/DatasmithUE4ArchiCAD.Target.cs
--- a/Engine/Source/Programs/Enterprise/Datasmith/DatasmithARCHICADExporter/DatsmithUE4ArchiCAD/DatasmithUE4ArchiCAD.Target.cs
+++ b/Engine/Source/Programs/Enterprise/Datasmith/DatasmithARCHICADExporter/DatsmithUE4ArchiCAD/DatasmithUE4ArchiCAD.Target.cs
@@ -20,6 +20,10 @@ public class DatasmithUE4ArchiCADTarget : TargetRules
 
 		ExtraModuleNames.AddRange( new string[] { "DatasmithCore", "DatasmithExporter"} );
 
+		if (Target.Platform == UnrealTargetPlatform.Mac)
+		{
+		}
+
 		bBuildDeveloperTools = false;
 		bUseMallocProfiler = false;
 		bBuildWithEditorOnlyData = true;
@@ -32,5 +36,6 @@ public class DatasmithUE4ArchiCADTarget : TargetRules
 		bHasExports = true;
 
 		GlobalDefinitions.Add("UE_EXTERNAL_PROFILING_ENABLED=0"); // For DirectLinkUI (see FDatasmithExporterManager::FInitOptions)
+		GlobalDefinitions.Add("UE_LARGE_WORLD_COORDINATES_DISABLED=1"); // LARGE WORLD COORDINATES
 	}
 }
