Fixed build warning when compiling SketchUp exporter

diff --git a/Engine/Source/Programs/Enterprise/Datasmith/DatasmithSketchUpRubyExporter/Private/DatasmithSketchUpUtils.cpp b/Engine/Source/Programs/Enterprise/Datasmith/DatasmithSketchUpRubyExporter/Private/DatasmithSketchUpUtils.cpp
--- a/Engine/Source/Programs/Enterprise/Datasmith/DatasmithSketchUpRubyExporter/Private/DatasmithSketchUpUtils.cpp
+++ b/Engine/Source/Programs/Enterprise/Datasmith/DatasmithSketchUpRubyExporter/Private/DatasmithSketchUpUtils.cpp
@@ -237,7 +237,11 @@ void DatasmithSketchUpUtils::SetActorTransform(
 	// Set the world transform of the Datasmith actor.
 	InActorElement->SetScale(Scale.x, Scale.y, Scale.z, false);
 	InActorElement->SetRotation(Rotation.X, Rotation.Y, Rotation.Z, Rotation.W, false);
+#if UE_LARGE_WORLD_COORDINATES_DISABLED
 	InActorElement->SetTranslation(Translation.X, Translation.Y, Translation.Z, false);
+#else
+	InActorElement->SetTranslation((float)Translation.X, (float)Translation.Y, (float)Translation.Z, false);
+#endif
 }
 
 SUTransformation DatasmithSketchUpUtils::GetComponentInstanceTransform(SUComponentInstanceRef InComponentInstanceRef, SUTransformation const& InWorldTransform)
