Fixed build breakage when compiling ArchiCAD exporter

diff --git a/Engine/Source/Programs/Enterprise/Datasmith/DatasmithARCHICADExporter/Private/GeometryUtil.h b/Engine/Source/Programs/Enterprise/Datasmith/DatasmithARCHICADExporter/Private/GeometryUtil.h
--- a/Engine/Source/Programs/Enterprise/Datasmith/DatasmithARCHICADExporter/Private/GeometryUtil.h
+++ b/Engine/Source/Programs/Enterprise/Datasmith/DatasmithARCHICADExporter/Private/GeometryUtil.h
@@ -8,9 +8,6 @@
 #include "ModelVector.hpp"
 #include "Vertex.hpp"
 
-struct FQuat;
-DECLARE_LWC_TYPE(Vector, 3);
-
 BEGIN_NAMESPACE_UE_AC
 
 // Tool geometry class
