Replace PBR Materials with Engine Datasmith Materials

diff --git a/Engine/Plugins/Enterprise/DatasmithImporter/Source/DatasmithTranslator/Private/DatasmithTranslatorModule.cpp b/Engine/Plugins/Enterprise/DatasmithImporter/Source/DatasmithTranslator/Private/DatasmithTranslatorModule.cpp
--- a/Engine/Plugins/Enterprise/DatasmithImporter/Source/DatasmithTranslator/Private/DatasmithTranslatorModule.cpp
+++ b/Engine/Plugins/Enterprise/DatasmithImporter/Source/DatasmithTranslator/Private/DatasmithTranslatorModule.cpp
@@ -7,6 +7,7 @@
 #include "MasterMaterials/DatasmithMasterMaterialManager.h"
 #include "MasterMaterials/DatasmithRevitMaterialSelector.h"
 #include "MasterMaterials/DatasmithSketchupMaterialSelector.h"
+#include "MasterMaterials/DatasmithStdMaterialSelector.h"
 
 void IDatasmithTranslatorModule::StartupModule()
 {
@@ -17,6 +18,7 @@ void IDatasmithTranslatorModule::StartupModule()
 	FDatasmithMasterMaterialManager::Get().RegisterSelector(TEXT("Revit"), MakeShared< FDatasmithRevitMaterialSelector >());
 	FDatasmithMasterMaterialManager::Get().RegisterSelector(TEXT("SketchUp"), MakeShared< FDatasmithSketchUpMaterialSelector >());
 	FDatasmithMasterMaterialManager::Get().RegisterSelector(TEXT("CityEngine"), MakeShared< FDatasmithCityEngineMaterialSelector >());
+	FDatasmithMasterMaterialManager::Get().RegisterSelector(TEXT("StdMaterial"), MakeShared< FDatasmithStdMaterialSelector >());
 }
 
 void IDatasmithTranslatorModule::ShutdownModule()
@@ -25,6 +27,7 @@ void IDatasmithTranslatorModule::ShutdownModule()
 	FDatasmithMasterMaterialManager::Get().UnregisterSelector(TEXT("Revit"));
 	FDatasmithMasterMaterialManager::Get().UnregisterSelector(TEXT("SketchUp"));
 	FDatasmithMasterMaterialManager::Get().UnregisterSelector(TEXT("CityEngine"));
+	FDatasmithMasterMaterialManager::Get().UnregisterSelector(TEXT("StdMaterial"));
 
 	FDatasmithMasterMaterialManager::Destroy();
 }
diff --git a/Engine/Plugins/Enterprise/DatasmithImporter/Source/DatasmithTranslator/Private/MasterMaterials/DatasmithStdMaterialSelector.cpp b/Engine/Plugins/Enterprise/DatasmithImporter/Source/DatasmithTranslator/Private/MasterMaterials/DatasmithStdMaterialSelector.cpp
new file mode 100644
--- /dev/null
+++ b/Engine/Plugins/Enterprise/DatasmithImporter/Source/DatasmithTranslator/Private/MasterMaterials/DatasmithStdMaterialSelector.cpp
@@ -0,0 +1,31 @@
+// Copyright Epic Games, Inc. All Rights Reserved.
+
+#include "DatasmithStdMaterialSelector.h"
+
+#include "IDatasmithSceneElements.h"
+
+#include "Materials/Material.h"
+#include "Materials/MaterialInstanceConstant.h"
+#include "UObject/SoftObjectPath.h"
+#include "Templates/Casts.h"
+
+FDatasmithStdMaterialSelector::FDatasmithStdMaterialSelector()
+{
+	// Master
+	MasterMaterialOpaque.FromSoftObjectPath( FSoftObjectPath("/DatasmithContent/Materials/StdOpaque/M_StdOpaque.M_StdOpaque") );
+	MasterMaterialTranslucent.FromSoftObjectPath( FSoftObjectPath("/DatasmithContent/Materials/StdTranslucent/M_StdTranslucent.M_StdTranslucent") );
+}
+
+bool FDatasmithStdMaterialSelector::IsValid() const
+{
+	return MasterMaterialOpaque.IsValid() && MasterMaterialTranslucent.IsValid();
+}
+
+const FDatasmithMasterMaterial& FDatasmithStdMaterialSelector::GetMasterMaterial( const TSharedPtr< IDatasmithMasterMaterialElement >& InDatasmithMaterial ) const
+{
+	if (InDatasmithMaterial->GetMaterialType() == EDatasmithMasterMaterialType::Transparent)
+	{
+		return MasterMaterialTranslucent;
+	}
+	return MasterMaterialOpaque;
+}
diff --git a/Engine/Plugins/Enterprise/DatasmithImporter/Source/DatasmithTranslator/Private/MasterMaterials/DatasmithStdMaterialSelector.h b/Engine/Plugins/Enterprise/DatasmithImporter/Source/DatasmithTranslator/Private/MasterMaterials/DatasmithStdMaterialSelector.h
new file mode 100644
--- /dev/null
+++ b/Engine/Plugins/Enterprise/DatasmithImporter/Source/DatasmithTranslator/Private/MasterMaterials/DatasmithStdMaterialSelector.h
@@ -0,0 +1,25 @@
+// Copyright Epic Games, Inc. All Rights Reserved.
+#pragma once
+
+#include "DatasmithDefinitions.h"
+#include "MasterMaterials/DatasmithMasterMaterial.h"
+#include "MasterMaterials/DatasmithMasterMaterialSelector.h"
+
+#include "Templates/SharedPointer.h"
+
+class IDatasmithMasterMaterialElement;
+
+class FDatasmithStdMaterialSelector : public FDatasmithMasterMaterialSelector
+{
+public:
+	FDatasmithStdMaterialSelector();
+	virtual ~FDatasmithStdMaterialSelector() = default;
+
+	virtual bool IsValid() const override;
+	virtual const FDatasmithMasterMaterial& GetMasterMaterial( const TSharedPtr< IDatasmithMasterMaterialElement >& InDatasmithMaterial ) const override;
+	virtual void FinalizeMaterialInstance(const TSharedPtr< IDatasmithMasterMaterialElement >& InDatasmithMaterial, UMaterialInstanceConstant* MaterialInstance) const override {}
+
+private:
+	FDatasmithMasterMaterial MasterMaterialOpaque;
+	FDatasmithMasterMaterial MasterMaterialTranslucent;
+};
diff --git a/Engine/Source/Programs/Enterprise/Datasmith/DatasmithSolidworksExporter/Private/Engine/Processor.cs b/Engine/Source/Programs/Enterprise/Datasmith/DatasmithSolidworksExporter/Private/Engine/Processor.cs
--- a/Engine/Source/Programs/Enterprise/Datasmith/DatasmithSolidworksExporter/Private/Engine/Processor.cs
+++ b/Engine/Source/Programs/Enterprise/Datasmith/DatasmithSolidworksExporter/Private/Engine/Processor.cs
@@ -70,7 +70,7 @@ namespace SolidworksDatasmith.Engine
         private FDatasmithFacadeUEPbrMaterial DefaultMaterial;
         private static readonly string DefaultMaterialName = "Default";
         private Dictionary<string, FDatasmithFacadeTexture> DatasmithTextures = new Dictionary<string, FDatasmithFacadeTexture>();
-        private Dictionary<SwMaterial, FDatasmithFacadeUEPbrMaterial> SwMat2Datasmith = new Dictionary<SwMaterial, FDatasmithFacadeUEPbrMaterial>();
+        private Dictionary<SwMaterial, FDatasmithFacadeMasterMaterial> SwMat2Datasmith = new Dictionary<SwMaterial, FDatasmithFacadeMasterMaterial>();
         private Dictionary<string, SwMaterial> ComponentMaterials = new Dictionary<string, SwMaterial>();
         private MeshFactory MeshFactory = new MeshFactory();
         private HashSet<string> MeshNames = new HashSet<string>();
@@ -121,9 +121,9 @@ namespace SolidworksDatasmith.Engine
             return res;
         }
 
-        private FDatasmithFacadeUEPbrMaterial GetDatasmithMaterial(SwMaterial mat, bool addIfNotExisting)
+        private FDatasmithFacadeMasterMaterial GetDatasmithMaterial(SwMaterial mat, bool addIfNotExisting)
         {
-            FDatasmithFacadeUEPbrMaterial dm = null;
+			FDatasmithFacadeMasterMaterial dm = null;
             foreach (var mm in SwMat2Datasmith)
             {
                 if (SwMaterial.AreTheSame(mm.Key, mat, false))
@@ -194,7 +194,7 @@ namespace SolidworksDatasmith.Engine
                 Directory.CreateDirectory(directLinkPath);
 
             // datasmith scene setup
-            processor._datasmithScene = new FDatasmithFacadeScene("Solidworks", "Solidworks", "Solidworks", "2021");
+            processor._datasmithScene = new FDatasmithFacadeScene("StdMaterial", "Solidworks", "Solidworks", "2021");
             processor._datasmithScene.SetOutputPath(directLinkPath);
             processor._datasmithScene.SetName("SolidWorks_Scene");
             processor.DatasmithScene.PreExport();
@@ -747,7 +747,7 @@ namespace SolidworksDatasmith.Engine
             Queue.Push(cmd);
         }
 
-        private FDatasmithFacadeUEPbrMaterial MaterialToDatasmith(SwMaterial material)
+        private FDatasmithFacadeMasterMaterial MaterialToDatasmith(SwMaterial material)
         {
             var type = SwMaterial.GetMaterialType(material.ShaderName);
 
@@ -774,96 +774,96 @@ namespace SolidworksDatasmith.Engine
                         roughness = 1f;
                 }
             }
+		
+            float mult = (type == SwMaterial.MaterialType.TYPE_LIGHTWEIGHT) ? (float)material.Diffuse : 1.0f;
 
-            var pbr = new FDatasmithFacadeUEPbrMaterial(material.Name);
+            float fR = mult * material.PrimaryColor.R * 1.0f / 255.0f;
+            float fG = mult * material.PrimaryColor.G * 1.0f / 255.0f;
+            float fB = mult * material.PrimaryColor.B * 1.0f / 255.0f;
 
-            double mult = (type == SwMaterial.MaterialType.TYPE_LIGHTWEIGHT) ? material.Diffuse : 1.0;
+			FDatasmithFacadeMasterMaterial MasterMaterial = new FDatasmithFacadeMasterMaterial(material.Name);
 
-            double fR = mult * material.PrimaryColor.R * 1.0 / 255.0;
-            double fG = mult * material.PrimaryColor.G * 1.0 / 255.0;
-            double fB = mult * material.PrimaryColor.B * 1.0 / 255.0;
-
-            var colorExpr = pbr.AddMaterialExpressionColor();
-            colorExpr.SetColor((float)fR, (float)fG, (float)fB, 1f);
-            colorExpr.SetName("Diffuse Color");
-            colorExpr.ConnectExpression(pbr.GetBaseColor());
-
-            FDatasmithFacadeMaterialExpressionScalar roughExpr = pbr.AddMaterialExpressionScalar();
-            roughExpr.SetName("Roughness");
-            roughExpr.SetScalar(roughness);
-            roughExpr.ConnectExpression(pbr.GetRoughness());
+			MasterMaterial.SetMaterialType(FDatasmithFacadeMasterMaterial.EMasterMaterialType.Opaque);
+			MasterMaterial.AddColor("TintColor", fR, fG, fB, 1.0F);
+			MasterMaterial.AddFloat("RoughnessAmount", roughness);
+		
+		    if (material.Transparency > 0.0)
+            {
+				MasterMaterial.SetMaterialType(FDatasmithFacadeMasterMaterial.EMasterMaterialType.Transparent);
+				MasterMaterial.AddFloat("Metalness", metallic);
 
-            FDatasmithFacadeMaterialExpressionScalar metalExpr = pbr.AddMaterialExpressionScalar();
-            metalExpr.SetName("Metallic");
-            metalExpr.SetScalar(metallic);
-            metalExpr.ConnectExpression(pbr.GetMetallic());
+				MasterMaterial.AddColor("OpacityAndRefraction", 
+					0.25f,                              // Opacity
+					1.0f,								// Refraction
+					0.0f,								// Refraction Exponent
+					1f - (float)material.Transparency  // Fresnel Opacity
+				);
 
-            if (material.Emission > 0.0)
-            {
-                var emissionExpr = pbr.AddMaterialExpressionColor();
-                emissionExpr.SetColor((float)fR, (float)fG, (float)fB, 1f);
-                emissionExpr.SetName("Emissive Color");
-                emissionExpr.ConnectExpression(pbr.GetEmissiveColor());
+				SetNormalMap(material, MasterMaterial, "NormalMap");
             }
+			else
+			{
+				MasterMaterial.AddFloat("MetallicAmount", metallic);
 
-            if (material.Transparency > 0.0)
-            {
-                var opacityExpr = pbr.AddMaterialExpressionScalar();
-                opacityExpr.SetName("Opacity");
-                opacityExpr.SetScalar(1f - (float)material.Transparency);
-                opacityExpr.ConnectExpression(pbr.GetOpacity());
+				if (material.Emission > 0.0)
+				{
+					MasterMaterial.AddFloat("LuminanceAmount", (float)material.Emission);
+					MasterMaterial.AddColor("LuminanceFilter", fR, fG, fB, 1.0f);
+				}
 
-                //pbr.SetBlendMode()
-            }
+				SetDiffuseMap(material, MasterMaterial, "ColorMap");
+				SetNormalMap(material, MasterMaterial, "NormalMap");
+			}
 
-            if (!string.IsNullOrEmpty(material.Texture) && !File.Exists(material.Texture))
-            {
-                material.Texture = SwMaterial.ComputeAssemblySideTexturePath(material.Texture);
-            }
+			return MasterMaterial;
+        }
 
-            if (!string.IsNullOrEmpty(material.Texture) && File.Exists(material.Texture))
-            {
-                string textureName = Path.GetFileNameWithoutExtension(material.Texture);
+		void SetDiffuseMap(SwMaterial InSwMaterial, FDatasmithFacadeMasterMaterial InMasterMaterial, string InParamName)
+		{
+			if (!string.IsNullOrEmpty(InSwMaterial.Texture) && !File.Exists(InSwMaterial.Texture))
+			{
+				InSwMaterial.Texture = SwMaterial.ComputeAssemblySideTexturePath(InSwMaterial.Texture);
+			}
 
-                FDatasmithFacadeTexture TextureElement = null;
-                if (!DatasmithTextures.TryGetValue(material.Texture, out TextureElement))
-                {
-                    TextureElement = new FDatasmithFacadeTexture(textureName);
-                    TextureElement.SetFile(material.Texture);
-                    TextureElement.SetTextureFilter(FDatasmithFacadeTexture.ETextureFilter.Default);
-                    TextureElement.SetRGBCurve(1);
-                    TextureElement.SetTextureAddressX(FDatasmithFacadeTexture.ETextureAddress.Wrap);
-                    TextureElement.SetTextureAddressY(FDatasmithFacadeTexture.ETextureAddress.Wrap);
-                    FDatasmithFacadeTexture.ETextureMode TextureMode = FDatasmithFacadeTexture.ETextureMode.Diffuse;
-                    TextureElement.SetTextureMode(TextureMode);
-                    DatasmithScene.AddTexture(TextureElement);
-                    DatasmithTextures.Add(material.Texture, TextureElement);
-                }
+			if (!string.IsNullOrEmpty(InSwMaterial.Texture) && File.Exists(InSwMaterial.Texture))
+			{
+				string textureName = Path.GetFileNameWithoutExtension(InSwMaterial.Texture);
 
-                FDatasmithFacadeMaterialsUtils.FWeightedMaterialExpressionParameters WeightedExpressionParameters = new FDatasmithFacadeMaterialsUtils.FWeightedMaterialExpressionParameters(1f);
-                FDatasmithFacadeMaterialsUtils.FUVEditParameters UVParameters = new FDatasmithFacadeMaterialsUtils.FUVEditParameters();
-                textureName = DatasmithExporter.SanitizeName(textureName);
-                FDatasmithFacadeMaterialExpression TextureExpression = FDatasmithFacadeMaterialsUtils.CreateTextureExpression(pbr, "Diffuse Map", textureName, UVParameters);
-                WeightedExpressionParameters.SetColorsRGB(material.PrimaryColor.R, material.PrimaryColor.G, material.PrimaryColor.B, 255);
-                WeightedExpressionParameters.SetExpression(TextureExpression);
-                FDatasmithFacadeMaterialExpression Expression = FDatasmithFacadeMaterialsUtils.CreateWeightedMaterialExpression(pbr, "Diffuse Color", WeightedExpressionParameters);
-                pbr.GetBaseColor().SetExpression(Expression);
-            }
+				FDatasmithFacadeTexture TextureElement = null;
+				if (!DatasmithTextures.TryGetValue(InSwMaterial.Texture, out TextureElement))
+				{
+					TextureElement = new FDatasmithFacadeTexture(textureName);
+					TextureElement.SetFile(InSwMaterial.Texture);
+					TextureElement.SetTextureFilter(FDatasmithFacadeTexture.ETextureFilter.Default);
+					TextureElement.SetRGBCurve(1);
+					TextureElement.SetTextureAddressX(FDatasmithFacadeTexture.ETextureAddress.Wrap);
+					TextureElement.SetTextureAddressY(FDatasmithFacadeTexture.ETextureAddress.Wrap);
+					FDatasmithFacadeTexture.ETextureMode TextureMode = FDatasmithFacadeTexture.ETextureMode.Diffuse;
+					TextureElement.SetTextureMode(TextureMode);
+					DatasmithScene.AddTexture(TextureElement);
+					DatasmithTextures.Add(InSwMaterial.Texture, TextureElement);
+
+					InMasterMaterial.AddTexture(InParamName, TextureElement);
+				}
+			}
+		}
 
-            if (!string.IsNullOrEmpty(material.BumpTextureFileName) && !File.Exists(material.BumpTextureFileName))
+		void SetNormalMap(SwMaterial InSwMaterial, FDatasmithFacadeMasterMaterial InMasterMaterial, string InParamName)
+		{
+            if (!string.IsNullOrEmpty(InSwMaterial.BumpTextureFileName) && !File.Exists(InSwMaterial.BumpTextureFileName))
             {
-                material.BumpTextureFileName = SwMaterial.ComputeAssemblySideTexturePath(material.BumpTextureFileName);
+                InSwMaterial.BumpTextureFileName = SwMaterial.ComputeAssemblySideTexturePath(InSwMaterial.BumpTextureFileName);
             }
 
-            if (!string.IsNullOrEmpty(material.BumpTextureFileName) && File.Exists(material.BumpTextureFileName))
+            if (!string.IsNullOrEmpty(InSwMaterial.BumpTextureFileName) && File.Exists(InSwMaterial.BumpTextureFileName))
             {
-                string textureName = Path.GetFileNameWithoutExtension(material.BumpTextureFileName);
+                string textureName = Path.GetFileNameWithoutExtension(InSwMaterial.BumpTextureFileName);
 
                 FDatasmithFacadeTexture TextureElement = null;
-                if (!DatasmithTextures.TryGetValue(material.BumpTextureFileName, out TextureElement))
+                if (!DatasmithTextures.TryGetValue(InSwMaterial.BumpTextureFileName, out TextureElement))
                 {
                     TextureElement = new FDatasmithFacadeTexture(textureName);
-                    TextureElement.SetFile(material.BumpTextureFileName);
+                    TextureElement.SetFile(InSwMaterial.BumpTextureFileName);
                     TextureElement.SetTextureFilter(FDatasmithFacadeTexture.ETextureFilter.Default);
                     TextureElement.SetRGBCurve(1);
                     TextureElement.SetTextureAddressX(FDatasmithFacadeTexture.ETextureAddress.Wrap);
@@ -871,16 +871,12 @@ namespace SolidworksDatasmith.Engine
                     FDatasmithFacadeTexture.ETextureMode TextureMode = FDatasmithFacadeTexture.ETextureMode.Normal;
                     TextureElement.SetTextureMode(TextureMode);
                     DatasmithScene.AddTexture(TextureElement);
-                    DatasmithTextures.Add(material.BumpTextureFileName, TextureElement);
-                }
+                    DatasmithTextures.Add(InSwMaterial.BumpTextureFileName, TextureElement);
 
-                FDatasmithFacadeMaterialsUtils.FUVEditParameters UVParameters = new FDatasmithFacadeMaterialsUtils.FUVEditParameters();
-                FDatasmithFacadeMaterialExpression TextureExpression = FDatasmithFacadeMaterialsUtils.CreateTextureExpression(pbr, "Bump Map", textureName, UVParameters);
-                pbr.GetNormal().SetExpression(TextureExpression);
+					InMasterMaterial.AddTexture(InParamName, TextureElement);
+                }
             }
-
-            return pbr;
-        }
+		}
 
 		// Keep list of all actor bindings for a single variant
 		class ComponentMapper
