Datasmith Revit: Export User Defined Collection of Metadata

diff --git a/Engine/Source/Programs/Enterprise/Datasmith/DatasmithRevitExporter/DatasmithRevit2018/DatasmithRevit2018.csproj b/Engine/Source/Programs/Enterprise/Datasmith/DatasmithRevitExporter/DatasmithRevit2018/DatasmithRevit2018.csproj
--- a/Engine/Source/Programs/Enterprise/Datasmith/DatasmithRevitExporter/DatasmithRevit2018/DatasmithRevit2018.csproj
+++ b/Engine/Source/Programs/Enterprise/Datasmith/DatasmithRevitExporter/DatasmithRevit2018/DatasmithRevit2018.csproj
@@ -87,6 +87,13 @@
     <Compile Include="..\Private\DatasmithRevitMaterial.cs">
       <Link>DatasmithRevitExporter\DatasmithRevitMaterial.cs</Link>
     </Compile>
+    <Compile Include="..\Private\DatasmithRevitMetadataManager.cs">
+      <Link>DatasmithRevitExporter\DatasmithRevitMetdataManager.cs</Link>
+    </Compile>
+    <Compile Include="..\Private\DatasmithRevitSettingsDialog.cs">
+      <Link>DatasmithRevitExporter\DatasmithRevitSettingsDialog.cs</Link>
+	  <SubType>Form</SubType>
+    </Compile>
     <Compile Include="..\..\..\..\..\..\Binaries\Win64\DatasmithFacadeCSharp\Public\*.cs">
       <Link>DatasmithSwigFacade\%(Filename)%(Extension)</Link>
     </Compile>
diff --git a/Engine/Source/Programs/Enterprise/Datasmith/DatasmithRevitExporter/DatasmithRevit2019/DatasmithRevit2019.csproj b/Engine/Source/Programs/Enterprise/Datasmith/DatasmithRevitExporter/DatasmithRevit2019/DatasmithRevit2019.csproj
--- a/Engine/Source/Programs/Enterprise/Datasmith/DatasmithRevitExporter/DatasmithRevit2019/DatasmithRevit2019.csproj
+++ b/Engine/Source/Programs/Enterprise/Datasmith/DatasmithRevitExporter/DatasmithRevit2019/DatasmithRevit2019.csproj
@@ -87,6 +87,13 @@
     <Compile Include="..\Private\DatasmithRevitMaterial.cs">
       <Link>DatasmithRevitExporter\DatasmithRevitMaterial.cs</Link>
     </Compile>
+    <Compile Include="..\Private\DatasmithRevitMetadataManager.cs">
+      <Link>DatasmithRevitExporter\DatasmithRevitMetdataManager.cs</Link>
+    </Compile>
+    <Compile Include="..\Private\DatasmithRevitSettingsDialog.cs">
+      <Link>DatasmithRevitExporter\DatasmithRevitSettingsDialog.cs</Link>
+	  <SubType>Form</SubType>
+    </Compile>
     <Compile Include="..\..\..\..\..\..\Binaries\Win64\DatasmithFacadeCSharp\Public\*.cs">
       <Link>DatasmithSwigFacade\%(Filename)%(Extension)</Link>
     </Compile>
diff --git a/Engine/Source/Programs/Enterprise/Datasmith/DatasmithRevitExporter/DatasmithRevit2020/DatasmithRevit2020.csproj b/Engine/Source/Programs/Enterprise/Datasmith/DatasmithRevitExporter/DatasmithRevit2020/DatasmithRevit2020.csproj
--- a/Engine/Source/Programs/Enterprise/Datasmith/DatasmithRevitExporter/DatasmithRevit2020/DatasmithRevit2020.csproj
+++ b/Engine/Source/Programs/Enterprise/Datasmith/DatasmithRevitExporter/DatasmithRevit2020/DatasmithRevit2020.csproj
@@ -87,6 +87,13 @@
     <Compile Include="..\Private\DatasmithRevitMaterial.cs">
       <Link>DatasmithRevitExporter\DatasmithRevitMaterial.cs</Link>
     </Compile>
+    <Compile Include="..\Private\DatasmithRevitMetadataManager.cs">
+      <Link>DatasmithRevitExporter\DatasmithRevitMetdataManager.cs</Link>
+    </Compile>
+    <Compile Include="..\Private\DatasmithRevitSettingsDialog.cs">
+      <Link>DatasmithRevitExporter\DatasmithRevitSettingsDialog.cs</Link>
+	  <SubType>Form</SubType>
+    </Compile>
     <Compile Include="..\..\..\..\..\..\Binaries\Win64\DatasmithFacadeCSharp\Public\*.cs">
       <Link>DatasmithSwigFacade\%(Filename)%(Extension)</Link>
     </Compile>
diff --git a/Engine/Source/Programs/Enterprise/Datasmith/DatasmithRevitExporter/DatasmithRevit2021/DatasmithRevit2021.csproj b/Engine/Source/Programs/Enterprise/Datasmith/DatasmithRevitExporter/DatasmithRevit2021/DatasmithRevit2021.csproj
--- a/Engine/Source/Programs/Enterprise/Datasmith/DatasmithRevitExporter/DatasmithRevit2021/DatasmithRevit2021.csproj
+++ b/Engine/Source/Programs/Enterprise/Datasmith/DatasmithRevitExporter/DatasmithRevit2021/DatasmithRevit2021.csproj
@@ -87,6 +87,13 @@
     <Compile Include="..\Private\DatasmithRevitMaterial.cs">
       <Link>DatasmithRevitExporter\DatasmithRevitMaterial.cs</Link>
     </Compile>
+    <Compile Include="..\Private\DatasmithRevitMetadataManager.cs">
+      <Link>DatasmithRevitExporter\DatasmithRevitMetdataManager.cs</Link>
+    </Compile>
+    <Compile Include="..\Private\DatasmithRevitSettingsDialog.cs">
+      <Link>DatasmithRevitExporter\DatasmithRevitSettingsDialog.cs</Link>
+      <SubType>Form</SubType>
+    </Compile>
     <Compile Include="..\..\..\..\..\..\Binaries\Win64\DatasmithFacadeCSharp\Public\*.cs">
       <Link>DatasmithSwigFacade\%(Filename)%(Extension)</Link>
     </Compile>
diff --git a/Engine/Source/Programs/Enterprise/Datasmith/DatasmithRevitExporter/DatasmithRevit2022/DatasmithRevit2022.csproj b/Engine/Source/Programs/Enterprise/Datasmith/DatasmithRevitExporter/DatasmithRevit2022/DatasmithRevit2022.csproj
--- a/Engine/Source/Programs/Enterprise/Datasmith/DatasmithRevitExporter/DatasmithRevit2022/DatasmithRevit2022.csproj
+++ b/Engine/Source/Programs/Enterprise/Datasmith/DatasmithRevitExporter/DatasmithRevit2022/DatasmithRevit2022.csproj
@@ -87,6 +87,13 @@
     <Compile Include="..\Private\DatasmithRevitMaterial.cs">
       <Link>DatasmithRevitExporter\DatasmithRevitMaterial.cs</Link>
     </Compile>
+    <Compile Include="..\Private\DatasmithRevitMetadataManager.cs">
+      <Link>DatasmithRevitExporter\DatasmithRevitMetdataManager.cs</Link>
+    </Compile>
+    <Compile Include="..\Private\DatasmithRevitSettingsDialog.cs">
+      <Link>DatasmithRevitExporter\DatasmithRevitSettingsDialog.cs</Link>
+      <SubType>Form</SubType>
+    </Compile>
     <Compile Include="..\..\..\..\..\..\Binaries\Win64\DatasmithFacadeCSharp\Public\*.cs">
       <Link>DatasmithSwigFacade\%(Filename)%(Extension)</Link>
     </Compile>
diff --git a/Engine/Source/Programs/Enterprise/Datasmith/DatasmithRevitExporter/Private/DatasmithRevitApplication.cs b/Engine/Source/Programs/Enterprise/Datasmith/DatasmithRevitExporter/Private/DatasmithRevitApplication.cs
--- a/Engine/Source/Programs/Enterprise/Datasmith/DatasmithRevitExporter/Private/DatasmithRevitApplication.cs
+++ b/Engine/Source/Programs/Enterprise/Datasmith/DatasmithRevitExporter/Private/DatasmithRevitApplication.cs
@@ -43,10 +43,12 @@ namespace DatasmithRevitExporter
 			PushButtonData ExportButtonData = new PushButtonData("Export3DView", DatasmithRevitResources.Strings.ButtonExport3DView, AssemblyPath, "DatasmithRevitExporter.DatasmithExportRevitCommand");
 			PushButtonData SyncButtonData = new PushButtonData("Sync3DView", DatasmithRevitResources.Strings.ButtonSync, AssemblyPath, "DatasmithRevitExporter.DatasmithSyncRevitCommand");
 			PushButtonData ManageConnectionsButtonData = new PushButtonData("Connections", DatasmithRevitResources.Strings.ButtonConnections, AssemblyPath, "DatasmithRevitExporter.DatasmithManageConnectionsRevitCommand");
+			PushButtonData SettingsButtonData = new PushButtonData("Settings", DatasmithRevitResources.Strings.ButtonSettings, AssemblyPath, "DatasmithRevitExporter.DatasmithShowSettingsRevitCommand");
 			PushButtonData LogButtonData = new PushButtonData("Messages", DatasmithRevitResources.Strings.ButtonMessages, AssemblyPath, "DatasmithRevitExporter.DatasmithShowMessagesRevitCommand");
 
 			PushButton SyncPushButton = DirectLinkRibbonPanel.AddItem(SyncButtonData) as PushButton;
 			PushButton ManageConnectionsButton = DirectLinkRibbonPanel.AddItem(ManageConnectionsButtonData) as PushButton;
+			PushButton SettingsButton = DirectLinkRibbonPanel.AddItem(SettingsButtonData) as PushButton;
 			PushButton ExportPushButton = FileExportRibbonPanel.AddItem(ExportButtonData) as PushButton;
 			PushButton ShowLogButton = DatasmithRibbonPanel.AddItem(LogButtonData) as PushButton;
 
@@ -65,6 +67,11 @@ namespace DatasmithRevitExporter
 			ManageConnectionsButton.LargeImage = new BitmapImage(new Uri(DatasmithIconBase + "32.png"));
 			ManageConnectionsButton.ToolTip = "Manage connections";
 
+			DatasmithIconBase = Path.Combine(Path.GetDirectoryName(AssemblyPath), "DatasmithSettingsIcon");
+			SettingsButton.Image = new BitmapImage(new Uri(DatasmithIconBase + "16.png"));
+			SettingsButton.LargeImage = new BitmapImage(new Uri(DatasmithIconBase + "32.png"));
+			SettingsButton.ToolTip = "Show Settings";
+
 			DatasmithIconBase = Path.Combine(Path.GetDirectoryName(AssemblyPath), "DatasmithLogIcon");
 			ShowLogButton.Image = new BitmapImage(new Uri(DatasmithIconBase + "16.png"));
 			ShowLogButton.LargeImage = new BitmapImage(new Uri(DatasmithIconBase + "32.png"));
@@ -100,6 +107,8 @@ namespace DatasmithRevitExporter
 
 			Debug.Assert(bDirectLinkInitOk);
 
+			FMetadataManager.Init(InApplication);
+
 			return Result.Succeeded;
 		}
 
@@ -134,6 +143,7 @@ namespace DatasmithRevitExporter
 			{
 				ExportMessagesDialog.Close();
 			}
+			FMetadataManager.Destroy(InApplication);
 			FDatasmithFacadeDirectLink.Shutdown();
 			return Result.Succeeded;
 		}
diff --git a/Engine/Source/Programs/Enterprise/Datasmith/DatasmithRevitExporter/Private/DatasmithRevitCommand.cs b/Engine/Source/Programs/Enterprise/Datasmith/DatasmithRevitExporter/Private/DatasmithRevitCommand.cs
--- a/Engine/Source/Programs/Enterprise/Datasmith/DatasmithRevitExporter/Private/DatasmithRevitCommand.cs
+++ b/Engine/Source/Programs/Enterprise/Datasmith/DatasmithRevitExporter/Private/DatasmithRevitCommand.cs
@@ -418,4 +418,15 @@ namespace DatasmithRevitExporter
 			return Result.Succeeded;
 		}
 	}
+
+	[Transaction(TransactionMode.Manual)]
+	public class DatasmithShowSettingsRevitCommand : IExternalCommand
+	{
+		public Result Execute(ExternalCommandData InCommandData, ref string OutCommandMessage, ElementSet OutElements)
+		{
+			DatasmithRevitSettingsDialog ExportOptions = new DatasmithRevitSettingsDialog(InCommandData.Application.ActiveUIDocument.Document);
+			ExportOptions.ShowDialog();
+			return Result.Succeeded;
+		}
+	}
 }
diff --git a/Engine/Source/Programs/Enterprise/Datasmith/DatasmithRevitExporter/Private/DatasmithRevitDirectLink.cs b/Engine/Source/Programs/Enterprise/Datasmith/DatasmithRevitExporter/Private/DatasmithRevitDirectLink.cs
--- a/Engine/Source/Programs/Enterprise/Datasmith/DatasmithRevitExporter/Private/DatasmithRevitDirectLink.cs
+++ b/Engine/Source/Programs/Enterprise/Datasmith/DatasmithRevitExporter/Private/DatasmithRevitDirectLink.cs
@@ -677,7 +677,7 @@ namespace DatasmithRevitExporter
 					ElementData.ElementMetaData.SetLabel(Actor.GetLabel());
 					ElementData.ElementMetaData.SetAssociatedElement(Actor);
 
-					FDocumentData.AddActorMetadata(RevitElement, ElementData.ElementMetaData);
+					FMetadataManager.AddActorMetadata(RevitElement, ElementData.ElementMetaData, true);
 
 					++CurrentBatchSize;
 
diff --git a/Engine/Source/Programs/Enterprise/Datasmith/DatasmithRevitExporter/Private/DatasmithRevitDocumentData.cs b/Engine/Source/Programs/Enterprise/Datasmith/DatasmithRevitExporter/Private/DatasmithRevitDocumentData.cs
--- a/Engine/Source/Programs/Enterprise/Datasmith/DatasmithRevitExporter/Private/DatasmithRevitDocumentData.cs
+++ b/Engine/Source/Programs/Enterprise/Datasmith/DatasmithRevitExporter/Private/DatasmithRevitDocumentData.cs
@@ -589,7 +589,12 @@ namespace DatasmithRevitExporter
 				string LayerName = Category.GetCategory(CurrentElement.Document, BuiltInCategory.OST_LightingFixtureSource)?.Name ?? "Light Sources";
 				SetActorProperties(LayerName, LightActor);
 
-				FDatasmithFacadeMetaData LightMetaData = GetActorMetaData(LightActor);
+				FDatasmithFacadeMetaData LightMetaData = null;
+
+				if (!DocumentData.bSkipMetadataExport)
+				{
+					LightMetaData = GetActorMetaData(LightActor);
+				}
 
 				// Set the Datasmith light actor layer to its predefined name.
 				string CategoryName = Category.GetCategory(CurrentElement.Document, BuiltInCategory.OST_LightingFixtureSource)?.Name ?? "Light Sources";
@@ -1034,12 +1039,12 @@ namespace DatasmithRevitExporter
 				}
 
 				// Add Revit element metadata to the Datasmith actor.
-				FDocumentData.AddActorMetadata(CurrentElement, "Element*", ElementMetaData);
+				FMetadataManager.AddActorMetadata(CurrentElement, "Element*", ElementMetaData, false);
 
 				if (BaseElementType != null)
 				{
 					// Add Revit element type metadata to the Datasmith actor.
-					FDocumentData.AddActorMetadata(BaseElementType, "Type*", ElementMetaData);
+					FMetadataManager.AddActorMetadata(BaseElementType, "Type*", ElementMetaData, false);
 				}
 
 				return ElementMetaData;
@@ -1856,7 +1861,11 @@ namespace DatasmithRevitExporter
 					BasePointMetaData.SetAssociatedElement(BasePointActor);
 
 					BasePointMetaData.AddPropertyVector(MetadataPrefix + "Location", $"{BasePointPosition.X} {BasePointPosition.Y} {BasePointPosition.Z}");
-					FDocumentData.AddActorMetadata(BasePointLocation, MetadataPrefix, BasePointMetaData);
+
+					if (!bSkipMetadataExport)
+					{
+						FMetadataManager.AddActorMetadata(BasePointLocation, MetadataPrefix, BasePointMetaData, false);
+					}
 
 					if (BasePointElement == null)
 					{
@@ -2186,98 +2195,5 @@ namespace DatasmithRevitExporter
 			// Set the world transform of the Datasmith actor.
 			IOActor.SetWorldTransform(worldMatrix);
 		}
-
-		public static ElementType GetElementType(Element InElement)
-		{
-			return InElement.Document.GetElement(InElement.GetTypeId()) as ElementType;
-		}
-
-		public static string GetCategoryName(Element InElement)
-		{
-			ElementType Type = GetElementType(InElement);
-			return Type?.Category?.Name ?? InElement.Category?.Name;
-		}
-
-		public static void AddActorMetadata(
-			Element InElement,
-			FDatasmithFacadeMetaData ActorMetadata
-		)
-		{
-			// Add the Revit element category name metadata to the Datasmith actor.
-			string CategoryName = GetCategoryName(InElement);
-			if (!string.IsNullOrEmpty(CategoryName))
-			{
-				ActorMetadata.AddPropertyString("Element*Category", CategoryName);
-			}
-
-			// Add the Revit element family name metadata to the Datasmith actor.
-			ElementType ElemType = GetElementType(InElement);
-			string FamilyName = ElemType?.FamilyName;
-			if (!string.IsNullOrEmpty(FamilyName))
-			{
-				ActorMetadata.AddPropertyString("Element*Family", FamilyName);
-			}
-
-			// Add the Revit element type name metadata to the Datasmith actor.
-			string TypeName = ElemType?.Name;
-			if (!string.IsNullOrEmpty(TypeName))
-			{
-				ActorMetadata.AddPropertyString("Element*Type", TypeName);
-			}
-
-			// Add Revit element metadata to the Datasmith actor.
-			FDocumentData.AddActorMetadata(InElement, "Element*", ActorMetadata);
-
-			if (ElemType != null)
-			{
-				// Add Revit element type metadata to the Datasmith actor.
-				FDocumentData.AddActorMetadata(ElemType, "Type*", ActorMetadata);
-			}
-		}
-
-		private static void AddActorMetadata(
-			Element InSourceElement,
-			string InMetadataPrefix,
-			FDatasmithFacadeMetaData ElementMetaData
-		)
-		{
-			IList<Parameter> Parameters = InSourceElement.GetOrderedParameters();
-
-			if (Parameters != null)
-			{
-				foreach (Parameter Parameter in Parameters)
-				{
-					if (Parameter.HasValue)
-					{
-						string ParameterValue = Parameter.AsValueString();
-
-						if (string.IsNullOrEmpty(ParameterValue))
-						{
-							switch (Parameter.StorageType)
-							{
-								case StorageType.Integer:
-									ParameterValue = Parameter.AsInteger().ToString();
-									break;
-								case StorageType.Double:
-									ParameterValue = Parameter.AsDouble().ToString();
-									break;
-								case StorageType.String:
-									ParameterValue = Parameter.AsString();
-									break;
-								case StorageType.ElementId:
-									ParameterValue = Parameter.AsElementId().ToString();
-									break;
-							}
-						}
-
-						if (!string.IsNullOrEmpty(ParameterValue))
-						{
-							string MetadataKey = InMetadataPrefix + Parameter.Definition.Name;
-							ElementMetaData.AddPropertyString(MetadataKey, ParameterValue);
-						}
-					}
-				}
-			}
-		}
 	}
 }
diff --git a/Engine/Source/Programs/Enterprise/Datasmith/DatasmithRevitExporter/Private/DatasmithRevitMetadataManager.cs b/Engine/Source/Programs/Enterprise/Datasmith/DatasmithRevitExporter/Private/DatasmithRevitMetadataManager.cs
new file mode 100644
--- /dev/null
+++ b/Engine/Source/Programs/Enterprise/Datasmith/DatasmithRevitExporter/Private/DatasmithRevitMetadataManager.cs
@@ -0,0 +1,342 @@
+// Copyright Epic Games, Inc. All Rights Reserved.
+
+using System;
+using System.Collections.Generic;
+using Autodesk.Revit.DB;
+using Autodesk.Revit.DB.Events;
+using Autodesk.Revit.DB.ExtensibleStorage;
+using Autodesk.Revit.UI;
+
+namespace DatasmithRevitExporter
+{
+	public class FMetadataSettings
+	{
+		private IList<string> _ParamNamesFilter = new List<string>();
+		private IList<int> _ParamGroupsFilter = new List<int>();
+
+		// Derived data for fast lookup
+		private HashSet<string> ParamNamesSet;
+		private HashSet<int> ParamGroupsSet;
+
+		public IList<string> ParamNamesFilter 
+		{
+			get => _ParamNamesFilter;
+			set
+			{
+				_ParamNamesFilter = value;
+				ParamNamesSet = new HashSet<string>(_ParamNamesFilter);
+			}
+		}
+
+		public IList<int> ParamGroupsFilter 
+		{
+			get => _ParamGroupsFilter;
+			set
+			{
+				_ParamGroupsFilter = value;
+				ParamGroupsSet = new HashSet<int>(_ParamGroupsFilter);
+			}
+		}
+
+		public bool ParamterPassesFilter(Parameter InParam)
+		{
+			if (ParamNamesSet != null && ParamNamesSet.Count > 0)
+			{
+				bool bNameMatch = false;
+				foreach (string NameFilter in ParamNamesSet)
+				{
+					// Check if NameFilter is contained in the current param name
+					if (InParam.Definition.Name.IndexOf(NameFilter, StringComparison.OrdinalIgnoreCase) >= 0)
+					{
+						bNameMatch = true;
+						break;
+					}
+				}
+
+				if (!bNameMatch)
+				{
+					return false;
+				}
+			}
+
+			// Name filter passes (or empty), so check group filter.
+			bool bGroupMatch = ParamGroupsSet.Contains((int)InParam.Definition.ParameterGroup);
+
+			return bGroupMatch;
+		}
+	}
+
+	public class FMetadataManager
+	{
+		static class FMatadataSettingsSchema
+		{
+			readonly static Guid SchemaGuid = new Guid("{CF94439A-55C3-46FC-966D-B5D42036AFAF}");
+			
+			public readonly static string ParamNamesField = "ParamNames";
+			public readonly static string ParamGroupsField = "ParamGroups";
+
+			public static Schema GetSchema()
+			{
+				Schema Schema = Schema.Lookup(SchemaGuid);
+				if (Schema != null)
+				{
+					return Schema;
+				}
+
+				SchemaBuilder SchemaBuilder = new SchemaBuilder(SchemaGuid);
+				SchemaBuilder.SetSchemaName("DatasmithRevitMetadataExportSettings");
+
+				SchemaBuilder.AddArrayField(ParamNamesField, typeof(string));
+				SchemaBuilder.AddArrayField(ParamGroupsField, typeof(int));
+
+				return SchemaBuilder.Finish();
+			}
+		}
+
+		static class DataStorageUniqueIdSchema
+		{
+			static readonly Guid SchemaGuid = new Guid("{0A99DE0F-F096-4EA9-9739-219ECDF83F7E}");
+
+			public static Schema GetSchema()
+			{
+				Schema Schema = Schema.Lookup(SchemaGuid);
+				if (Schema != null)
+				{
+					return Schema;
+				}
+
+				SchemaBuilder SchemaBuilder = new SchemaBuilder(SchemaGuid);
+				SchemaBuilder.SetSchemaName("DataStorageUniqueId");
+				SchemaBuilder.AddSimpleField("Id", typeof(Guid));
+
+				return SchemaBuilder.Finish();
+			}
+		}
+
+		static EventHandler<DocumentOpenedEventArgs> DocumentOpenedHandler;
+		static EventHandler<DocumentClosingEventArgs> DocumentClosingHandler;
+		static readonly Guid SettingDataStorageId = new Guid("{6CD2BB42-4A04-4F12-8E7F-5C6723C33791}");
+
+		public static FMetadataSettings CurrentSettings = null;
+
+		public static void Init(UIControlledApplication InApplication)
+		{
+			DocumentOpenedHandler = new EventHandler<DocumentOpenedEventArgs>(OnDocumentOpened);
+			DocumentClosingHandler = new EventHandler<DocumentClosingEventArgs>(OnDocumentClosing);
+			InApplication.ControlledApplication.DocumentOpened += DocumentOpenedHandler;
+			InApplication.ControlledApplication.DocumentClosing += DocumentClosingHandler;
+		}
+
+		public static void Destroy(UIControlledApplication InApplication)
+		{
+			InApplication.ControlledApplication.DocumentClosing -= DocumentClosingHandler;
+			InApplication.ControlledApplication.DocumentOpened -= DocumentOpenedHandler;
+			DocumentClosingHandler = null;
+			DocumentOpenedHandler = null;
+		}
+
+		private static FMetadataSettings ReadSettings(Document Doc)
+		{
+			Entity SettingsEntity = GetSettingsEntity(Doc);
+
+			if (SettingsEntity == null || !SettingsEntity.IsValid())
+			{
+				return null;
+			}
+
+			FMetadataSettings Settings = new FMetadataSettings();
+			Settings.ParamNamesFilter = SettingsEntity.Get<IList<string>>(FMatadataSettingsSchema.ParamNamesField);
+			Settings.ParamGroupsFilter = SettingsEntity.Get<IList<int>>(FMatadataSettingsSchema.ParamGroupsField);
+
+			return Settings;
+		}
+
+		public static void WriteSettings(Document Doc, FMetadataSettings Settings)
+		{
+			using (Transaction T = new Transaction(Doc, "Write settings"))
+			{
+				T.Start();
+
+				DataStorage SettingDataStorage = GetSettingsDataStorage(Doc);
+				if (SettingDataStorage == null)
+				{
+					SettingDataStorage = DataStorage.Create(Doc);
+				}
+
+				Entity SettingsEntity = new Entity(FMatadataSettingsSchema.GetSchema());
+				SettingsEntity.Set(FMatadataSettingsSchema.ParamNamesField, Settings.ParamNamesFilter);
+				SettingsEntity.Set(FMatadataSettingsSchema.ParamGroupsField, Settings.ParamGroupsFilter);
+
+				// Identify settings data storage
+				Entity IdEntity = new Entity(DataStorageUniqueIdSchema.GetSchema());
+				IdEntity.Set("Id", SettingDataStorageId);
+
+				SettingDataStorage.SetEntity(IdEntity);
+				SettingDataStorage.SetEntity(SettingsEntity);
+
+				CurrentSettings = Settings;
+
+				T.Commit();
+			}
+		}
+
+		public static void AddActorMetadata(Element InElement, FDatasmithFacadeMetaData ActorMetadata, bool bInRespectFilter)
+		{
+			// Add the Revit element category name metadata to the Datasmith actor.
+			string CategoryName = GetCategoryName(InElement);
+			if (!string.IsNullOrEmpty(CategoryName))
+			{
+				ActorMetadata.AddPropertyString("Element*Category", CategoryName);
+			}
+
+			// Add the Revit element family name metadata to the Datasmith actor.
+			ElementType ElemType = GetElementType(InElement);
+			string FamilyName = ElemType?.FamilyName;
+			if (!string.IsNullOrEmpty(FamilyName))
+			{
+				ActorMetadata.AddPropertyString("Element*Family", FamilyName);
+			}
+
+			// Add the Revit element type name metadata to the Datasmith actor.
+			string TypeName = ElemType?.Name;
+			if (!string.IsNullOrEmpty(TypeName))
+			{
+				ActorMetadata.AddPropertyString("Element*Type", TypeName);
+			}
+
+			// Add Revit element metadata to the Datasmith actor.
+			AddActorMetadata(InElement, "Element*", ActorMetadata, bInRespectFilter);
+
+			if (ElemType != null)
+			{
+				// Add Revit element type metadata to the Datasmith actor.
+				AddActorMetadata(ElemType, "Type*", ActorMetadata, bInRespectFilter);
+			}
+		}
+
+		public static void AddActorMetadata(
+			Element InSourceElement,
+			string InMetadataPrefix,
+			FDatasmithFacadeMetaData ElementMetaData,
+			bool bInRespectFilter
+		)
+		{
+			IList<Parameter> Parameters = InSourceElement.GetOrderedParameters();
+
+			if (Parameters != null)
+			{
+				foreach (Parameter Parameter in Parameters)
+				{
+					if (Parameter.HasValue)
+					{
+						if (bInRespectFilter && CurrentSettings != null)
+						{
+							if (!CurrentSettings.ParamterPassesFilter(Parameter))
+							{
+								continue; // Skip export of this param
+							}
+						}
+
+						string ParameterValue = Parameter.AsValueString();
+
+						if (string.IsNullOrEmpty(ParameterValue))
+						{
+							switch (Parameter.StorageType)
+							{
+								case StorageType.Integer:
+									ParameterValue = Parameter.AsInteger().ToString();
+									break;
+								case StorageType.Double:
+									ParameterValue = Parameter.AsDouble().ToString();
+									break;
+								case StorageType.String:
+									ParameterValue = Parameter.AsString();
+									break;
+								case StorageType.ElementId:
+									ParameterValue = Parameter.AsElementId().ToString();
+									break;
+							}
+						}
+
+						if (!string.IsNullOrEmpty(ParameterValue))
+						{
+							string MetadataKey = InMetadataPrefix + Parameter.Definition.Name;
+							ElementMetaData.AddPropertyString(MetadataKey, ParameterValue);
+						}
+					}
+				}
+			}
+		}
+
+		static void OnDocumentOpened(object sender, DocumentOpenedEventArgs e)
+		{
+			CurrentSettings = ReadSettings(e.Document);
+		}
+
+		static void OnDocumentClosing(object sender, DocumentClosingEventArgs e)
+		{
+			CurrentSettings = null;
+		}
+
+		private static DataStorage GetSettingsDataStorage(Document Doc)
+		{
+			// Retrieve all data storages from project
+			FilteredElementCollector Collector = new FilteredElementCollector(Doc);
+			var DataStorages = Collector.OfClass(typeof(DataStorage));
+
+			// Find setting data storage
+			foreach (DataStorage DataStorage in DataStorages)
+			{
+				Entity SettingIdEntity = DataStorage.GetEntity(DataStorageUniqueIdSchema.GetSchema());
+
+				if (!SettingIdEntity.IsValid())
+				{
+					continue;
+				}
+
+				Guid Id = SettingIdEntity.Get<Guid>("Id");
+
+				if (!Id.Equals(SettingDataStorageId))
+				{
+					continue;
+				}
+
+				return DataStorage;
+			}
+			return null;
+		}
+
+		private static Entity GetSettingsEntity(Document Doc)
+		{
+			FilteredElementCollector Collector = new FilteredElementCollector(Doc);
+			var DataStorages = Collector.OfClass(typeof(DataStorage));
+
+			// Find setting data storage
+			foreach (DataStorage DataStorage in DataStorages)
+			{
+				Entity SettingEntity = DataStorage.GetEntity(FMatadataSettingsSchema.GetSchema());
+
+				// If a DataStorage contains setting entity, we found it
+				if (!SettingEntity.IsValid())
+				{
+					continue;
+				}
+
+				return SettingEntity;
+			}
+
+			return null;
+		}
+
+		private static string GetCategoryName(Element InElement)
+		{
+			ElementType Type = GetElementType(InElement);
+			return Type?.Category?.Name ?? InElement.Category?.Name;
+		}
+
+		private static ElementType GetElementType(Element InElement)
+		{
+			return InElement.Document.GetElement(InElement.GetTypeId()) as ElementType;
+		}
+	}
+}
diff --git a/Engine/Source/Programs/Enterprise/Datasmith/DatasmithRevitExporter/Private/DatasmithRevitSettingsDialog.cs b/Engine/Source/Programs/Enterprise/Datasmith/DatasmithRevitExporter/Private/DatasmithRevitSettingsDialog.cs
new file mode 100644
--- /dev/null
+++ b/Engine/Source/Programs/Enterprise/Datasmith/DatasmithRevitExporter/Private/DatasmithRevitSettingsDialog.cs
@@ -0,0 +1,290 @@
+﻿// Copyright Epic Games, Inc. All Rights Reserved.
+
+using System;
+using System.Collections.Generic;
+using System.Drawing;
+using System.Text;
+using System.Windows.Forms;
+
+namespace DatasmithRevitExporter
+{
+	public class DatasmithRevitParamGroupsDialog : Form
+	{
+		public static SortedList<string, int> AllBuiltinParamGroups { get; private set; }
+		public CheckedListBox GroupsList { get; private set; }
+
+		public DatasmithRevitParamGroupsDialog()
+		{
+			GroupsList = new CheckedListBox();
+			GroupsList.Size = new Size(310, 350);
+
+			if (AllBuiltinParamGroups == null)
+			{
+				AllBuiltinParamGroups = new SortedList<string, int>();
+
+				foreach (Autodesk.Revit.DB.BuiltInParameterGroup BIP in Enum.GetValues(typeof(Autodesk.Revit.DB.BuiltInParameterGroup)))
+				{
+					string GroupLabel = Autodesk.Revit.DB.LabelUtils.GetLabelFor(BIP);
+
+					if (!AllBuiltinParamGroups.ContainsKey(GroupLabel))
+					{
+						AllBuiltinParamGroups.Add(GroupLabel, (int)BIP);
+					}
+				}
+			}
+
+			foreach (var KV in AllBuiltinParamGroups)
+			{
+				GroupsList.Items.Add(KV.Key);
+			}
+
+			Button OKButton = new Button();
+			OKButton.Name = "OKButton";
+			OKButton.Text = DatasmithRevitResources.Strings.ExportOptionsDialog_ButtonOK;
+			OKButton.Anchor = AnchorStyles.Bottom | AnchorStyles.Right;
+			OKButton.DialogResult = DialogResult.OK;
+			OKButton.Location = new Point(211, 68);
+			OKButton.Margin = new Padding(78, 0, 3, 3);
+			OKButton.Size = new Size(75, 23);
+			OKButton.TabIndex = 4;
+			OKButton.UseVisualStyleBackColor = true;
+
+			TableLayoutPanel DialogLayout = new TableLayoutPanel();
+			DialogLayout.Name = "DialogLayout";
+			DialogLayout.ColumnCount = 1;
+			DialogLayout.ColumnStyles.Add(new ColumnStyle());
+			DialogLayout.RowCount = 2;
+			DialogLayout.RowStyles.Add(new RowStyle());
+			DialogLayout.RowStyles.Add(new RowStyle());
+			DialogLayout.AutoSize = true;
+			DialogLayout.AutoSizeMode = AutoSizeMode.GrowAndShrink;
+			DialogLayout.Dock = DockStyle.Fill;
+			DialogLayout.Location = new Point(10, 10);
+			DialogLayout.Size = new Size(300, 100);
+			DialogLayout.TabIndex = 0;
+
+			DialogLayout.Controls.Add(GroupsList, 0,0);
+			DialogLayout.Controls.Add(OKButton, 0,1);
+
+			Name = "AddGroups";
+			Text = DatasmithRevitResources.Strings.SettingsDialog_AddGroupsDialogTitle; // "Add Groups";
+			AcceptButton = OKButton;
+			AutoScaleDimensions = new SizeF(6F, 13F);
+			AutoScaleMode = AutoScaleMode.Font;
+			AutoSize = true;
+			AutoSizeMode = AutoSizeMode.GrowAndShrink;
+			ClientSize = new Size(309, 114);
+			FormBorderStyle = FormBorderStyle.FixedDialog;
+			MaximizeBox = false;
+			MinimizeBox = false;
+			Padding = new Padding(10);
+			SizeGripStyle = SizeGripStyle.Hide;
+			StartPosition = FormStartPosition.CenterParent;
+
+			Controls.Add(DialogLayout);
+		}
+	}
+
+	public class DatasmithRevitSettingsDialog : Form
+	{
+		private CheckedListBox GroupsList;
+		private TextBox ParamNameTextBox;
+		private FMetadataSettings MetadataSettings;
+		private Autodesk.Revit.DB.Document Document;
+		private SortedList<string, int> AddedBuiltinParamGroups = new SortedList<string, int>();
+		private List<string> AddedParamNames = new List<string>();
+
+		public DatasmithRevitSettingsDialog(Autodesk.Revit.DB.Document InDocument)
+		{
+			Document = InDocument;
+			FormClosing += OnClosing;
+
+			MetadataSettings = FMetadataManager.CurrentSettings;
+
+			ToolTip OptionToolTip = new ToolTip();
+			OptionToolTip.AutoPopDelay = 10000; // milliseconds
+
+			Label DescriptionLabel = new Label();
+			DescriptionLabel.Name = "MatchParameterName";
+			DescriptionLabel.Text = DatasmithRevitResources.Strings.SettingsDialog_LabelMetadataDescription;
+			DescriptionLabel.Anchor = AnchorStyles.Left;
+			DescriptionLabel.AutoSize = false;
+			DescriptionLabel.TabIndex = 0;
+			DescriptionLabel.Font = new Font(FontFamily.GenericSansSerif, 8f);
+			DescriptionLabel.BorderStyle = BorderStyle.FixedSingle;
+			DescriptionLabel.BackColor = Color.Wheat;
+			DescriptionLabel.TextAlign = ContentAlignment.MiddleLeft;
+			DescriptionLabel.Margin = new Padding(0, 10, 0, 20);
+			DescriptionLabel.Width = 385;
+			DescriptionLabel.Height = 55;
+
+			Label ParamNameLabel = new Label();
+			ParamNameLabel.Name = "MatchParameterNames";
+			ParamNameLabel.Text = DatasmithRevitResources.Strings.SettingsDialog_LabelMatchParamNames;
+			ParamNameLabel.Anchor = AnchorStyles.Left;
+			ParamNameLabel.AutoSize = true;
+			ParamNameLabel.TabIndex = 0;
+			ParamNameLabel.TextAlign = ContentAlignment.MiddleLeft;
+			ParamNameLabel.Margin = new Padding(0, 0, 0, 10);
+			OptionToolTip.SetToolTip(ParamNameLabel, "Paramter names to match, separated by new lines (can be partial names).\nCase insensitive.");
+
+			ParamNameTextBox = new TextBox();
+			ParamNameTextBox.Width = 385;
+			ParamNameTextBox.Height = 150;
+			ParamNameTextBox.TabStop = false;
+			ParamNameTextBox.Multiline = true;
+			ParamNameTextBox.Margin = new Padding(0, 0, 0, 20);
+
+			Label GroupsLabel = new Label();
+			GroupsLabel.Name = "GroupsLabel";
+			GroupsLabel.Text = DatasmithRevitResources.Strings.SettingsDialog_LabelMatchGroups;
+			GroupsLabel.Anchor = AnchorStyles.Left;
+			GroupsLabel.AutoSize = true;
+
+			GroupsList = new CheckedListBox();
+			GroupsList.Size = new Size(380, 200);
+
+			Button AddGroupButton = new Button();
+			AddGroupButton.Name = "Add";
+			AddGroupButton.Text = DatasmithRevitResources.Strings.SettingsDialog_ButtonAddGroups;
+			AddGroupButton.AutoSize = true;
+			AddGroupButton.Anchor = AnchorStyles.Left;
+			AddGroupButton.UseVisualStyleBackColor = true;
+			AddGroupButton.Click += (s, e) =>
+			{
+				DatasmithRevitParamGroupsDialog Dlg = new DatasmithRevitParamGroupsDialog();
+				if (DialogResult.OK == Dlg.ShowDialog())
+				{
+					bool bAddedNewGroups = false;
+					foreach (int Index in Dlg.GroupsList.CheckedIndices)
+					{
+						string GroupLabel = DatasmithRevitParamGroupsDialog.AllBuiltinParamGroups.Keys[Index];
+						if (!AddedBuiltinParamGroups.ContainsKey(GroupLabel))
+						{
+							int GroupId;
+							if (DatasmithRevitParamGroupsDialog.AllBuiltinParamGroups.TryGetValue(GroupLabel, out GroupId))
+							{
+								AddedBuiltinParamGroups.Add(GroupLabel, GroupId);
+								bAddedNewGroups = true;
+							}
+						}
+					}
+
+					if (bAddedNewGroups)
+					{
+						ReloadAddedGroupsList();
+					}
+				}
+			};
+
+			Button RemoveGroupButton = new Button();
+			RemoveGroupButton.Name = "Remove";
+			RemoveGroupButton.Text = DatasmithRevitResources.Strings.SettingsDialog_ButtonRemoveGroups;
+			RemoveGroupButton.AutoSize = true;
+			RemoveGroupButton.Anchor = AnchorStyles.Right;
+			RemoveGroupButton.UseVisualStyleBackColor = true;
+			RemoveGroupButton.Click += (s, e) =>
+			{
+				if (GroupsList.CheckedIndices.Count > 0)
+				{
+					int[] IndicesToRemove = new int[GroupsList.CheckedIndices.Count];
+					GroupsList.CheckedIndices.CopyTo(IndicesToRemove, 0);
+					Array.Sort(IndicesToRemove, (A, B) => B.CompareTo(A));
+
+					foreach (int Index in IndicesToRemove)
+					{
+						GroupsList.Items.RemoveAt(Index);
+						AddedBuiltinParamGroups.RemoveAt(Index);
+					}
+				}
+			};
+
+			TableLayoutPanel DialogLayout = new TableLayoutPanel();
+			DialogLayout.Name = "DialogLayout";
+			DialogLayout.ColumnCount = 2;
+			DialogLayout.ColumnStyles.Add(new ColumnStyle());
+			DialogLayout.ColumnStyles.Add(new ColumnStyle());
+			DialogLayout.RowCount = 6;
+			DialogLayout.RowStyles.Add(new RowStyle());
+			DialogLayout.RowStyles.Add(new RowStyle());
+			DialogLayout.RowStyles.Add(new RowStyle());
+			DialogLayout.RowStyles.Add(new RowStyle());
+			DialogLayout.RowStyles.Add(new RowStyle());
+			DialogLayout.AutoSize = true;
+			DialogLayout.AutoSizeMode = AutoSizeMode.GrowAndShrink;
+			DialogLayout.Dock = DockStyle.Fill;
+			DialogLayout.Location = new Point(10, 10);
+			DialogLayout.TabIndex = 0;
+
+			DialogLayout.Controls.Add(DescriptionLabel, 0,0);
+			DialogLayout.Controls.Add(ParamNameLabel, 0, 1);
+			DialogLayout.Controls.Add(ParamNameTextBox, 0, 2);
+			DialogLayout.Controls.Add(GroupsLabel, 0, 3);
+			DialogLayout.Controls.Add(GroupsList, 0, 4);
+			DialogLayout.Controls.Add(AddGroupButton, 0, 5);
+			DialogLayout.Controls.Add(RemoveGroupButton, 1, 5);
+
+			DialogLayout.SetColumnSpan(ParamNameTextBox, 2);
+			DialogLayout.SetColumnSpan(GroupsList, 2);
+			DialogLayout.SetColumnSpan(DescriptionLabel, 2);
+
+			Name = "MetadataExportFilter";
+			Text = DatasmithRevitResources.Strings.SettingsDialog_DialogTitle;
+			AutoScaleDimensions = new SizeF(6F, 13F);
+			AutoScaleMode = AutoScaleMode.Font;
+			AutoSize = true;
+			AutoSizeMode = AutoSizeMode.GrowAndShrink;
+			ClientSize = new Size(309, 114);
+			FormBorderStyle = FormBorderStyle.FixedDialog;
+			MaximizeBox = false;
+			MinimizeBox = false;
+			Padding = new Padding(10);
+			SizeGripStyle = SizeGripStyle.Hide;
+			StartPosition = FormStartPosition.CenterParent;
+
+			Controls.Add(DialogLayout);
+
+			if (MetadataSettings == null)
+			{
+				MetadataSettings = new FMetadataSettings();
+			}
+			else
+			{
+				StringBuilder SB = new StringBuilder();
+				foreach (string ParamNameFilter in MetadataSettings.ParamNamesFilter)
+				{
+					if (ParamNameFilter.Length > 0)
+					{
+						SB.AppendLine(ParamNameFilter);
+					}
+				}
+				ParamNameTextBox.Text = SB.ToString();
+
+				foreach (int Group in MetadataSettings.ParamGroupsFilter)
+				{
+					string GroupLabel = Autodesk.Revit.DB.LabelUtils.GetLabelFor((Autodesk.Revit.DB.BuiltInParameterGroup)Group);
+					AddedBuiltinParamGroups.Add(GroupLabel, Group);
+				}
+
+				ReloadAddedGroupsList();
+			}
+		}
+
+		private void ReloadAddedGroupsList()
+		{
+			GroupsList.Items.Clear();
+
+			foreach (var KV in AddedBuiltinParamGroups)
+			{
+				GroupsList.Items.Add(KV.Key);
+			}
+		}
+
+		private void OnClosing(object InSender, FormClosingEventArgs InArgs)
+		{
+			MetadataSettings.ParamNamesFilter = ParamNameTextBox.Text.Split(new[] { Environment.NewLine }, StringSplitOptions.RemoveEmptyEntries);
+			MetadataSettings.ParamGroupsFilter = AddedBuiltinParamGroups.Values;
+			FMetadataManager.WriteSettings(Document, MetadataSettings);
+		}
+	}
+}
