Datasmith Solidworks: Export animations

diff --git a/Engine/Source/Developer/Datasmith/DatasmithFacade/Private/DatasmithFacadeAnimation.cpp b/Engine/Source/Developer/Datasmith/DatasmithFacade/Private/DatasmithFacadeAnimation.cpp
new file mode 100644
--- /dev/null
+++ b/Engine/Source/Developer/Datasmith/DatasmithFacade/Private/DatasmithFacadeAnimation.cpp
@@ -0,0 +1,328 @@
+// Copyright Epic Games, Inc. All Rights Reserved.
+
+
+#include "DatasmithFacadeAnimation.h"
+#include "DatasmithFacadeScene.h"
+
+// FDatasmithFacadeBaseAnimation
+
+FDatasmithFacadeBaseAnimation::FDatasmithFacadeBaseAnimation(const TSharedRef<IDatasmithBaseAnimationElement>& InInternalAnimation)
+	: FDatasmithFacadeElement( InInternalAnimation )
+{
+}
+
+bool FDatasmithFacadeBaseAnimation::IsSubType( const EDatasmithFacadeElementAnimationSubType AnimSubType ) const
+{
+	return GetDatasmithBaseAnimation()->IsSubType(static_cast<EDatasmithElementAnimationSubType>(AnimSubType));
+}
+
+void FDatasmithFacadeBaseAnimation::SetCompletionMode(EDatasmithFacadeCompletionMode CompletionMode)
+{
+	GetDatasmithBaseAnimation()->SetCompletionMode(static_cast<EDatasmithCompletionMode>(CompletionMode));
+}
+
+EDatasmithFacadeCompletionMode FDatasmithFacadeBaseAnimation::GetCompletionMode() const
+{
+	return static_cast<EDatasmithFacadeCompletionMode>(GetDatasmithBaseAnimation()->GetCompletionMode());
+}
+
+TSharedRef<IDatasmithBaseAnimationElement> FDatasmithFacadeBaseAnimation::GetDatasmithBaseAnimation() const
+{
+	return StaticCastSharedRef<IDatasmithBaseAnimationElement>( InternalDatasmithElement );
+}
+
+// FDatasmithFacadeTransformAnimation
+
+FDatasmithFacadeTransformAnimation::FDatasmithFacadeTransformAnimation(const TCHAR* InName)
+	: FDatasmithFacadeBaseAnimation( FDatasmithSceneFactory::CreateTransformAnimation(InName) )
+{
+}
+
+FDatasmithFacadeTransformAnimation::FDatasmithFacadeTransformAnimation(
+	const TSharedRef<IDatasmithTransformAnimationElement>& InInternalAnimation
+) : FDatasmithFacadeBaseAnimation( InInternalAnimation )
+{
+}
+
+TSharedRef<IDatasmithTransformAnimationElement> FDatasmithFacadeTransformAnimation::GetDatasmithTransformAnimation() const
+{
+	return StaticCastSharedRef<IDatasmithTransformAnimationElement>( InternalDatasmithElement );
+}
+
+void FDatasmithFacadeTransformAnimation::AddFrame(EDatasmithFacadeAnimationTransformType TransformType, int32 FrameNumber, float X, float Y, float Z)
+{
+	GetDatasmithTransformAnimation()->AddFrame(
+		static_cast<EDatasmithTransformType>(TransformType),
+		FDatasmithTransformFrameInfo(FrameNumber, X, Y, Z)
+	);
+}
+
+int32 FDatasmithFacadeTransformAnimation::GetFramesCount(EDatasmithFacadeAnimationTransformType TransformType) const
+{
+	return GetDatasmithTransformAnimation()->GetFramesCount(static_cast<EDatasmithTransformType>(TransformType));
+}
+
+void FDatasmithFacadeTransformAnimation::SetCurveInterpMode(EDatasmithFacadeAnimationTransformType TransformType, EDatasmithFacadeCurveInterpMode CurveInterpMode)
+{
+	GetDatasmithTransformAnimation()->SetCurveInterpMode(
+		static_cast<EDatasmithTransformType>(TransformType),
+		static_cast<EDatasmithCurveInterpMode>(CurveInterpMode)
+	);
+}
+
+EDatasmithFacadeCurveInterpMode FDatasmithFacadeTransformAnimation::GetCurveInterpMode(EDatasmithFacadeAnimationTransformType TransformType) const
+{
+	return static_cast<EDatasmithFacadeCurveInterpMode>(GetDatasmithTransformAnimation()->GetCurveInterpMode(static_cast<EDatasmithTransformType>(TransformType)));
+}
+
+void FDatasmithFacadeTransformAnimation::GetFrame(EDatasmithFacadeAnimationTransformType TransformType, int32 Index, int32& OutFrameNumber, float& OutX, float& OutY, float& OutZ) const
+{
+	const FDatasmithTransformFrameInfo& FrameInfo = GetDatasmithTransformAnimation()->GetFrame(static_cast<EDatasmithTransformType>(TransformType), Index);
+	OutFrameNumber = FrameInfo.FrameNumber;
+	OutX = FrameInfo.X;
+	OutY = FrameInfo.Y;
+	OutZ = FrameInfo.Z;
+}
+
+void FDatasmithFacadeTransformAnimation::RemoveFrame(EDatasmithFacadeAnimationTransformType TransformType, int32 Index)
+{
+	GetDatasmithTransformAnimation()->RemoveFrame(static_cast<EDatasmithTransformType>(TransformType), Index);
+}
+
+EDatasmithFacadeTransformChannels FDatasmithFacadeTransformAnimation::GetEnabledTransformChannels() const
+{
+	return static_cast<EDatasmithFacadeTransformChannels>(GetDatasmithTransformAnimation()->GetEnabledTransformChannels());
+}
+
+void FDatasmithFacadeTransformAnimation::SetEnabledTransformChannels(EDatasmithFacadeTransformChannels Channels)
+{
+	GetDatasmithTransformAnimation()->SetEnabledTransformChannels(static_cast<EDatasmithTransformChannels>(Channels));
+}
+
+// FDatasmithFacadeVisibilityAnimation
+
+FDatasmithFacadeVisibilityAnimation::FDatasmithFacadeVisibilityAnimation(const TCHAR* InName)
+	: FDatasmithFacadeBaseAnimation( FDatasmithSceneFactory::CreateVisibilityAnimation(InName) )
+{
+}
+
+FDatasmithFacadeVisibilityAnimation::FDatasmithFacadeVisibilityAnimation(
+	const TSharedRef<IDatasmithVisibilityAnimationElement>& InInternalAnimation
+) : FDatasmithFacadeBaseAnimation( InInternalAnimation )
+{
+}
+
+TSharedRef<IDatasmithVisibilityAnimationElement> FDatasmithFacadeVisibilityAnimation::GetDatasmithVisibilityAnimation() const
+{
+	return StaticCastSharedRef<IDatasmithVisibilityAnimationElement>( InternalDatasmithElement );
+}
+
+void FDatasmithFacadeVisibilityAnimation::AddFrame(int32 FrameNumber, bool bVisible)
+{
+	GetDatasmithVisibilityAnimation()->AddFrame(FDatasmithVisibilityFrameInfo(FrameNumber, bVisible));
+}
+
+int32 FDatasmithFacadeVisibilityAnimation::GetFramesCount() const
+{
+	return GetDatasmithVisibilityAnimation()->GetFramesCount();
+}
+
+void FDatasmithFacadeVisibilityAnimation::SetCurveInterpMode(EDatasmithFacadeCurveInterpMode CurveInterpMode)
+{
+	GetDatasmithVisibilityAnimation()->SetCurveInterpMode(static_cast<EDatasmithCurveInterpMode>(CurveInterpMode));
+}
+
+EDatasmithFacadeCurveInterpMode FDatasmithFacadeVisibilityAnimation::GetCurveInterpMode() const
+{
+	return static_cast<EDatasmithFacadeCurveInterpMode>( GetDatasmithVisibilityAnimation()->GetCurveInterpMode() );
+}
+
+void FDatasmithFacadeVisibilityAnimation::GetFrame(int32 Index, int32& OutFrameNumber, bool& bOutVisible) const
+{
+	const FDatasmithVisibilityFrameInfo& VisbilityInfo = GetDatasmithVisibilityAnimation()->GetFrame(Index);
+	OutFrameNumber = VisbilityInfo.FrameNumber;
+	bOutVisible = VisbilityInfo.bVisible;
+}
+
+void FDatasmithFacadeVisibilityAnimation::RemoveFrame(int32 Index)
+{
+	GetDatasmithVisibilityAnimation()->RemoveFrame(Index);
+}
+
+void FDatasmithFacadeVisibilityAnimation::SetPropagateToChildren(bool bPropagate)
+{
+	GetDatasmithVisibilityAnimation()->SetPropagateToChildren(bPropagate);
+}
+
+bool FDatasmithFacadeVisibilityAnimation::GetPropagateToChildren() const
+{
+	return GetDatasmithVisibilityAnimation()->GetPropagateToChildren();
+}
+
+// FDatasmithFacadeSubsequenceAnimation
+
+FDatasmithFacadeSubsequenceAnimation::FDatasmithFacadeSubsequenceAnimation(const TCHAR* InName)
+	: FDatasmithFacadeSubsequenceAnimation( FDatasmithSceneFactory::CreateSubsequenceAnimation(InName) )
+{
+}
+
+FDatasmithFacadeSubsequenceAnimation::FDatasmithFacadeSubsequenceAnimation(
+	const TSharedRef<IDatasmithSubsequenceAnimationElement>& InInternalAnimation
+) : FDatasmithFacadeBaseAnimation( InInternalAnimation )
+{
+}
+
+TSharedRef<IDatasmithSubsequenceAnimationElement> FDatasmithFacadeSubsequenceAnimation::GetDatasmithSubsequenceAnimation() const 
+{
+	return StaticCastSharedRef<IDatasmithSubsequenceAnimationElement>(InternalDatasmithElement);
+}
+
+int32 FDatasmithFacadeSubsequenceAnimation::GetStartTime() const
+{
+	const FFrameNumber FrameNum = GetDatasmithSubsequenceAnimation()->GetStartTime();
+	return FrameNum.Value;
+}
+
+void FDatasmithFacadeSubsequenceAnimation::SetStartTime(int32 InFrameNumber)
+{
+	GetDatasmithSubsequenceAnimation()->SetStartTime(FFrameNumber(InFrameNumber));
+}
+
+int32 FDatasmithFacadeSubsequenceAnimation::GetDuration() const
+{
+	return GetDatasmithSubsequenceAnimation()->GetDuration();
+}
+
+void FDatasmithFacadeSubsequenceAnimation::SetDuration(int32 InDuration)
+{
+	GetDatasmithSubsequenceAnimation()->SetDuration(InDuration);
+}
+
+float FDatasmithFacadeSubsequenceAnimation::GetTimeScale() const
+{
+	return GetDatasmithSubsequenceAnimation()->GetTimeScale();
+}
+
+void FDatasmithFacadeSubsequenceAnimation::SetTimeScale(float InTimeScale)
+{
+	GetDatasmithSubsequenceAnimation()->SetTimeScale(InTimeScale);
+}
+
+FDatasmithFacadeLevelSequence* FDatasmithFacadeSubsequenceAnimation::GetNewSubsequence() const
+{
+	TWeakPtr<IDatasmithLevelSequenceElement> LevelSeq = GetDatasmithSubsequenceAnimation()->GetSubsequence();
+
+	if (LevelSeq.IsValid())
+	{
+		return new FDatasmithFacadeLevelSequence(LevelSeq.Pin().ToSharedRef());
+	}
+	return nullptr;
+}
+
+void FDatasmithFacadeSubsequenceAnimation::SetSubsequence(FDatasmithFacadeLevelSequence* InSubsequence)
+{
+	GetDatasmithSubsequenceAnimation()->SetSubsequence(InSubsequence->GetDatasmithLevelSequence());
+}
+
+// FDatasmithFacadeLevelSequence
+
+FDatasmithFacadeLevelSequence::FDatasmithFacadeLevelSequence(const TCHAR* InName)
+	: FDatasmithFacadeElement(FDatasmithSceneFactory::CreateLevelSequence(InName))
+{
+}
+
+FDatasmithFacadeLevelSequence::FDatasmithFacadeLevelSequence(
+	const TSharedRef<IDatasmithLevelSequenceElement>& InInternalLevelSequence
+) : FDatasmithFacadeElement( InInternalLevelSequence )
+{
+}
+
+TSharedRef<IDatasmithLevelSequenceElement> FDatasmithFacadeLevelSequence::GetDatasmithLevelSequence() const
+{
+	return StaticCastSharedRef<IDatasmithLevelSequenceElement>( InternalDatasmithElement );
+}
+
+const TCHAR* FDatasmithFacadeLevelSequence::GetFile() const
+{
+	return GetDatasmithLevelSequence()->GetFile();
+}
+
+void FDatasmithFacadeLevelSequence::SetFile(const TCHAR* InFile)
+{
+	GetDatasmithLevelSequence()->SetFile(InFile);
+}
+
+void FDatasmithFacadeLevelSequence::GetFileHash(TCHAR OutBuffer[33], size_t BufferSize) const
+{
+	FString HashString = LexToString(GetDatasmithLevelSequence()->GetFileHash());
+	FCString::Strncpy(OutBuffer, *HashString, BufferSize);
+}
+
+void FDatasmithFacadeLevelSequence::SetFileHash(const TCHAR* Hash)
+{
+	FMD5Hash Md5Hash;
+	LexFromString(Md5Hash, Hash);
+
+	GetDatasmithLevelSequence()->SetFileHash(Md5Hash);
+}
+
+float FDatasmithFacadeLevelSequence::GetFrameRate() const
+{
+	return GetDatasmithLevelSequence()->GetFrameRate();
+}
+
+void FDatasmithFacadeLevelSequence::SetFrameRate(float FramePerSecs)
+{
+	GetDatasmithLevelSequence()->SetFrameRate(FramePerSecs);
+}
+
+void FDatasmithFacadeLevelSequence::AddAnimation(FDatasmithFacadeTransformAnimation* InAnimation)
+{
+	if (InAnimation)
+	{
+		GetDatasmithLevelSequence()->AddAnimation(InAnimation->GetDatasmithTransformAnimation());
+	}
+}
+
+int32 FDatasmithFacadeLevelSequence::GetAnimationsCount() const
+{
+	return GetDatasmithLevelSequence()->GetAnimationsCount();
+}
+
+FDatasmithFacadeTransformAnimation* FDatasmithFacadeLevelSequence::GetNewTransformAnimation(int32 InIndex)
+{
+	if (TSharedPtr<IDatasmithTransformAnimationElement> AnimElement = StaticCastSharedPtr<IDatasmithTransformAnimationElement>( GetDatasmithLevelSequence()->GetAnimation(InIndex) ))
+	{
+		return new FDatasmithFacadeTransformAnimation(AnimElement.ToSharedRef());
+	}
+
+	return nullptr;
+}
+
+FDatasmithFacadeVisibilityAnimation* FDatasmithFacadeLevelSequence::GetNewVisibilityAnimation(int32 InIndex) 
+{
+	if (TSharedPtr<IDatasmithVisibilityAnimationElement> AnimElement = StaticCastSharedPtr<IDatasmithVisibilityAnimationElement>(GetDatasmithLevelSequence()->GetAnimation(InIndex))) 
+	{
+		return new FDatasmithFacadeVisibilityAnimation(AnimElement.ToSharedRef());
+	}
+
+	return nullptr;
+}
+
+FDatasmithFacadeSubsequenceAnimation* FDatasmithFacadeLevelSequence::GetNewSubsequenceAnimation(int32 InIndex) 
+{
+	if (TSharedPtr<IDatasmithSubsequenceAnimationElement> AnimElement = StaticCastSharedPtr<IDatasmithSubsequenceAnimationElement>(GetDatasmithLevelSequence()->GetAnimation(InIndex))) 
+	{
+		return new FDatasmithFacadeSubsequenceAnimation(AnimElement.ToSharedRef());
+	}
+
+	return nullptr;
+}
+
+void FDatasmithFacadeLevelSequence::RemoveAnimation(FDatasmithFacadeBaseAnimation* InAnimation)
+{
+	if (InAnimation)
+	{
+		GetDatasmithLevelSequence()->RemoveAnimation(InAnimation->GetDatasmithBaseAnimation());
+	}
+}
diff --git a/Engine/Source/Developer/Datasmith/DatasmithFacade/Private/DatasmithFacadeScene.cpp b/Engine/Source/Developer/Datasmith/DatasmithFacade/Private/DatasmithFacadeScene.cpp
--- a/Engine/Source/Developer/Datasmith/DatasmithFacade/Private/DatasmithFacadeScene.cpp
+++ b/Engine/Source/Developer/Datasmith/DatasmithFacade/Private/DatasmithFacadeScene.cpp
@@ -10,6 +10,7 @@
 #include "DatasmithFacadeMetaData.h"
 #include "DatasmithFacadeTexture.h"
 #include "DatasmithFacadeVariant.h"
+#include "DatasmithFacadeAnimation.h"
 
 // Datasmith SDK.
 #include "DatasmithExporterManager.h"
@@ -222,6 +223,40 @@ void FDatasmithFacadeScene::RemoveTexture(
 	SceneRef->RemoveTexture(InTexturePtr->GetDatasmithTextureElement());
 }
 
+void FDatasmithFacadeScene::AddLevelSequence(
+	FDatasmithFacadeLevelSequence* InLevelSequence
+)
+{
+	if (InLevelSequence)
+	{
+		SceneRef->AddLevelSequence(InLevelSequence->GetDatasmithLevelSequence());
+	}
+}
+
+int32 FDatasmithFacadeScene::GetLevelSequencesCount() const
+{
+	return SceneRef->GetLevelSequencesCount();
+}
+
+FDatasmithFacadeLevelSequence* FDatasmithFacadeScene::GetNewLevelSequence(
+	int32 LevelSequenceIndex
+)
+{
+	if (TSharedPtr<IDatasmithLevelSequenceElement> LevelSequenceElement = SceneRef->GetLevelSequence(LevelSequenceIndex))
+	{
+		return new FDatasmithFacadeLevelSequence(LevelSequenceElement.ToSharedRef());
+	}
+
+	return nullptr;
+}
+
+void FDatasmithFacadeScene::RemoveLevelSequence(
+	FDatasmithFacadeLevelSequence* InLevelSequence
+)
+{
+	SceneRef->RemoveLevelSequence(InLevelSequence->GetDatasmithLevelSequence());
+}
+
 void FDatasmithFacadeScene::AddLevelVariantSets(
 	FDatasmithFacadeLevelVariantSets* InLevelVariantSetsPtr
 )
diff --git a/Engine/Source/Developer/Datasmith/DatasmithFacade/Public/DatasmithFacadeAnimation.h b/Engine/Source/Developer/Datasmith/DatasmithFacade/Public/DatasmithFacadeAnimation.h
new file mode 100644
--- /dev/null
+++ b/Engine/Source/Developer/Datasmith/DatasmithFacade/Public/DatasmithFacadeAnimation.h
@@ -0,0 +1,257 @@
+// Copyright Epic Games, Inc. All Rights Reserved.
+
+#pragma once
+
+#include "DatasmithFacadeElement.h"
+
+enum class EDatasmithFacadeCurveInterpMode
+{
+	Linear,
+	Constant,
+	Cubic
+};
+
+enum class EDatasmithFacadeAnimationTransformType
+{
+	Translation,
+	Rotation,
+	Scale
+};
+
+enum class EDatasmithFacadeElementAnimationSubType : uint64
+{
+	BaseAnimation = 0,
+	TransformAnimation = 1 << 0,
+	VisibilityAnimation = 1 << 1,
+	SubsequenceAnimation = 1 << 2,
+};
+
+enum class EDatasmithFacadeCompletionMode : uint8
+{
+	KeepState,
+	RestoreState,
+	ProjectDefault,
+};
+
+enum class EDatasmithFacadeTransformChannels : uint16
+{
+	None			= 0x000,
+
+	TranslationX 	= 0x001,
+	TranslationY 	= 0x002,
+	TranslationZ 	= 0x004,
+	Translation 	= TranslationX | TranslationY | TranslationZ,
+
+	RotationX 		= 0x008,
+	RotationY 		= 0x010,
+	RotationZ 		= 0x020,
+	Rotation 		= RotationX | RotationY | RotationZ,
+
+	ScaleX 			= 0x040,
+	ScaleY 			= 0x080,
+	ScaleZ 			= 0x100,
+	Scale 			= ScaleX | ScaleY | ScaleZ,
+
+	All				= Translation | Rotation | Scale,
+};
+
+class DATASMITHFACADE_API FDatasmithFacadeBaseAnimation : public FDatasmithFacadeElement
+{
+public:
+	bool IsSubType( const EDatasmithFacadeElementAnimationSubType AnimSubType ) const;
+
+	/** Set how the actor should behave once its animation completes */
+	void SetCompletionMode(EDatasmithFacadeCompletionMode CompletionMode);
+
+	/** Get how the actor behaves once this animation is complete */
+	EDatasmithFacadeCompletionMode GetCompletionMode() const;
+
+#ifdef SWIG_FACADE
+protected:
+#endif
+
+	explicit FDatasmithFacadeBaseAnimation(
+		const TSharedRef<IDatasmithBaseAnimationElement>& InInternalAnimation
+	);
+
+	TSharedRef<IDatasmithBaseAnimationElement> GetDatasmithBaseAnimation() const;
+};
+
+class DATASMITHFACADE_API FDatasmithFacadeTransformAnimation : public FDatasmithFacadeBaseAnimation
+{
+public:
+	FDatasmithFacadeTransformAnimation(const TCHAR* InName);
+
+	/** Add a frame of the given transform type to the animation */
+	void AddFrame(EDatasmithFacadeAnimationTransformType TransformType, int32 FrameNumber, float X, float Y, float Z);
+
+	/** Return the number of frames of the given transform type in the animation */
+	int32 GetFramesCount(EDatasmithFacadeAnimationTransformType TransformType) const;
+
+	/** Set the interpolation mode of the given transform type in the animation */
+	void SetCurveInterpMode(EDatasmithFacadeAnimationTransformType TransformType, EDatasmithFacadeCurveInterpMode CurveInterpMode);
+
+	/** Get the interpolation mode of the given transform type in the animation */
+	EDatasmithFacadeCurveInterpMode GetCurveInterpMode(EDatasmithFacadeAnimationTransformType TransformType) const;
+
+	/** Return the frame of the given transform type at the given index or an invalid frame if the index was out of bounds */
+	void GetFrame(EDatasmithFacadeAnimationTransformType TransformType, int32 Index, int32& OutFrameNumber, float& OutX, float& OutY, float& OutZ) const;
+
+	/** Remove the frame of the given transform type at the given index from the animation */
+	void RemoveFrame(EDatasmithFacadeAnimationTransformType TransformType, int32 Index);
+
+	/** Gets which channels of this animation will be moved to the animation assets or serialized. All channels are enabled by default */
+	EDatasmithFacadeTransformChannels GetEnabledTransformChannels() const;
+
+	/** Sets which channels of this animation will be moved to the animation assets or serialized. All channels are enabled by default */
+	void SetEnabledTransformChannels(EDatasmithFacadeTransformChannels Channels);
+
+#ifdef SWIG_FACADE
+protected:
+#endif
+
+	explicit FDatasmithFacadeTransformAnimation(
+		const TSharedRef<IDatasmithTransformAnimationElement>& InInternalAnimation
+	);
+
+	TSharedRef<IDatasmithTransformAnimationElement> GetDatasmithTransformAnimation() const;
+};
+
+class DATASMITHFACADE_API FDatasmithFacadeVisibilityAnimation : public FDatasmithFacadeBaseAnimation
+{
+public:
+	FDatasmithFacadeVisibilityAnimation(const TCHAR* InName);
+
+	/** Add a frame to the animation */
+	void AddFrame(int32 FrameNumber, bool bVisible);
+
+	/** Return the number of frames of the animation */
+	int32 GetFramesCount() const;
+
+	/** Set the interpolation mode of the animation */
+	void SetCurveInterpMode(EDatasmithFacadeCurveInterpMode CurveInterpMode);
+
+	/** Get the interpolation mode of the animation */
+	EDatasmithFacadeCurveInterpMode GetCurveInterpMode() const;
+
+	/** Return the frame of the animation at the given index or an invalid frame if the index was out of bounds */
+	void GetFrame(int32 Index, int32& OutFrameNumber, bool& bOutVisible) const;
+
+	/** Remove the frame at the given index from the animation */
+	void RemoveFrame(int32 Index);
+
+	/** Sets whether this animation will be duplicated to all children (recursively) when imported */
+	void SetPropagateToChildren(bool bPropagate);
+
+	/** Gets whether this animation will be duplicated to all children (recursively) when imported */
+	bool GetPropagateToChildren() const;
+
+#ifdef SWIG_FACADE
+protected:
+#endif
+
+	explicit FDatasmithFacadeVisibilityAnimation(
+		const TSharedRef<IDatasmithVisibilityAnimationElement>& InInternalAnimation
+	);
+
+	TSharedRef<IDatasmithVisibilityAnimationElement> GetDatasmithVisibilityAnimation() const;
+};
+
+class FDatasmithFacadeLevelSequence;
+
+class DATASMITHFACADE_API FDatasmithFacadeSubsequenceAnimation : public FDatasmithFacadeBaseAnimation
+{
+public:
+	FDatasmithFacadeSubsequenceAnimation(const TCHAR* InName);
+
+	/** Get the frame where the subsequence starts */
+	int32 GetStartTime() const;
+
+	/** Set the frame where the subsequence starts */
+	void SetStartTime(int32 InFrameNumber);
+
+	/** Get the subsequence duration in number of frames */
+	int32 GetDuration() const;
+
+	/** Set the subsequence duration in number of frames */
+	void SetDuration(int32 InDuration);
+
+	/** Get the time scale used for the subsequence */
+	float GetTimeScale() const;
+
+	/** Set the time scale used for the subsequence */
+	void SetTimeScale(float InTimeScale);
+
+	/** Get a pointer to the subsequence that this element references */
+	FDatasmithFacadeLevelSequence* GetNewSubsequence() const;
+
+	/** Set the subsequence that this element references */
+	void SetSubsequence(FDatasmithFacadeLevelSequence* InSubsequence);
+
+#ifdef SWIG_FACADE
+protected:
+#endif
+
+	explicit FDatasmithFacadeSubsequenceAnimation(
+		const TSharedRef<IDatasmithSubsequenceAnimationElement>& InInternalAnimation
+	);
+
+	TSharedRef<IDatasmithSubsequenceAnimationElement> GetDatasmithSubsequenceAnimation() const;
+};
+
+class DATASMITHFACADE_API FDatasmithFacadeLevelSequence : public FDatasmithFacadeElement
+{
+public:
+	FDatasmithFacadeLevelSequence(const TCHAR* InName);
+
+	/** Get the output filename, it can be absolute or relative to the scene file */
+	const TCHAR* GetFile() const;
+
+	/** Set the output filename, it can be absolute or relative to the scene file */
+	void SetFile(const TCHAR* InFile);
+
+	/** Return a MD5 hash of the content of the Level Sequence Element. Used in CalculateElementHash to quickly identify Element with identical content */
+	void GetFileHash(TCHAR OutBuffer[33], size_t BufferSize) const;
+
+	/** Set the MD5 hash of the Level Sequence file. This should be a hash of its content. */
+	void SetFileHash(const TCHAR* Hash);
+
+	/* Gets the frame rate for the animations in the level sequence */
+	float GetFrameRate() const;
+
+	/* Sets the frame rate for the animations in the level sequence */
+	void SetFrameRate(float FramePerSecs);
+
+	/**
+	 * Adds an animation to the level sequence.
+	 *
+	 * @param InAnimation the animation to add
+	 */
+	void AddAnimation(FDatasmithFacadeTransformAnimation* InAnimation);
+
+	/** Returns the number of animations in the level sequence */
+	int32 GetAnimationsCount() const;
+
+	/** Returns the animation using this index */
+	FDatasmithFacadeTransformAnimation* GetNewTransformAnimation(int32 InIndex);
+	FDatasmithFacadeVisibilityAnimation* GetNewVisibilityAnimation(int32 InIndex);
+	FDatasmithFacadeSubsequenceAnimation* GetNewSubsequenceAnimation(int32 InIndex);
+
+	/**
+	 * Removes an animation from the level sequence.
+	 *
+	 * @param InAnimation the animation to remove
+	 */
+	void RemoveAnimation(FDatasmithFacadeBaseAnimation* InAnimation);
+
+#ifdef SWIG_FACADE
+protected:
+#endif
+
+	explicit FDatasmithFacadeLevelSequence(
+		const TSharedRef<IDatasmithLevelSequenceElement>& InInternalLevelSequence
+	);
+
+	TSharedRef<IDatasmithLevelSequenceElement> GetDatasmithLevelSequence() const;
+};
+
diff --git a/Engine/Source/Developer/Datasmith/DatasmithFacade/Public/DatasmithFacadeScene.h b/Engine/Source/Developer/Datasmith/DatasmithFacade/Public/DatasmithFacadeScene.h
--- a/Engine/Source/Developer/Datasmith/DatasmithFacade/Public/DatasmithFacadeScene.h
+++ b/Engine/Source/Developer/Datasmith/DatasmithFacade/Public/DatasmithFacadeScene.h
@@ -16,7 +16,7 @@ class FDatasmithFacadeMeshElement;
 class FDatasmithFacadeMetaData;
 class FDatasmithFacadeTexture;
 class FDatasmithFacadeLevelVariantSets;
-
+class FDatasmithFacadeLevelSequence;
 
 class DATASMITHFACADE_API FDatasmithFacadeScene
 {
@@ -143,6 +143,20 @@ public:
 		FDatasmithFacadeLevelVariantSets* InLevelVariantSetsPtr
 	);
 
+	void AddLevelSequence(
+		FDatasmithFacadeLevelSequence* InLevelSequence
+	);
+
+	int32 GetLevelSequencesCount() const;
+
+	FDatasmithFacadeLevelSequence* GetNewLevelSequence(
+		int32 LevelSequenceIndex
+	);
+
+	void RemoveLevelSequence(
+		FDatasmithFacadeLevelSequence* InLevelSequence
+	);
+
 	void AddMetaData(
 		FDatasmithFacadeMetaData* InMetaDataPtr
 	);
diff --git a/Engine/Source/Programs/Enterprise/Datasmith/DatasmithFacadeCSharp/Private/DatasmithFacadeCSharp.cpp b/Engine/Source/Programs/Enterprise/Datasmith/DatasmithFacadeCSharp/Private/DatasmithFacadeCSharp.cpp
--- a/Engine/Source/Programs/Enterprise/Datasmith/DatasmithFacadeCSharp/Private/DatasmithFacadeCSharp.cpp
+++ b/Engine/Source/Programs/Enterprise/Datasmith/DatasmithFacadeCSharp/Private/DatasmithFacadeCSharp.cpp
@@ -334,6 +334,7 @@ SWIGEXPORT void SWIGSTDCALL SWIGRegisterWStringCallback_DatasmithFacadeCSharp(SW
 #include "DatasmithFacadeActorCamera.h"
 #include "DatasmithFacadeActorLight.h"
 #include "DatasmithFacadeActorMesh.h"
+#include "DatasmithFacadeAnimation.h"
 #include "DatasmithFacadeScene.h"
 #include "DatasmithFacadeVariant.h"
 #include "DatasmithFacadeLog.h"
@@ -4821,6 +4822,563 @@ SWIGEXPORT void SWIGSTDCALL CSharp_FDatasmithFacadeActorMesh_ResetMaterialOverri
 }
 
 
+SWIGEXPORT unsigned int SWIGSTDCALL CSharp_FDatasmithFacadeBaseAnimation_IsSubType(void * jarg1, int jarg2) {
+  unsigned int jresult ;
+  FDatasmithFacadeBaseAnimation *arg1 = (FDatasmithFacadeBaseAnimation *) 0 ;
+  EDatasmithFacadeElementAnimationSubType arg2 ;
+  bool result;
+  
+  arg1 = (FDatasmithFacadeBaseAnimation *)jarg1; 
+  arg2 = (EDatasmithFacadeElementAnimationSubType)jarg2; 
+  result = (bool)((FDatasmithFacadeBaseAnimation const *)arg1)->IsSubType(arg2);
+  jresult = result; 
+  return jresult;
+}
+
+
+SWIGEXPORT void SWIGSTDCALL CSharp_FDatasmithFacadeBaseAnimation_SetCompletionMode(void * jarg1, int jarg2) {
+  FDatasmithFacadeBaseAnimation *arg1 = (FDatasmithFacadeBaseAnimation *) 0 ;
+  EDatasmithFacadeCompletionMode arg2 ;
+  
+  arg1 = (FDatasmithFacadeBaseAnimation *)jarg1; 
+  arg2 = (EDatasmithFacadeCompletionMode)jarg2; 
+  (arg1)->SetCompletionMode(arg2);
+}
+
+
+SWIGEXPORT int SWIGSTDCALL CSharp_FDatasmithFacadeBaseAnimation_GetCompletionMode(void * jarg1) {
+  int jresult ;
+  FDatasmithFacadeBaseAnimation *arg1 = (FDatasmithFacadeBaseAnimation *) 0 ;
+  EDatasmithFacadeCompletionMode result;
+  
+  arg1 = (FDatasmithFacadeBaseAnimation *)jarg1; 
+  result = (EDatasmithFacadeCompletionMode)((FDatasmithFacadeBaseAnimation const *)arg1)->GetCompletionMode();
+  jresult = (int)result; 
+  return jresult;
+}
+
+
+SWIGEXPORT void SWIGSTDCALL CSharp_delete_FDatasmithFacadeBaseAnimation(void * jarg1) {
+  FDatasmithFacadeBaseAnimation *arg1 = (FDatasmithFacadeBaseAnimation *) 0 ;
+  
+  arg1 = (FDatasmithFacadeBaseAnimation *)jarg1; 
+  delete arg1;
+}
+
+
+SWIGEXPORT void * SWIGSTDCALL CSharp_new_FDatasmithFacadeTransformAnimation(wchar_t * jarg1) {
+  void * jresult ;
+  TCHAR *arg1 = (TCHAR *) 0 ;
+  FDatasmithFacadeTransformAnimation *result = 0 ;
+  
+  arg1 = (TCHAR *)jarg1; 
+  result = (FDatasmithFacadeTransformAnimation *)new FDatasmithFacadeTransformAnimation((TCHAR const *)arg1);
+  jresult = (void *)result; 
+  return jresult;
+}
+
+
+SWIGEXPORT void SWIGSTDCALL CSharp_FDatasmithFacadeTransformAnimation_AddFrame(void * jarg1, int jarg2, int jarg3, float jarg4, float jarg5, float jarg6) {
+  FDatasmithFacadeTransformAnimation *arg1 = (FDatasmithFacadeTransformAnimation *) 0 ;
+  EDatasmithFacadeAnimationTransformType arg2 ;
+  int32 arg3 ;
+  float arg4 ;
+  float arg5 ;
+  float arg6 ;
+  
+  arg1 = (FDatasmithFacadeTransformAnimation *)jarg1; 
+  arg2 = (EDatasmithFacadeAnimationTransformType)jarg2; 
+  arg3 = (int32)jarg3; 
+  arg4 = (float)jarg4; 
+  arg5 = (float)jarg5; 
+  arg6 = (float)jarg6; 
+  (arg1)->AddFrame(arg2,arg3,arg4,arg5,arg6);
+}
+
+
+SWIGEXPORT int SWIGSTDCALL CSharp_FDatasmithFacadeTransformAnimation_GetFramesCount(void * jarg1, int jarg2) {
+  int jresult ;
+  FDatasmithFacadeTransformAnimation *arg1 = (FDatasmithFacadeTransformAnimation *) 0 ;
+  EDatasmithFacadeAnimationTransformType arg2 ;
+  int32 result;
+  
+  arg1 = (FDatasmithFacadeTransformAnimation *)jarg1; 
+  arg2 = (EDatasmithFacadeAnimationTransformType)jarg2; 
+  result = (int32)((FDatasmithFacadeTransformAnimation const *)arg1)->GetFramesCount(arg2);
+  jresult = result; 
+  return jresult;
+}
+
+
+SWIGEXPORT void SWIGSTDCALL CSharp_FDatasmithFacadeTransformAnimation_SetCurveInterpMode(void * jarg1, int jarg2, int jarg3) {
+  FDatasmithFacadeTransformAnimation *arg1 = (FDatasmithFacadeTransformAnimation *) 0 ;
+  EDatasmithFacadeAnimationTransformType arg2 ;
+  EDatasmithFacadeCurveInterpMode arg3 ;
+  
+  arg1 = (FDatasmithFacadeTransformAnimation *)jarg1; 
+  arg2 = (EDatasmithFacadeAnimationTransformType)jarg2; 
+  arg3 = (EDatasmithFacadeCurveInterpMode)jarg3; 
+  (arg1)->SetCurveInterpMode(arg2,arg3);
+}
+
+
+SWIGEXPORT int SWIGSTDCALL CSharp_FDatasmithFacadeTransformAnimation_GetCurveInterpMode(void * jarg1, int jarg2) {
+  int jresult ;
+  FDatasmithFacadeTransformAnimation *arg1 = (FDatasmithFacadeTransformAnimation *) 0 ;
+  EDatasmithFacadeAnimationTransformType arg2 ;
+  EDatasmithFacadeCurveInterpMode result;
+  
+  arg1 = (FDatasmithFacadeTransformAnimation *)jarg1; 
+  arg2 = (EDatasmithFacadeAnimationTransformType)jarg2; 
+  result = (EDatasmithFacadeCurveInterpMode)((FDatasmithFacadeTransformAnimation const *)arg1)->GetCurveInterpMode(arg2);
+  jresult = (int)result; 
+  return jresult;
+}
+
+
+SWIGEXPORT void SWIGSTDCALL CSharp_FDatasmithFacadeTransformAnimation_GetFrame(void * jarg1, int jarg2, int jarg3, int * jarg4, float * jarg5, float * jarg6, float * jarg7) {
+  FDatasmithFacadeTransformAnimation *arg1 = (FDatasmithFacadeTransformAnimation *) 0 ;
+  EDatasmithFacadeAnimationTransformType arg2 ;
+  int32 arg3 ;
+  int32 *arg4 = 0 ;
+  float *arg5 = 0 ;
+  float *arg6 = 0 ;
+  float *arg7 = 0 ;
+  
+  arg1 = (FDatasmithFacadeTransformAnimation *)jarg1; 
+  arg2 = (EDatasmithFacadeAnimationTransformType)jarg2; 
+  arg3 = (int32)jarg3; 
+  arg4 = (int32 *)jarg4; 
+  arg5 = (float *)jarg5; 
+  arg6 = (float *)jarg6; 
+  arg7 = (float *)jarg7; 
+  ((FDatasmithFacadeTransformAnimation const *)arg1)->GetFrame(arg2,arg3,*arg4,*arg5,*arg6,*arg7);
+}
+
+
+SWIGEXPORT void SWIGSTDCALL CSharp_FDatasmithFacadeTransformAnimation_RemoveFrame(void * jarg1, int jarg2, int jarg3) {
+  FDatasmithFacadeTransformAnimation *arg1 = (FDatasmithFacadeTransformAnimation *) 0 ;
+  EDatasmithFacadeAnimationTransformType arg2 ;
+  int32 arg3 ;
+  
+  arg1 = (FDatasmithFacadeTransformAnimation *)jarg1; 
+  arg2 = (EDatasmithFacadeAnimationTransformType)jarg2; 
+  arg3 = (int32)jarg3; 
+  (arg1)->RemoveFrame(arg2,arg3);
+}
+
+
+SWIGEXPORT int SWIGSTDCALL CSharp_FDatasmithFacadeTransformAnimation_GetEnabledTransformChannels(void * jarg1) {
+  int jresult ;
+  FDatasmithFacadeTransformAnimation *arg1 = (FDatasmithFacadeTransformAnimation *) 0 ;
+  EDatasmithFacadeTransformChannels result;
+  
+  arg1 = (FDatasmithFacadeTransformAnimation *)jarg1; 
+  result = (EDatasmithFacadeTransformChannels)((FDatasmithFacadeTransformAnimation const *)arg1)->GetEnabledTransformChannels();
+  jresult = (int)result; 
+  return jresult;
+}
+
+
+SWIGEXPORT void SWIGSTDCALL CSharp_FDatasmithFacadeTransformAnimation_SetEnabledTransformChannels(void * jarg1, int jarg2) {
+  FDatasmithFacadeTransformAnimation *arg1 = (FDatasmithFacadeTransformAnimation *) 0 ;
+  EDatasmithFacadeTransformChannels arg2 ;
+  
+  arg1 = (FDatasmithFacadeTransformAnimation *)jarg1; 
+  arg2 = (EDatasmithFacadeTransformChannels)jarg2; 
+  (arg1)->SetEnabledTransformChannels(arg2);
+}
+
+
+SWIGEXPORT void SWIGSTDCALL CSharp_delete_FDatasmithFacadeTransformAnimation(void * jarg1) {
+  FDatasmithFacadeTransformAnimation *arg1 = (FDatasmithFacadeTransformAnimation *) 0 ;
+  
+  arg1 = (FDatasmithFacadeTransformAnimation *)jarg1; 
+  delete arg1;
+}
+
+
+SWIGEXPORT void * SWIGSTDCALL CSharp_new_FDatasmithFacadeVisibilityAnimation(wchar_t * jarg1) {
+  void * jresult ;
+  TCHAR *arg1 = (TCHAR *) 0 ;
+  FDatasmithFacadeVisibilityAnimation *result = 0 ;
+  
+  arg1 = (TCHAR *)jarg1; 
+  result = (FDatasmithFacadeVisibilityAnimation *)new FDatasmithFacadeVisibilityAnimation((TCHAR const *)arg1);
+  jresult = (void *)result; 
+  return jresult;
+}
+
+
+SWIGEXPORT void SWIGSTDCALL CSharp_FDatasmithFacadeVisibilityAnimation_AddFrame(void * jarg1, int jarg2, unsigned int jarg3) {
+  FDatasmithFacadeVisibilityAnimation *arg1 = (FDatasmithFacadeVisibilityAnimation *) 0 ;
+  int32 arg2 ;
+  bool arg3 ;
+  
+  arg1 = (FDatasmithFacadeVisibilityAnimation *)jarg1; 
+  arg2 = (int32)jarg2; 
+  arg3 = jarg3 ? true : false; 
+  (arg1)->AddFrame(arg2,arg3);
+}
+
+
+SWIGEXPORT int SWIGSTDCALL CSharp_FDatasmithFacadeVisibilityAnimation_GetFramesCount(void * jarg1) {
+  int jresult ;
+  FDatasmithFacadeVisibilityAnimation *arg1 = (FDatasmithFacadeVisibilityAnimation *) 0 ;
+  int32 result;
+  
+  arg1 = (FDatasmithFacadeVisibilityAnimation *)jarg1; 
+  result = (int32)((FDatasmithFacadeVisibilityAnimation const *)arg1)->GetFramesCount();
+  jresult = result; 
+  return jresult;
+}
+
+
+SWIGEXPORT void SWIGSTDCALL CSharp_FDatasmithFacadeVisibilityAnimation_SetCurveInterpMode(void * jarg1, int jarg2) {
+  FDatasmithFacadeVisibilityAnimation *arg1 = (FDatasmithFacadeVisibilityAnimation *) 0 ;
+  EDatasmithFacadeCurveInterpMode arg2 ;
+  
+  arg1 = (FDatasmithFacadeVisibilityAnimation *)jarg1; 
+  arg2 = (EDatasmithFacadeCurveInterpMode)jarg2; 
+  (arg1)->SetCurveInterpMode(arg2);
+}
+
+
+SWIGEXPORT int SWIGSTDCALL CSharp_FDatasmithFacadeVisibilityAnimation_GetCurveInterpMode(void * jarg1) {
+  int jresult ;
+  FDatasmithFacadeVisibilityAnimation *arg1 = (FDatasmithFacadeVisibilityAnimation *) 0 ;
+  EDatasmithFacadeCurveInterpMode result;
+  
+  arg1 = (FDatasmithFacadeVisibilityAnimation *)jarg1; 
+  result = (EDatasmithFacadeCurveInterpMode)((FDatasmithFacadeVisibilityAnimation const *)arg1)->GetCurveInterpMode();
+  jresult = (int)result; 
+  return jresult;
+}
+
+
+SWIGEXPORT void SWIGSTDCALL CSharp_FDatasmithFacadeVisibilityAnimation_GetFrame(void * jarg1, int jarg2, int * jarg3, unsigned int * jarg4) {
+  FDatasmithFacadeVisibilityAnimation *arg1 = (FDatasmithFacadeVisibilityAnimation *) 0 ;
+  int32 arg2 ;
+  int32 *arg3 = 0 ;
+  bool *arg4 = 0 ;
+  
+  arg1 = (FDatasmithFacadeVisibilityAnimation *)jarg1; 
+  arg2 = (int32)jarg2; 
+  arg3 = (int32 *)jarg3; 
+  *jarg4 = 0; 
+  arg4 = (bool *)jarg4; 
+  ((FDatasmithFacadeVisibilityAnimation const *)arg1)->GetFrame(arg2,*arg3,*arg4);
+}
+
+
+SWIGEXPORT void SWIGSTDCALL CSharp_FDatasmithFacadeVisibilityAnimation_RemoveFrame(void * jarg1, int jarg2) {
+  FDatasmithFacadeVisibilityAnimation *arg1 = (FDatasmithFacadeVisibilityAnimation *) 0 ;
+  int32 arg2 ;
+  
+  arg1 = (FDatasmithFacadeVisibilityAnimation *)jarg1; 
+  arg2 = (int32)jarg2; 
+  (arg1)->RemoveFrame(arg2);
+}
+
+
+SWIGEXPORT void SWIGSTDCALL CSharp_FDatasmithFacadeVisibilityAnimation_SetPropagateToChildren(void * jarg1, unsigned int jarg2) {
+  FDatasmithFacadeVisibilityAnimation *arg1 = (FDatasmithFacadeVisibilityAnimation *) 0 ;
+  bool arg2 ;
+  
+  arg1 = (FDatasmithFacadeVisibilityAnimation *)jarg1; 
+  arg2 = jarg2 ? true : false; 
+  (arg1)->SetPropagateToChildren(arg2);
+}
+
+
+SWIGEXPORT unsigned int SWIGSTDCALL CSharp_FDatasmithFacadeVisibilityAnimation_GetPropagateToChildren(void * jarg1) {
+  unsigned int jresult ;
+  FDatasmithFacadeVisibilityAnimation *arg1 = (FDatasmithFacadeVisibilityAnimation *) 0 ;
+  bool result;
+  
+  arg1 = (FDatasmithFacadeVisibilityAnimation *)jarg1; 
+  result = (bool)((FDatasmithFacadeVisibilityAnimation const *)arg1)->GetPropagateToChildren();
+  jresult = result; 
+  return jresult;
+}
+
+
+SWIGEXPORT void SWIGSTDCALL CSharp_delete_FDatasmithFacadeVisibilityAnimation(void * jarg1) {
+  FDatasmithFacadeVisibilityAnimation *arg1 = (FDatasmithFacadeVisibilityAnimation *) 0 ;
+  
+  arg1 = (FDatasmithFacadeVisibilityAnimation *)jarg1; 
+  delete arg1;
+}
+
+
+SWIGEXPORT void * SWIGSTDCALL CSharp_new_FDatasmithFacadeSubsequenceAnimation(wchar_t * jarg1) {
+  void * jresult ;
+  TCHAR *arg1 = (TCHAR *) 0 ;
+  FDatasmithFacadeSubsequenceAnimation *result = 0 ;
+  
+  arg1 = (TCHAR *)jarg1; 
+  result = (FDatasmithFacadeSubsequenceAnimation *)new FDatasmithFacadeSubsequenceAnimation((TCHAR const *)arg1);
+  jresult = (void *)result; 
+  return jresult;
+}
+
+
+SWIGEXPORT int SWIGSTDCALL CSharp_FDatasmithFacadeSubsequenceAnimation_GetStartTime(void * jarg1) {
+  int jresult ;
+  FDatasmithFacadeSubsequenceAnimation *arg1 = (FDatasmithFacadeSubsequenceAnimation *) 0 ;
+  int32 result;
+  
+  arg1 = (FDatasmithFacadeSubsequenceAnimation *)jarg1; 
+  result = (int32)((FDatasmithFacadeSubsequenceAnimation const *)arg1)->GetStartTime();
+  jresult = result; 
+  return jresult;
+}
+
+
+SWIGEXPORT void SWIGSTDCALL CSharp_FDatasmithFacadeSubsequenceAnimation_SetStartTime(void * jarg1, int jarg2) {
+  FDatasmithFacadeSubsequenceAnimation *arg1 = (FDatasmithFacadeSubsequenceAnimation *) 0 ;
+  int32 arg2 ;
+  
+  arg1 = (FDatasmithFacadeSubsequenceAnimation *)jarg1; 
+  arg2 = (int32)jarg2; 
+  (arg1)->SetStartTime(arg2);
+}
+
+
+SWIGEXPORT int SWIGSTDCALL CSharp_FDatasmithFacadeSubsequenceAnimation_GetDuration(void * jarg1) {
+  int jresult ;
+  FDatasmithFacadeSubsequenceAnimation *arg1 = (FDatasmithFacadeSubsequenceAnimation *) 0 ;
+  int32 result;
+  
+  arg1 = (FDatasmithFacadeSubsequenceAnimation *)jarg1; 
+  result = (int32)((FDatasmithFacadeSubsequenceAnimation const *)arg1)->GetDuration();
+  jresult = result; 
+  return jresult;
+}
+
+
+SWIGEXPORT void SWIGSTDCALL CSharp_FDatasmithFacadeSubsequenceAnimation_SetDuration(void * jarg1, int jarg2) {
+  FDatasmithFacadeSubsequenceAnimation *arg1 = (FDatasmithFacadeSubsequenceAnimation *) 0 ;
+  int32 arg2 ;
+  
+  arg1 = (FDatasmithFacadeSubsequenceAnimation *)jarg1; 
+  arg2 = (int32)jarg2; 
+  (arg1)->SetDuration(arg2);
+}
+
+
+SWIGEXPORT float SWIGSTDCALL CSharp_FDatasmithFacadeSubsequenceAnimation_GetTimeScale(void * jarg1) {
+  float jresult ;
+  FDatasmithFacadeSubsequenceAnimation *arg1 = (FDatasmithFacadeSubsequenceAnimation *) 0 ;
+  float result;
+  
+  arg1 = (FDatasmithFacadeSubsequenceAnimation *)jarg1; 
+  result = (float)((FDatasmithFacadeSubsequenceAnimation const *)arg1)->GetTimeScale();
+  jresult = result; 
+  return jresult;
+}
+
+
+SWIGEXPORT void SWIGSTDCALL CSharp_FDatasmithFacadeSubsequenceAnimation_SetTimeScale(void * jarg1, float jarg2) {
+  FDatasmithFacadeSubsequenceAnimation *arg1 = (FDatasmithFacadeSubsequenceAnimation *) 0 ;
+  float arg2 ;
+  
+  arg1 = (FDatasmithFacadeSubsequenceAnimation *)jarg1; 
+  arg2 = (float)jarg2; 
+  (arg1)->SetTimeScale(arg2);
+}
+
+
+SWIGEXPORT void * SWIGSTDCALL CSharp_FDatasmithFacadeSubsequenceAnimation_GetNewSubsequence(void * jarg1) {
+  void * jresult ;
+  FDatasmithFacadeSubsequenceAnimation *arg1 = (FDatasmithFacadeSubsequenceAnimation *) 0 ;
+  FDatasmithFacadeLevelSequence *result = 0 ;
+  
+  arg1 = (FDatasmithFacadeSubsequenceAnimation *)jarg1; 
+  result = (FDatasmithFacadeLevelSequence *)((FDatasmithFacadeSubsequenceAnimation const *)arg1)->GetNewSubsequence();
+  jresult = (void *)result; 
+  return jresult;
+}
+
+
+SWIGEXPORT void SWIGSTDCALL CSharp_FDatasmithFacadeSubsequenceAnimation_SetSubsequence(void * jarg1, void * jarg2) {
+  FDatasmithFacadeSubsequenceAnimation *arg1 = (FDatasmithFacadeSubsequenceAnimation *) 0 ;
+  FDatasmithFacadeLevelSequence *arg2 = (FDatasmithFacadeLevelSequence *) 0 ;
+  
+  arg1 = (FDatasmithFacadeSubsequenceAnimation *)jarg1; 
+  arg2 = (FDatasmithFacadeLevelSequence *)jarg2; 
+  (arg1)->SetSubsequence(arg2);
+}
+
+
+SWIGEXPORT void SWIGSTDCALL CSharp_delete_FDatasmithFacadeSubsequenceAnimation(void * jarg1) {
+  FDatasmithFacadeSubsequenceAnimation *arg1 = (FDatasmithFacadeSubsequenceAnimation *) 0 ;
+  
+  arg1 = (FDatasmithFacadeSubsequenceAnimation *)jarg1; 
+  delete arg1;
+}
+
+
+SWIGEXPORT void * SWIGSTDCALL CSharp_new_FDatasmithFacadeLevelSequence(wchar_t * jarg1) {
+  void * jresult ;
+  TCHAR *arg1 = (TCHAR *) 0 ;
+  FDatasmithFacadeLevelSequence *result = 0 ;
+  
+  arg1 = (TCHAR *)jarg1; 
+  result = (FDatasmithFacadeLevelSequence *)new FDatasmithFacadeLevelSequence((TCHAR const *)arg1);
+  jresult = (void *)result; 
+  return jresult;
+}
+
+
+SWIGEXPORT wchar_t * SWIGSTDCALL CSharp_FDatasmithFacadeLevelSequence_GetFile(void * jarg1) {
+  wchar_t * jresult ;
+  FDatasmithFacadeLevelSequence *arg1 = (FDatasmithFacadeLevelSequence *) 0 ;
+  TCHAR *result = 0 ;
+  
+  arg1 = (FDatasmithFacadeLevelSequence *)jarg1; 
+  result = (TCHAR *)((FDatasmithFacadeLevelSequence const *)arg1)->GetFile();
+  jresult = (wchar_t *)result; 
+  return jresult;
+}
+
+
+SWIGEXPORT void SWIGSTDCALL CSharp_FDatasmithFacadeLevelSequence_SetFile(void * jarg1, wchar_t * jarg2) {
+  FDatasmithFacadeLevelSequence *arg1 = (FDatasmithFacadeLevelSequence *) 0 ;
+  TCHAR *arg2 = (TCHAR *) 0 ;
+  
+  arg1 = (FDatasmithFacadeLevelSequence *)jarg1; 
+  arg2 = (TCHAR *)jarg2; 
+  (arg1)->SetFile((TCHAR const *)arg2);
+}
+
+
+SWIGEXPORT void SWIGSTDCALL CSharp_FDatasmithFacadeLevelSequence_InternalGetFileHash(void * jarg1, void * jarg2, unsigned long jarg3) {
+  FDatasmithFacadeLevelSequence *arg1 = (FDatasmithFacadeLevelSequence *) 0 ;
+  TCHAR *arg2 ;
+  size_t arg3 ;
+  
+  arg1 = (FDatasmithFacadeLevelSequence *)jarg1; 
+  arg2 = (TCHAR *)jarg2; 
+  arg3 = (size_t)jarg3; 
+  ((FDatasmithFacadeLevelSequence const *)arg1)->GetFileHash(arg2,arg3);
+}
+
+
+SWIGEXPORT void SWIGSTDCALL CSharp_FDatasmithFacadeLevelSequence_SetFileHash(void * jarg1, wchar_t * jarg2) {
+  FDatasmithFacadeLevelSequence *arg1 = (FDatasmithFacadeLevelSequence *) 0 ;
+  TCHAR *arg2 = (TCHAR *) 0 ;
+  
+  arg1 = (FDatasmithFacadeLevelSequence *)jarg1; 
+  arg2 = (TCHAR *)jarg2; 
+  (arg1)->SetFileHash((TCHAR const *)arg2);
+}
+
+
+SWIGEXPORT float SWIGSTDCALL CSharp_FDatasmithFacadeLevelSequence_GetFrameRate(void * jarg1) {
+  float jresult ;
+  FDatasmithFacadeLevelSequence *arg1 = (FDatasmithFacadeLevelSequence *) 0 ;
+  float result;
+  
+  arg1 = (FDatasmithFacadeLevelSequence *)jarg1; 
+  result = (float)((FDatasmithFacadeLevelSequence const *)arg1)->GetFrameRate();
+  jresult = result; 
+  return jresult;
+}
+
+
+SWIGEXPORT void SWIGSTDCALL CSharp_FDatasmithFacadeLevelSequence_SetFrameRate(void * jarg1, float jarg2) {
+  FDatasmithFacadeLevelSequence *arg1 = (FDatasmithFacadeLevelSequence *) 0 ;
+  float arg2 ;
+  
+  arg1 = (FDatasmithFacadeLevelSequence *)jarg1; 
+  arg2 = (float)jarg2; 
+  (arg1)->SetFrameRate(arg2);
+}
+
+
+SWIGEXPORT void SWIGSTDCALL CSharp_FDatasmithFacadeLevelSequence_AddAnimation(void * jarg1, void * jarg2) {
+  FDatasmithFacadeLevelSequence *arg1 = (FDatasmithFacadeLevelSequence *) 0 ;
+  FDatasmithFacadeTransformAnimation *arg2 = (FDatasmithFacadeTransformAnimation *) 0 ;
+  
+  arg1 = (FDatasmithFacadeLevelSequence *)jarg1; 
+  arg2 = (FDatasmithFacadeTransformAnimation *)jarg2; 
+  (arg1)->AddAnimation(arg2);
+}
+
+
+SWIGEXPORT int SWIGSTDCALL CSharp_FDatasmithFacadeLevelSequence_GetAnimationsCount(void * jarg1) {
+  int jresult ;
+  FDatasmithFacadeLevelSequence *arg1 = (FDatasmithFacadeLevelSequence *) 0 ;
+  int32 result;
+  
+  arg1 = (FDatasmithFacadeLevelSequence *)jarg1; 
+  result = (int32)((FDatasmithFacadeLevelSequence const *)arg1)->GetAnimationsCount();
+  jresult = result; 
+  return jresult;
+}
+
+
+SWIGEXPORT void * SWIGSTDCALL CSharp_FDatasmithFacadeLevelSequence_GetTransformAnimation(void * jarg1, int jarg2) {
+  void * jresult ;
+  FDatasmithFacadeLevelSequence *arg1 = (FDatasmithFacadeLevelSequence *) 0 ;
+  int32 arg2 ;
+  FDatasmithFacadeTransformAnimation *result = 0 ;
+  
+  arg1 = (FDatasmithFacadeLevelSequence *)jarg1; 
+  arg2 = (int32)jarg2; 
+  result = (FDatasmithFacadeTransformAnimation *)(arg1)->GetNewTransformAnimation(arg2);
+  jresult = (void *)result; 
+  return jresult;
+}
+
+
+SWIGEXPORT void * SWIGSTDCALL CSharp_FDatasmithFacadeLevelSequence_GetVisibilityAnimation(void * jarg1, int jarg2) {
+  void * jresult ;
+  FDatasmithFacadeLevelSequence *arg1 = (FDatasmithFacadeLevelSequence *) 0 ;
+  int32 arg2 ;
+  FDatasmithFacadeVisibilityAnimation *result = 0 ;
+  
+  arg1 = (FDatasmithFacadeLevelSequence *)jarg1; 
+  arg2 = (int32)jarg2; 
+  result = (FDatasmithFacadeVisibilityAnimation *)(arg1)->GetNewVisibilityAnimation(arg2);
+  jresult = (void *)result; 
+  return jresult;
+}
+
+
+SWIGEXPORT void * SWIGSTDCALL CSharp_FDatasmithFacadeLevelSequence_GetSubsequenceAnimation(void * jarg1, int jarg2) {
+  void * jresult ;
+  FDatasmithFacadeLevelSequence *arg1 = (FDatasmithFacadeLevelSequence *) 0 ;
+  int32 arg2 ;
+  FDatasmithFacadeSubsequenceAnimation *result = 0 ;
+  
+  arg1 = (FDatasmithFacadeLevelSequence *)jarg1; 
+  arg2 = (int32)jarg2; 
+  result = (FDatasmithFacadeSubsequenceAnimation *)(arg1)->GetNewSubsequenceAnimation(arg2);
+  jresult = (void *)result; 
+  return jresult;
+}
+
+
+SWIGEXPORT void SWIGSTDCALL CSharp_FDatasmithFacadeLevelSequence_RemoveAnimation(void * jarg1, void * jarg2) {
+  FDatasmithFacadeLevelSequence *arg1 = (FDatasmithFacadeLevelSequence *) 0 ;
+  FDatasmithFacadeBaseAnimation *arg2 = (FDatasmithFacadeBaseAnimation *) 0 ;
+  
+  arg1 = (FDatasmithFacadeLevelSequence *)jarg1; 
+  arg2 = (FDatasmithFacadeBaseAnimation *)jarg2; 
+  (arg1)->RemoveAnimation(arg2);
+}
+
+
+SWIGEXPORT void SWIGSTDCALL CSharp_delete_FDatasmithFacadeLevelSequence(void * jarg1) {
+  FDatasmithFacadeLevelSequence *arg1 = (FDatasmithFacadeLevelSequence *) 0 ;
+  
+  arg1 = (FDatasmithFacadeLevelSequence *)jarg1; 
+  delete arg1;
+}
+
+
 SWIGEXPORT void * SWIGSTDCALL CSharp_new_FDatasmithFacadeScene(wchar_t * jarg1, wchar_t * jarg2, wchar_t * jarg3, wchar_t * jarg4) {
   void * jresult ;
   TCHAR *arg1 = (TCHAR *) 0 ;
@@ -5145,6 +5703,52 @@ SWIGEXPORT void SWIGSTDCALL CSharp_FDatasmithFacadeScene_RemoveLevelVariantSets(
 }
 
 
+SWIGEXPORT void SWIGSTDCALL CSharp_FDatasmithFacadeScene_AddLevelSequence(void * jarg1, void * jarg2) {
+  FDatasmithFacadeScene *arg1 = (FDatasmithFacadeScene *) 0 ;
+  FDatasmithFacadeLevelSequence *arg2 = (FDatasmithFacadeLevelSequence *) 0 ;
+  
+  arg1 = (FDatasmithFacadeScene *)jarg1; 
+  arg2 = (FDatasmithFacadeLevelSequence *)jarg2; 
+  (arg1)->AddLevelSequence(arg2);
+}
+
+
+SWIGEXPORT int SWIGSTDCALL CSharp_FDatasmithFacadeScene_GetLevelSequencesCount(void * jarg1) {
+  int jresult ;
+  FDatasmithFacadeScene *arg1 = (FDatasmithFacadeScene *) 0 ;
+  int32 result;
+  
+  arg1 = (FDatasmithFacadeScene *)jarg1; 
+  result = (int32)((FDatasmithFacadeScene const *)arg1)->GetLevelSequencesCount();
+  jresult = result; 
+  return jresult;
+}
+
+
+SWIGEXPORT void * SWIGSTDCALL CSharp_FDatasmithFacadeScene_GetLevelSequence(void * jarg1, int jarg2) {
+  void * jresult ;
+  FDatasmithFacadeScene *arg1 = (FDatasmithFacadeScene *) 0 ;
+  int32 arg2 ;
+  FDatasmithFacadeLevelSequence *result = 0 ;
+  
+  arg1 = (FDatasmithFacadeScene *)jarg1; 
+  arg2 = (int32)jarg2; 
+  result = (FDatasmithFacadeLevelSequence *)(arg1)->GetNewLevelSequence(arg2);
+  jresult = (void *)result; 
+  return jresult;
+}
+
+
+SWIGEXPORT void SWIGSTDCALL CSharp_FDatasmithFacadeScene_RemoveLevelSequence(void * jarg1, void * jarg2) {
+  FDatasmithFacadeScene *arg1 = (FDatasmithFacadeScene *) 0 ;
+  FDatasmithFacadeLevelSequence *arg2 = (FDatasmithFacadeLevelSequence *) 0 ;
+  
+  arg1 = (FDatasmithFacadeScene *)jarg1; 
+  arg2 = (FDatasmithFacadeLevelSequence *)jarg2; 
+  (arg1)->RemoveLevelSequence(arg2);
+}
+
+
 SWIGEXPORT void SWIGSTDCALL CSharp_FDatasmithFacadeScene_AddMetaData(void * jarg1, void * jarg2) {
   FDatasmithFacadeScene *arg1 = (FDatasmithFacadeScene *) 0 ;
   FDatasmithFacadeMetaData *arg2 = (FDatasmithFacadeMetaData *) 0 ;
@@ -6696,6 +7300,26 @@ SWIGEXPORT FDatasmithFacadeActor * SWIGSTDCALL CSharp_FDatasmithFacadeActorMesh_
     return (FDatasmithFacadeActor *)jarg1;
 }
 
+SWIGEXPORT FDatasmithFacadeElement * SWIGSTDCALL CSharp_FDatasmithFacadeBaseAnimation_SWIGUpcast(FDatasmithFacadeBaseAnimation *jarg1) {
+    return (FDatasmithFacadeElement *)jarg1;
+}
+
+SWIGEXPORT FDatasmithFacadeBaseAnimation * SWIGSTDCALL CSharp_FDatasmithFacadeTransformAnimation_SWIGUpcast(FDatasmithFacadeTransformAnimation *jarg1) {
+    return (FDatasmithFacadeBaseAnimation *)jarg1;
+}
+
+SWIGEXPORT FDatasmithFacadeBaseAnimation * SWIGSTDCALL CSharp_FDatasmithFacadeVisibilityAnimation_SWIGUpcast(FDatasmithFacadeVisibilityAnimation *jarg1) {
+    return (FDatasmithFacadeBaseAnimation *)jarg1;
+}
+
+SWIGEXPORT FDatasmithFacadeBaseAnimation * SWIGSTDCALL CSharp_FDatasmithFacadeSubsequenceAnimation_SWIGUpcast(FDatasmithFacadeSubsequenceAnimation *jarg1) {
+    return (FDatasmithFacadeBaseAnimation *)jarg1;
+}
+
+SWIGEXPORT FDatasmithFacadeElement * SWIGSTDCALL CSharp_FDatasmithFacadeLevelSequence_SWIGUpcast(FDatasmithFacadeLevelSequence *jarg1) {
+    return (FDatasmithFacadeElement *)jarg1;
+}
+
 SWIGEXPORT FDatasmithFacadeElement * SWIGSTDCALL CSharp_FDatasmithFacadePropertyCapture_SWIGUpcast(FDatasmithFacadePropertyCapture *jarg1) {
     return (FDatasmithFacadeElement *)jarg1;
 }
diff --git a/Engine/Source/Programs/Enterprise/Datasmith/DatasmithFacadeCSharp/Public/DatasmithFacadeCSharpPINVOKE.cs b/Engine/Source/Programs/Enterprise/Datasmith/DatasmithFacadeCSharp/Public/DatasmithFacadeCSharpPINVOKE.cs
--- a/Engine/Source/Programs/Enterprise/Datasmith/DatasmithFacadeCSharp/Public/DatasmithFacadeCSharpPINVOKE.cs
+++ b/Engine/Source/Programs/Enterprise/Datasmith/DatasmithFacadeCSharp/Public/DatasmithFacadeCSharpPINVOKE.cs
@@ -1358,6 +1358,150 @@ class DatasmithFacadeCSharpPINVOKE {
   [global::System.Runtime.InteropServices.DllImport("DatasmithFacadeCSharp", EntryPoint="CSharp_FDatasmithFacadeActorMesh_ResetMaterialOverrides")]
   public static extern void FDatasmithFacadeActorMesh_ResetMaterialOverrides(global::System.Runtime.InteropServices.HandleRef jarg1);
 
+  [global::System.Runtime.InteropServices.DllImport("DatasmithFacadeCSharp", EntryPoint="CSharp_FDatasmithFacadeBaseAnimation_IsSubType")]
+  public static extern bool FDatasmithFacadeBaseAnimation_IsSubType(global::System.Runtime.InteropServices.HandleRef jarg1, int jarg2);
+
+  [global::System.Runtime.InteropServices.DllImport("DatasmithFacadeCSharp", EntryPoint="CSharp_FDatasmithFacadeBaseAnimation_SetCompletionMode")]
+  public static extern void FDatasmithFacadeBaseAnimation_SetCompletionMode(global::System.Runtime.InteropServices.HandleRef jarg1, int jarg2);
+
+  [global::System.Runtime.InteropServices.DllImport("DatasmithFacadeCSharp", EntryPoint="CSharp_FDatasmithFacadeBaseAnimation_GetCompletionMode")]
+  public static extern int FDatasmithFacadeBaseAnimation_GetCompletionMode(global::System.Runtime.InteropServices.HandleRef jarg1);
+
+  [global::System.Runtime.InteropServices.DllImport("DatasmithFacadeCSharp", EntryPoint="CSharp_delete_FDatasmithFacadeBaseAnimation")]
+  public static extern void delete_FDatasmithFacadeBaseAnimation(global::System.Runtime.InteropServices.HandleRef jarg1);
+
+  [global::System.Runtime.InteropServices.DllImport("DatasmithFacadeCSharp", EntryPoint="CSharp_new_FDatasmithFacadeTransformAnimation")]
+  public static extern global::System.IntPtr new_FDatasmithFacadeTransformAnimation([global::System.Runtime.InteropServices.MarshalAs(global::System.Runtime.InteropServices.UnmanagedType.LPWStr)]string jarg1);
+
+  [global::System.Runtime.InteropServices.DllImport("DatasmithFacadeCSharp", EntryPoint="CSharp_FDatasmithFacadeTransformAnimation_AddFrame")]
+  public static extern void FDatasmithFacadeTransformAnimation_AddFrame(global::System.Runtime.InteropServices.HandleRef jarg1, int jarg2, int jarg3, float jarg4, float jarg5, float jarg6);
+
+  [global::System.Runtime.InteropServices.DllImport("DatasmithFacadeCSharp", EntryPoint="CSharp_FDatasmithFacadeTransformAnimation_GetFramesCount")]
+  public static extern int FDatasmithFacadeTransformAnimation_GetFramesCount(global::System.Runtime.InteropServices.HandleRef jarg1, int jarg2);
+
+  [global::System.Runtime.InteropServices.DllImport("DatasmithFacadeCSharp", EntryPoint="CSharp_FDatasmithFacadeTransformAnimation_SetCurveInterpMode")]
+  public static extern void FDatasmithFacadeTransformAnimation_SetCurveInterpMode(global::System.Runtime.InteropServices.HandleRef jarg1, int jarg2, int jarg3);
+
+  [global::System.Runtime.InteropServices.DllImport("DatasmithFacadeCSharp", EntryPoint="CSharp_FDatasmithFacadeTransformAnimation_GetCurveInterpMode")]
+  public static extern int FDatasmithFacadeTransformAnimation_GetCurveInterpMode(global::System.Runtime.InteropServices.HandleRef jarg1, int jarg2);
+
+  [global::System.Runtime.InteropServices.DllImport("DatasmithFacadeCSharp", EntryPoint="CSharp_FDatasmithFacadeTransformAnimation_GetFrame")]
+  public static extern void FDatasmithFacadeTransformAnimation_GetFrame(global::System.Runtime.InteropServices.HandleRef jarg1, int jarg2, int jarg3, out int jarg4, out float jarg5, out float jarg6, out float jarg7);
+
+  [global::System.Runtime.InteropServices.DllImport("DatasmithFacadeCSharp", EntryPoint="CSharp_FDatasmithFacadeTransformAnimation_RemoveFrame")]
+  public static extern void FDatasmithFacadeTransformAnimation_RemoveFrame(global::System.Runtime.InteropServices.HandleRef jarg1, int jarg2, int jarg3);
+
+  [global::System.Runtime.InteropServices.DllImport("DatasmithFacadeCSharp", EntryPoint="CSharp_FDatasmithFacadeTransformAnimation_GetEnabledTransformChannels")]
+  public static extern int FDatasmithFacadeTransformAnimation_GetEnabledTransformChannels(global::System.Runtime.InteropServices.HandleRef jarg1);
+
+  [global::System.Runtime.InteropServices.DllImport("DatasmithFacadeCSharp", EntryPoint="CSharp_FDatasmithFacadeTransformAnimation_SetEnabledTransformChannels")]
+  public static extern void FDatasmithFacadeTransformAnimation_SetEnabledTransformChannels(global::System.Runtime.InteropServices.HandleRef jarg1, int jarg2);
+
+  [global::System.Runtime.InteropServices.DllImport("DatasmithFacadeCSharp", EntryPoint="CSharp_delete_FDatasmithFacadeTransformAnimation")]
+  public static extern void delete_FDatasmithFacadeTransformAnimation(global::System.Runtime.InteropServices.HandleRef jarg1);
+
+  [global::System.Runtime.InteropServices.DllImport("DatasmithFacadeCSharp", EntryPoint="CSharp_new_FDatasmithFacadeVisibilityAnimation")]
+  public static extern global::System.IntPtr new_FDatasmithFacadeVisibilityAnimation([global::System.Runtime.InteropServices.MarshalAs(global::System.Runtime.InteropServices.UnmanagedType.LPWStr)]string jarg1);
+
+  [global::System.Runtime.InteropServices.DllImport("DatasmithFacadeCSharp", EntryPoint="CSharp_FDatasmithFacadeVisibilityAnimation_AddFrame")]
+  public static extern void FDatasmithFacadeVisibilityAnimation_AddFrame(global::System.Runtime.InteropServices.HandleRef jarg1, int jarg2, bool jarg3);
+
+  [global::System.Runtime.InteropServices.DllImport("DatasmithFacadeCSharp", EntryPoint="CSharp_FDatasmithFacadeVisibilityAnimation_GetFramesCount")]
+  public static extern int FDatasmithFacadeVisibilityAnimation_GetFramesCount(global::System.Runtime.InteropServices.HandleRef jarg1);
+
+  [global::System.Runtime.InteropServices.DllImport("DatasmithFacadeCSharp", EntryPoint="CSharp_FDatasmithFacadeVisibilityAnimation_SetCurveInterpMode")]
+  public static extern void FDatasmithFacadeVisibilityAnimation_SetCurveInterpMode(global::System.Runtime.InteropServices.HandleRef jarg1, int jarg2);
+
+  [global::System.Runtime.InteropServices.DllImport("DatasmithFacadeCSharp", EntryPoint="CSharp_FDatasmithFacadeVisibilityAnimation_GetCurveInterpMode")]
+  public static extern int FDatasmithFacadeVisibilityAnimation_GetCurveInterpMode(global::System.Runtime.InteropServices.HandleRef jarg1);
+
+  [global::System.Runtime.InteropServices.DllImport("DatasmithFacadeCSharp", EntryPoint="CSharp_FDatasmithFacadeVisibilityAnimation_GetFrame")]
+  public static extern void FDatasmithFacadeVisibilityAnimation_GetFrame(global::System.Runtime.InteropServices.HandleRef jarg1, int jarg2, out int jarg3, out bool jarg4);
+
+  [global::System.Runtime.InteropServices.DllImport("DatasmithFacadeCSharp", EntryPoint="CSharp_FDatasmithFacadeVisibilityAnimation_RemoveFrame")]
+  public static extern void FDatasmithFacadeVisibilityAnimation_RemoveFrame(global::System.Runtime.InteropServices.HandleRef jarg1, int jarg2);
+
+  [global::System.Runtime.InteropServices.DllImport("DatasmithFacadeCSharp", EntryPoint="CSharp_FDatasmithFacadeVisibilityAnimation_SetPropagateToChildren")]
+  public static extern void FDatasmithFacadeVisibilityAnimation_SetPropagateToChildren(global::System.Runtime.InteropServices.HandleRef jarg1, bool jarg2);
+
+  [global::System.Runtime.InteropServices.DllImport("DatasmithFacadeCSharp", EntryPoint="CSharp_FDatasmithFacadeVisibilityAnimation_GetPropagateToChildren")]
+  public static extern bool FDatasmithFacadeVisibilityAnimation_GetPropagateToChildren(global::System.Runtime.InteropServices.HandleRef jarg1);
+
+  [global::System.Runtime.InteropServices.DllImport("DatasmithFacadeCSharp", EntryPoint="CSharp_delete_FDatasmithFacadeVisibilityAnimation")]
+  public static extern void delete_FDatasmithFacadeVisibilityAnimation(global::System.Runtime.InteropServices.HandleRef jarg1);
+
+  [global::System.Runtime.InteropServices.DllImport("DatasmithFacadeCSharp", EntryPoint="CSharp_new_FDatasmithFacadeSubsequenceAnimation")]
+  public static extern global::System.IntPtr new_FDatasmithFacadeSubsequenceAnimation([global::System.Runtime.InteropServices.MarshalAs(global::System.Runtime.InteropServices.UnmanagedType.LPWStr)]string jarg1);
+
+  [global::System.Runtime.InteropServices.DllImport("DatasmithFacadeCSharp", EntryPoint="CSharp_FDatasmithFacadeSubsequenceAnimation_GetStartTime")]
+  public static extern int FDatasmithFacadeSubsequenceAnimation_GetStartTime(global::System.Runtime.InteropServices.HandleRef jarg1);
+
+  [global::System.Runtime.InteropServices.DllImport("DatasmithFacadeCSharp", EntryPoint="CSharp_FDatasmithFacadeSubsequenceAnimation_SetStartTime")]
+  public static extern void FDatasmithFacadeSubsequenceAnimation_SetStartTime(global::System.Runtime.InteropServices.HandleRef jarg1, int jarg2);
+
+  [global::System.Runtime.InteropServices.DllImport("DatasmithFacadeCSharp", EntryPoint="CSharp_FDatasmithFacadeSubsequenceAnimation_GetDuration")]
+  public static extern int FDatasmithFacadeSubsequenceAnimation_GetDuration(global::System.Runtime.InteropServices.HandleRef jarg1);
+
+  [global::System.Runtime.InteropServices.DllImport("DatasmithFacadeCSharp", EntryPoint="CSharp_FDatasmithFacadeSubsequenceAnimation_SetDuration")]
+  public static extern void FDatasmithFacadeSubsequenceAnimation_SetDuration(global::System.Runtime.InteropServices.HandleRef jarg1, int jarg2);
+
+  [global::System.Runtime.InteropServices.DllImport("DatasmithFacadeCSharp", EntryPoint="CSharp_FDatasmithFacadeSubsequenceAnimation_GetTimeScale")]
+  public static extern float FDatasmithFacadeSubsequenceAnimation_GetTimeScale(global::System.Runtime.InteropServices.HandleRef jarg1);
+
+  [global::System.Runtime.InteropServices.DllImport("DatasmithFacadeCSharp", EntryPoint="CSharp_FDatasmithFacadeSubsequenceAnimation_SetTimeScale")]
+  public static extern void FDatasmithFacadeSubsequenceAnimation_SetTimeScale(global::System.Runtime.InteropServices.HandleRef jarg1, float jarg2);
+
+  [global::System.Runtime.InteropServices.DllImport("DatasmithFacadeCSharp", EntryPoint="CSharp_FDatasmithFacadeSubsequenceAnimation_GetNewSubsequence")]
+  public static extern global::System.IntPtr FDatasmithFacadeSubsequenceAnimation_GetNewSubsequence(global::System.Runtime.InteropServices.HandleRef jarg1);
+
+  [global::System.Runtime.InteropServices.DllImport("DatasmithFacadeCSharp", EntryPoint="CSharp_FDatasmithFacadeSubsequenceAnimation_SetSubsequence")]
+  public static extern void FDatasmithFacadeSubsequenceAnimation_SetSubsequence(global::System.Runtime.InteropServices.HandleRef jarg1, global::System.Runtime.InteropServices.HandleRef jarg2);
+
+  [global::System.Runtime.InteropServices.DllImport("DatasmithFacadeCSharp", EntryPoint="CSharp_delete_FDatasmithFacadeSubsequenceAnimation")]
+  public static extern void delete_FDatasmithFacadeSubsequenceAnimation(global::System.Runtime.InteropServices.HandleRef jarg1);
+
+  [global::System.Runtime.InteropServices.DllImport("DatasmithFacadeCSharp", EntryPoint="CSharp_new_FDatasmithFacadeLevelSequence")]
+  public static extern global::System.IntPtr new_FDatasmithFacadeLevelSequence([global::System.Runtime.InteropServices.MarshalAs(global::System.Runtime.InteropServices.UnmanagedType.LPWStr)]string jarg1);
+
+  [global::System.Runtime.InteropServices.DllImport("DatasmithFacadeCSharp", EntryPoint="CSharp_FDatasmithFacadeLevelSequence_GetFile")]
+  public static extern global::System.IntPtr FDatasmithFacadeLevelSequence_GetFile(global::System.Runtime.InteropServices.HandleRef jarg1);
+
+  [global::System.Runtime.InteropServices.DllImport("DatasmithFacadeCSharp", EntryPoint="CSharp_FDatasmithFacadeLevelSequence_SetFile")]
+  public static extern void FDatasmithFacadeLevelSequence_SetFile(global::System.Runtime.InteropServices.HandleRef jarg1, [global::System.Runtime.InteropServices.MarshalAs(global::System.Runtime.InteropServices.UnmanagedType.LPWStr)]string jarg2);
+
+  [global::System.Runtime.InteropServices.DllImport("DatasmithFacadeCSharp", EntryPoint="CSharp_FDatasmithFacadeLevelSequence_InternalGetFileHash")]
+  public static extern void FDatasmithFacadeLevelSequence_InternalGetFileHash(global::System.Runtime.InteropServices.HandleRef jarg1, [global::System.Runtime.InteropServices.MarshalAs(global::System.Runtime.InteropServices.UnmanagedType.LPWStr)] System.Text.StringBuilder jarg2, ulong jarg3);
+
+  [global::System.Runtime.InteropServices.DllImport("DatasmithFacadeCSharp", EntryPoint="CSharp_FDatasmithFacadeLevelSequence_SetFileHash")]
+  public static extern void FDatasmithFacadeLevelSequence_SetFileHash(global::System.Runtime.InteropServices.HandleRef jarg1, [global::System.Runtime.InteropServices.MarshalAs(global::System.Runtime.InteropServices.UnmanagedType.LPWStr)]string jarg2);
+
+  [global::System.Runtime.InteropServices.DllImport("DatasmithFacadeCSharp", EntryPoint="CSharp_FDatasmithFacadeLevelSequence_GetFrameRate")]
+  public static extern float FDatasmithFacadeLevelSequence_GetFrameRate(global::System.Runtime.InteropServices.HandleRef jarg1);
+
+  [global::System.Runtime.InteropServices.DllImport("DatasmithFacadeCSharp", EntryPoint="CSharp_FDatasmithFacadeLevelSequence_SetFrameRate")]
+  public static extern void FDatasmithFacadeLevelSequence_SetFrameRate(global::System.Runtime.InteropServices.HandleRef jarg1, float jarg2);
+
+  [global::System.Runtime.InteropServices.DllImport("DatasmithFacadeCSharp", EntryPoint="CSharp_FDatasmithFacadeLevelSequence_AddAnimation")]
+  public static extern void FDatasmithFacadeLevelSequence_AddAnimation(global::System.Runtime.InteropServices.HandleRef jarg1, global::System.Runtime.InteropServices.HandleRef jarg2);
+
+  [global::System.Runtime.InteropServices.DllImport("DatasmithFacadeCSharp", EntryPoint="CSharp_FDatasmithFacadeLevelSequence_GetAnimationsCount")]
+  public static extern int FDatasmithFacadeLevelSequence_GetAnimationsCount(global::System.Runtime.InteropServices.HandleRef jarg1);
+
+  [global::System.Runtime.InteropServices.DllImport("DatasmithFacadeCSharp", EntryPoint="CSharp_FDatasmithFacadeLevelSequence_GetTransformAnimation")]
+  public static extern global::System.IntPtr FDatasmithFacadeLevelSequence_GetTransformAnimation(global::System.Runtime.InteropServices.HandleRef jarg1, int jarg2);
+
+  [global::System.Runtime.InteropServices.DllImport("DatasmithFacadeCSharp", EntryPoint="CSharp_FDatasmithFacadeLevelSequence_GetVisibilityAnimation")]
+  public static extern global::System.IntPtr FDatasmithFacadeLevelSequence_GetVisibilityAnimation(global::System.Runtime.InteropServices.HandleRef jarg1, int jarg2);
+
+  [global::System.Runtime.InteropServices.DllImport("DatasmithFacadeCSharp", EntryPoint="CSharp_FDatasmithFacadeLevelSequence_GetSubsequenceAnimation")]
+  public static extern global::System.IntPtr FDatasmithFacadeLevelSequence_GetSubsequenceAnimation(global::System.Runtime.InteropServices.HandleRef jarg1, int jarg2);
+
+  [global::System.Runtime.InteropServices.DllImport("DatasmithFacadeCSharp", EntryPoint="CSharp_FDatasmithFacadeLevelSequence_RemoveAnimation")]
+  public static extern void FDatasmithFacadeLevelSequence_RemoveAnimation(global::System.Runtime.InteropServices.HandleRef jarg1, global::System.Runtime.InteropServices.HandleRef jarg2);
+
+  [global::System.Runtime.InteropServices.DllImport("DatasmithFacadeCSharp", EntryPoint="CSharp_delete_FDatasmithFacadeLevelSequence")]
+  public static extern void delete_FDatasmithFacadeLevelSequence(global::System.Runtime.InteropServices.HandleRef jarg1);
+
   [global::System.Runtime.InteropServices.DllImport("DatasmithFacadeCSharp", EntryPoint="CSharp_new_FDatasmithFacadeScene")]
   public static extern global::System.IntPtr new_FDatasmithFacadeScene([global::System.Runtime.InteropServices.MarshalAs(global::System.Runtime.InteropServices.UnmanagedType.LPWStr)]string jarg1, [global::System.Runtime.InteropServices.MarshalAs(global::System.Runtime.InteropServices.UnmanagedType.LPWStr)]string jarg2, [global::System.Runtime.InteropServices.MarshalAs(global::System.Runtime.InteropServices.UnmanagedType.LPWStr)]string jarg3, [global::System.Runtime.InteropServices.MarshalAs(global::System.Runtime.InteropServices.UnmanagedType.LPWStr)]string jarg4);
 
@@ -1436,6 +1580,18 @@ class DatasmithFacadeCSharpPINVOKE {
   [global::System.Runtime.InteropServices.DllImport("DatasmithFacadeCSharp", EntryPoint="CSharp_FDatasmithFacadeScene_RemoveLevelVariantSets")]
   public static extern void FDatasmithFacadeScene_RemoveLevelVariantSets(global::System.Runtime.InteropServices.HandleRef jarg1, global::System.Runtime.InteropServices.HandleRef jarg2);
 
+  [global::System.Runtime.InteropServices.DllImport("DatasmithFacadeCSharp", EntryPoint="CSharp_FDatasmithFacadeScene_AddLevelSequence")]
+  public static extern void FDatasmithFacadeScene_AddLevelSequence(global::System.Runtime.InteropServices.HandleRef jarg1, global::System.Runtime.InteropServices.HandleRef jarg2);
+
+  [global::System.Runtime.InteropServices.DllImport("DatasmithFacadeCSharp", EntryPoint="CSharp_FDatasmithFacadeScene_GetLevelSequencesCount")]
+  public static extern int FDatasmithFacadeScene_GetLevelSequencesCount(global::System.Runtime.InteropServices.HandleRef jarg1);
+
+  [global::System.Runtime.InteropServices.DllImport("DatasmithFacadeCSharp", EntryPoint="CSharp_FDatasmithFacadeScene_GetLevelSequence")]
+  public static extern global::System.IntPtr FDatasmithFacadeScene_GetLevelSequence(global::System.Runtime.InteropServices.HandleRef jarg1, int jarg2);
+
+  [global::System.Runtime.InteropServices.DllImport("DatasmithFacadeCSharp", EntryPoint="CSharp_FDatasmithFacadeScene_RemoveLevelSequence")]
+  public static extern void FDatasmithFacadeScene_RemoveLevelSequence(global::System.Runtime.InteropServices.HandleRef jarg1, global::System.Runtime.InteropServices.HandleRef jarg2);
+
   [global::System.Runtime.InteropServices.DllImport("DatasmithFacadeCSharp", EntryPoint="CSharp_FDatasmithFacadeScene_AddMetaData")]
   public static extern void FDatasmithFacadeScene_AddMetaData(global::System.Runtime.InteropServices.HandleRef jarg1, global::System.Runtime.InteropServices.HandleRef jarg2);
 
@@ -1916,6 +2072,21 @@ class DatasmithFacadeCSharpPINVOKE {
   [global::System.Runtime.InteropServices.DllImport("DatasmithFacadeCSharp", EntryPoint="CSharp_FDatasmithFacadeActorMesh_SWIGUpcast")]
   public static extern global::System.IntPtr FDatasmithFacadeActorMesh_SWIGUpcast(global::System.IntPtr jarg1);
 
+  [global::System.Runtime.InteropServices.DllImport("DatasmithFacadeCSharp", EntryPoint="CSharp_FDatasmithFacadeBaseAnimation_SWIGUpcast")]
+  public static extern global::System.IntPtr FDatasmithFacadeBaseAnimation_SWIGUpcast(global::System.IntPtr jarg1);
+
+  [global::System.Runtime.InteropServices.DllImport("DatasmithFacadeCSharp", EntryPoint="CSharp_FDatasmithFacadeTransformAnimation_SWIGUpcast")]
+  public static extern global::System.IntPtr FDatasmithFacadeTransformAnimation_SWIGUpcast(global::System.IntPtr jarg1);
+
+  [global::System.Runtime.InteropServices.DllImport("DatasmithFacadeCSharp", EntryPoint="CSharp_FDatasmithFacadeVisibilityAnimation_SWIGUpcast")]
+  public static extern global::System.IntPtr FDatasmithFacadeVisibilityAnimation_SWIGUpcast(global::System.IntPtr jarg1);
+
+  [global::System.Runtime.InteropServices.DllImport("DatasmithFacadeCSharp", EntryPoint="CSharp_FDatasmithFacadeSubsequenceAnimation_SWIGUpcast")]
+  public static extern global::System.IntPtr FDatasmithFacadeSubsequenceAnimation_SWIGUpcast(global::System.IntPtr jarg1);
+
+  [global::System.Runtime.InteropServices.DllImport("DatasmithFacadeCSharp", EntryPoint="CSharp_FDatasmithFacadeLevelSequence_SWIGUpcast")]
+  public static extern global::System.IntPtr FDatasmithFacadeLevelSequence_SWIGUpcast(global::System.IntPtr jarg1);
+
   [global::System.Runtime.InteropServices.DllImport("DatasmithFacadeCSharp", EntryPoint="CSharp_FDatasmithFacadePropertyCapture_SWIGUpcast")]
   public static extern global::System.IntPtr FDatasmithFacadePropertyCapture_SWIGUpcast(global::System.IntPtr jarg1);
 
diff --git a/Engine/Source/Programs/Enterprise/Datasmith/DatasmithFacadeCSharp/Public/FDatasmithFacadeScene.cs b/Engine/Source/Programs/Enterprise/Datasmith/DatasmithFacadeCSharp/Public/FDatasmithFacadeScene.cs
--- a/Engine/Source/Programs/Enterprise/Datasmith/DatasmithFacadeCSharp/Public/FDatasmithFacadeScene.cs
+++ b/Engine/Source/Programs/Enterprise/Datasmith/DatasmithFacadeCSharp/Public/FDatasmithFacadeScene.cs
@@ -233,6 +233,25 @@ public class FDatasmithFacadeScene : global::System.IDisposable {
     DatasmithFacadeCSharpPINVOKE.FDatasmithFacadeScene_RemoveLevelVariantSets(swigCPtr, FDatasmithFacadeLevelVariantSets.getCPtr(InLevelVariantSetsPtr));
   }
 
+  public void AddLevelSequence(FDatasmithFacadeLevelSequence InLevelSequence) {
+    DatasmithFacadeCSharpPINVOKE.FDatasmithFacadeScene_AddLevelSequence(swigCPtr, FDatasmithFacadeLevelSequence.getCPtr(InLevelSequence));
+  }
+
+  public int GetLevelSequencesCount() {
+    int ret = DatasmithFacadeCSharpPINVOKE.FDatasmithFacadeScene_GetLevelSequencesCount(swigCPtr);
+    return ret;
+  }
+
+  public FDatasmithFacadeLevelSequence GetLevelSequence(int LevelSequenceIndex) {
+    global::System.IntPtr cPtr = DatasmithFacadeCSharpPINVOKE.FDatasmithFacadeScene_GetLevelSequence(swigCPtr, LevelSequenceIndex);
+    FDatasmithFacadeLevelSequence ret = (cPtr == global::System.IntPtr.Zero) ? null : new FDatasmithFacadeLevelSequence(cPtr, false);
+    return ret;
+  }
+
+  public void RemoveLevelSequence(FDatasmithFacadeLevelSequence InLevelSequence) {
+    DatasmithFacadeCSharpPINVOKE.FDatasmithFacadeScene_RemoveLevelSequence(swigCPtr, FDatasmithFacadeLevelSequence.getCPtr(InLevelSequence));
+  }
+
   public void AddMetaData(FDatasmithFacadeMetaData InMetaDataPtr) {
     DatasmithFacadeCSharpPINVOKE.FDatasmithFacadeScene_AddMetaData(swigCPtr, FDatasmithFacadeMetaData.getCPtr(InMetaDataPtr));
   }
diff --git a/Engine/Source/Programs/Enterprise/Datasmith/DatasmithSolidworksExporter/Private/Animations/AnimationExtractor.cs b/Engine/Source/Programs/Enterprise/Datasmith/DatasmithSolidworksExporter/Private/Animations/AnimationExtractor.cs
--- a/Engine/Source/Programs/Enterprise/Datasmith/DatasmithSolidworksExporter/Private/Animations/AnimationExtractor.cs
+++ b/Engine/Source/Programs/Enterprise/Datasmith/DatasmithSolidworksExporter/Private/Animations/AnimationExtractor.cs
@@ -1,130 +1,157 @@
 ﻿// Copyright Epic Games, Inc. All Rights Reserved.
 
-using System;
-using System.Runtime.InteropServices;
-using System.IO;
-using System.Drawing;
 using System.Collections.Generic;
-using System.Linq;
-using System.Text;
 using SolidWorks.Interop.sldworks;
 using SolidWorks.Interop.swmotionstudy;
-using SolidworksDatasmith.Geometry;
 
-using SolidworksDatasmith.SwObjects;
-
-namespace SolidworksDatasmith.Animations
+namespace SolidworksDatasmith.SwObjects
 {
-    class AnimationExtractor
+    public class AnimationExtractor
     {
-        void ExtractAnimations(IModelDoc2 doc, Component2 root)
+        public static List<SwAnimation> ExtractAnimations(IAssemblyDoc Doc, Component2 Root)
         {
-            var ext = doc.Extension;
-            var manager = ext.GetMotionStudyManager() as IMotionStudyManager;
-            if (manager != null)
+            ModelDocExtension Ext = (Doc as IModelDoc2).Extension;
+			IMotionStudyManager MSManager = (Ext.GetMotionStudyManager() as IMotionStudyManager);
+
+			if (MSManager == null)
+			{
+				return null;
+			}
+
+			int StudyCount = MSManager.GetMotionStudyCount();
+			if (StudyCount == 0)
+			{
+				return null;
+			}
+
+			List<SwAnimation> Animations = new List<SwAnimation>();
+
+            bool bWasInPresentationMode = Doc.EnablePresentation;
+            
+            if (!bWasInPresentationMode)
             {
-                int num = manager.GetMotionStudyCount();
-                if (num > 0)
-                {
-                    string[] names = manager.GetMotionStudyNames();
-                    for (int i = 0; i < num; i++)
-                    {
-                        var study = manager.GetMotionStudy(names[i]);
-                        if (study != null)
-                            ExtractAnimation(study, names[i], root);
-                    }
-                }
+                // Enter presentation mode in order to get valid PresentationTransform 
+                // for components
+                Doc.EnablePresentation = true;
+            }
+            
+			string[] StudyNames = MSManager.GetMotionStudyNames();
+			for (int Idx = 0; Idx < StudyCount; Idx++)
+			{
+				MotionStudy Study = MSManager.GetMotionStudy(StudyNames[Idx]);
+				if (Study != null)
+				{
+					SwAnimation Anim = ExtractAnimation(Study, Root);
+					if (Anim != null)
+					{
+						Animations.Add(Anim);
+					}
+				}
+			}
+			
+			if (!bWasInPresentationMode)
+			{
+                Doc.EnablePresentation = false; // Return back to normal mode
             }
+            
+			return Animations;
         }
 
-        void ExtractAnimation(IMotionStudy study, string name, Component2 root)
+        static SwAnimation ExtractAnimation(IMotionStudy Study, Component2 Root)
         {
-            bool success = study.Activate();
-
-	        //if (result)
-	        {
-                var duration = study.GetDuration();
-		        if (duration > 0.0)
-		        {
-                    SwAnimation anim = new SwAnimation();
-                    anim.Name = name;
-                    double initialTime = study.GetTime();
-
-                    //swMotionStudyTypeAssembly   1 or 0x1 = Animation; D - cubed solver is used to do presentation animation only; no simulation is performed, so no results or plots are available; gravity, contact, springs, and forces cannot be used; mass and inertia values have no effect on the animation
-                    //swMotionStudyTypeCosmosMotion   4 or 0x4 = Motion Analysis; ADAMS(MSC.Software) solver is used to return accurate results; you must load the SOLIDWORKS Motion add -in with a SOLIDWORKS premium license to use this option
-                    //swMotionStudyTypeLegacyCosmosMotion 8 or 0x8 = Legacy COSMOSMotion; in SOLIDWORKS 2007 and earlier, motion analysis was provided through the COSMOSMotion add -in; this option is available if either the COSMOSMotion add -in is loaded or you open an older model that was created using that add-in; models with legacy COSMOSMotion data can be opened but not edited
-                    //swMotionStudyTypeNewCosmosMotion    16 or 0x10
-                    //swMotionStudyTypePhysicalSimulation
-                    int supportedTypes;
-                    study.GetSupportedStudyTypes(out supportedTypes);
-                    var properties = study.GetProperties((int)swMotionStudyType_e.swMotionStudyTypeAssembly);
-
-                    int fps = properties.GetFrameRate();
-                    double step = 1.0 / (double)fps;
-                    int st = 0;
-                    uint allSteps = (uint)(duration / step);
-
-                    var rootTm = root.GetTotalTransform(true);
-			        
-			        for (double t = 0.0; t<duration; t += step)
-			        {
-                        study.SetTime(t);
-                        ExtractAnimationData(st++, t, root, anim, rootTm);
-                    }
-                    study.SetTime(duration);
-                    ExtractAnimationData(st++, duration, root, anim, rootTm);
+            bool Success = Study.Activate();
+
+            double Duration = Study.GetDuration();
+
+			if (Duration <= 0.0)
+			{
+				return null;
+			}
+
+			SwAnimation Anim = new SwAnimation();
+			Anim.Name = Study.Name;
+
+			double InitialTime = Study.GetTime();
+
+			/*
+			 * swMotionStudyTypeAssembly			1 or 0x1 = Animation; D - cubed solver is used to do presentation animation only; no simulation is performed, so no results or plots are available; gravity, contact, springs, and forces cannot be used; mass and inertia values have no effect on the animation
+			 * swMotionStudyTypeCosmosMotion		4 or 0x4 = Motion Analysis; ADAMS(MSC.Software) solver is used to return accurate results; you must load the SOLIDWORKS Motion add -in with a SOLIDWORKS premium license to use this option
+			 * swMotionStudyTypeLegacyCosmosMotion	8 or 0x8 = Legacy COSMOSMotion; in SOLIDWORKS 2007 and earlier, motion analysis was provided through the COSMOSMotion add -in; this option is available if either the COSMOSMotion add -in is loaded or you open an older model that was created using that add-in; models with legacy COSMOSMotion data can be opened but not edited
+			 * swMotionStudyTypeNewCosmosMotion		16 or 0x10
+			 * swMotionStudyTypePhysicalSimulation
+			 */
+			MotionStudyProperties StudyProps = Study.GetProperties((int)swMotionStudyType_e.swMotionStudyTypeAssembly);
 
-                    study.SetTime(initialTime);
-                    study.Stop();
+			Anim.FPS = StudyProps.GetFrameRate();
+			double Step = 1.0 / Anim.FPS;
+			int StepIndex = 0;
+			uint AllSteps = (uint)(Duration / Step);
 
-                    // make transform matrices local and create actual gltf channels
-                    anim.BakeAnimationFromIntermediateChannels(false);
-		        }
-	        }
+			MathTransform RootTransform = Root.GetTotalTransform(true);
+
+			for (double CurTime = 0.0; CurTime < Duration; CurTime += Step)
+			{
+			    Study.SetTime(CurTime);
+			    ExtractAnimationData(StepIndex++, CurTime, Root, Anim, RootTransform);
+			}
+
+			Study.SetTime(Duration);
+			ExtractAnimationData(StepIndex, Duration, Root, Anim, RootTransform);
+
+			Study.SetTime(InitialTime);
+			Study.Stop();
+
+			return Anim;
         }
 
-        void ExtractAnimationData(int step, double t, Component2 component, SwAnimation anim, MathTransform rootTm)
+        static void ExtractAnimationData(int StepIndex, double Time, Component2 Component, SwAnimation Anim, MathTransform RootTransform)
         {
-            if (component == null)
+            if (Component == null)
+			{
                 return;
-            if (!component.IsRoot())
+			}
+
+            if (!Component.IsRoot())
             {
-                var parent = component.GetParent();
-                MathTransform tm = component.GetTotalTransform(true);
-                if (tm == null)
-                    tm = component.Transform2;
-                if (tm != null)
-                {
-                    MathTransform global;
-                    if (parent != null)
-                    {
-                        global = rootTm.IMultiply(tm);
-                        tm = global;
-                    }
-                }
+                MathTransform ComponentTransform = Component.GetTotalTransform(true);
+
+                if (ComponentTransform == null)
+				{
+                    ComponentTransform = Component.Transform2;
+				}
 
-                if (tm != null)
+                if (ComponentTransform != null)
                 {
-                    var channel = anim.getIntermediateChannel(component);
-                    if (channel == null)
+					MathTransform ComponentWorldTransform = RootTransform.IMultiply(ComponentTransform);
+					MathTransform ParentWorldTransform = RootTransform;
+
+					Component2 Parent = Component.GetParent();
+
+					if (Parent != null)
+					{
+						ParentWorldTransform = RootTransform.IMultiply(Parent.GetTotalTransform(true));
+					}
+
+					MathTransform ParentWorldTransformInverse = ParentWorldTransform.Inverse();
+					MathTransform ComponentLocalTransform = ComponentWorldTransform.IMultiply(ParentWorldTransformInverse);
+
+                    SwAnimationIntermediateChannel Channel = Anim.getIntermediateChannel(Component);
+                    if (Channel == null)
                     {
-                        channel = anim.newIntermediateChannel(component);
+                        Channel = Anim.NewIntermediateChannel(Component);
                     }
-                    var key = channel.newKeyframe(step, t);
-                    key.GlobalTm = tm;
+                    SwAnimationIntermediateKeyframe Key = Channel.NewKeyframe(StepIndex, ComponentLocalTransform, Time);
                 }
             }
 
-            if (component.IGetChildrenCount() > 0)
+            if (Component.IGetChildrenCount() > 0)
             {
-                Component2[] children = component.GetChildren();
-                foreach (var child in children)
+                object[] Children = (object[])Component.GetChildren();
+                foreach (Component2 Child in Children)
                 {
-                    ExtractAnimationData(step, t, child, anim, rootTm);
+                    ExtractAnimationData(StepIndex, Time, Child, Anim, RootTransform);
                 }
             }
         }
-
-
     }
 }
diff --git a/Engine/Source/Programs/Enterprise/Datasmith/DatasmithSolidworksExporter/Private/Engine/CommandQueue.cs b/Engine/Source/Programs/Enterprise/Datasmith/DatasmithSolidworksExporter/Private/Engine/CommandQueue.cs
--- a/Engine/Source/Programs/Enterprise/Datasmith/DatasmithSolidworksExporter/Private/Engine/CommandQueue.cs
+++ b/Engine/Source/Programs/Enterprise/Datasmith/DatasmithSolidworksExporter/Private/Engine/CommandQueue.cs
@@ -23,7 +23,8 @@ namespace SolidworksDatasmith.Engine
 		CONFIGURATION_DATA,
         LIVEUPDATE,
         DELETE_COMPONENT,
-        UPDATE_CAMERA
+        UPDATE_CAMERA,
+		ANIMATION
     }
 
     [ComVisible(false)]
@@ -139,6 +140,17 @@ namespace SolidworksDatasmith.Engine
 		}
 	};
 
+	[ComVisible(false)]
+	public class AnimationCommand : Command
+	{
+		public SwObjects.SwAnimation Animation = null;
+
+        public AnimationCommand()
+        {
+            Type = CommandType.ANIMATION;
+        }
+	};
+
 	[ComVisible(false)]
     public class ExportCommand : Command
     {
diff --git a/Engine/Source/Programs/Enterprise/Datasmith/DatasmithSolidworksExporter/Private/Engine/Processor.cs b/Engine/Source/Programs/Enterprise/Datasmith/DatasmithSolidworksExporter/Private/Engine/Processor.cs
--- a/Engine/Source/Programs/Enterprise/Datasmith/DatasmithSolidworksExporter/Private/Engine/Processor.cs
+++ b/Engine/Source/Programs/Enterprise/Datasmith/DatasmithSolidworksExporter/Private/Engine/Processor.cs
@@ -499,6 +499,28 @@ namespace SolidworksDatasmith.Engine
                                     processor.DatasmithScene.AddActor(dsCamera);
                             }
                             break;
+                        case CommandType.ANIMATION:
+                        {
+							var cmd = command as AnimationCommand;
+							SwSingleton.FireProgressEvent($"Exporting animation \"{cmd.Animation.Name}\"...");
+							FDatasmithFacadeLevelSequence LevelSeq = cmd.Animation.Export();
+
+							// Check if we already have a sequence with the same name and remove it if we do
+							int SequencesSetsCount = processor.DatasmithScene.GetLevelSequencesCount();
+							for (int Index = 0; Index < SequencesSetsCount; ++Index)
+							{
+								FDatasmithFacadeLevelSequence ExistingSeq = processor.DatasmithScene.GetLevelSequence(Index);
+
+								if (ExistingSeq.GetName() == LevelSeq.GetName())
+								{
+									processor.DatasmithScene.RemoveLevelSequence(ExistingSeq);
+									break;
+								}
+							}
+
+							processor.DatasmithScene.AddLevelSequence(LevelSeq);
+						}
+						break;
                     }
                 }
                 else
@@ -950,7 +972,8 @@ namespace SolidworksDatasmith.Engine
 			{
 				FDatasmithFacadeVariantSet VSet = LevelVariantSets.GetVariantSet(VariantSetIndex);
 
-				if (VSet.GetName() == ConfigurationsSetName){
+				if (VSet.GetName() == ConfigurationsSetName)
+				{
 					VariantSet = VSet;
 					break;
 				}
diff --git a/Engine/Source/Programs/Enterprise/Datasmith/DatasmithSolidworksExporter/Private/Geometry/MathUtil.cs b/Engine/Source/Programs/Enterprise/Datasmith/DatasmithSolidworksExporter/Private/Geometry/MathUtil.cs
--- a/Engine/Source/Programs/Enterprise/Datasmith/DatasmithSolidworksExporter/Private/Geometry/MathUtil.cs
+++ b/Engine/Source/Programs/Enterprise/Datasmith/DatasmithSolidworksExporter/Private/Geometry/MathUtil.cs
@@ -1,6 +1,7 @@
 ﻿// Copyright Epic Games, Inc. All Rights Reserved.
 
 using SolidWorks.Interop.sldworks;
+using System;
 using System.Runtime.InteropServices;
 
 namespace SolidworksDatasmith.Geometry
@@ -10,7 +11,7 @@ namespace SolidworksDatasmith.Geometry
     {
         private static MathUtility _util = null;
 
-        public static MathUtility Utility
+        public static MathUtility SWUtility
         {
             get
             {
@@ -25,22 +26,22 @@ namespace SolidworksDatasmith.Geometry
 
         public static MathVector CreateVector(double x, double y, double z)
         {
-            return (MathVector)Utility.CreateVector(new[] { x, y, z }); ;
+            return (MathVector)SWUtility.CreateVector(new[] { x, y, z }); ;
         }
 
         public static MathVector CreateVector(float x, float y, float z)
         {
-            return (MathVector)Utility.CreateVector(new[] { x, y, z }); ;
+            return (MathVector)SWUtility.CreateVector(new[] { x, y, z }); ;
         }
 
         public static MathPoint CreatePoint(double x, double y, double z)
         {
-            return (MathPoint)Utility.CreatePoint(new[] { x, y, z }); ;
+            return (MathPoint)SWUtility.CreatePoint(new[] { x, y, z }); ;
         }
 
         public static MathPoint CreatePoint(float x, float y, float z)
         {
-            return (MathPoint)Utility.CreatePoint(new[] { x, y, z }); ;
+            return (MathPoint)SWUtility.CreatePoint(new[] { x, y, z }); ;
         }
 
         // a solidworks matrix indices work as follows:
@@ -57,12 +58,12 @@ namespace SolidworksDatasmith.Geometry
             0.0, 0.0, 0.0, 0.0
         };
 
-        // datasmith matrix:
-        // 0 4 8  12 (tx)
-        // 1 5 9  13 (ty)
-        // 2 6 10 14 (tz)
-        // 3(0.0) 7(0.0) 11(0.0) 15(1.0)
-        public static float[] ConvertFromSolidworksTransform(MathTransform tm)
+		// datasmith matrix:
+		// 0 4 8  12 (tx)
+		// 1 5 9  13 (ty)
+		// 2 6 10 14 (tz)
+		// 3(0.0) 7(0.0) 11(0.0) 15(1.0)
+		public static float[] ConvertFromSolidworksTransform(MathTransform tm)
         {
             float[] t = new float[16];
             double[] tt = null;
@@ -88,12 +89,14 @@ namespace SolidworksDatasmith.Geometry
 
         public static MathTransform CreateTransform()
         {
-            return (MathTransform)Utility.CreateTransform(tIdenity);
+            return (MathTransform)SWUtility.CreateTransform(tIdenity);
         }
 
         public static MathTransform CreateTransformRotateAxis(Vec3 axis, float rotation)
         {
-            return (MathTransform)Utility.CreateTransformRotateAxis(Utility.CreatePoint(new double[] { 0.0, 0.0, 0.0 }), Utility.CreateVector(new double[] { (double)axis.x, (double)axis.y, (double)axis.z }), (double)rotation);
+            return (MathTransform)SWUtility.CreateTransformRotateAxis(
+				SWUtility.CreatePoint(new double[] { 0.0, 0.0, 0.0 }), 
+				SWUtility.CreateVector(new double[] { (double)axis.x, (double)axis.y, (double)axis.z }), (double)rotation);
         }
-    }
+	}
 }
diff --git a/Engine/Source/Programs/Enterprise/Datasmith/DatasmithSolidworksExporter/Private/Geometry/Matrix4.cs b/Engine/Source/Programs/Enterprise/Datasmith/DatasmithSolidworksExporter/Private/Geometry/Matrix4.cs
--- a/Engine/Source/Programs/Enterprise/Datasmith/DatasmithSolidworksExporter/Private/Geometry/Matrix4.cs
+++ b/Engine/Source/Programs/Enterprise/Datasmith/DatasmithSolidworksExporter/Private/Geometry/Matrix4.cs
@@ -65,6 +65,12 @@ namespace SolidworksDatasmith.Geometry
             return m._m;
         }
 
+		public float this[int Index]
+		{
+			get => _m[Index];
+			set => _m[Index] = value;
+		}
+
         public static Matrix4 Identity
         {
             get
diff --git a/Engine/Source/Programs/Enterprise/Datasmith/DatasmithSolidworksExporter/Private/Geometry/Vec3.cs b/Engine/Source/Programs/Enterprise/Datasmith/DatasmithSolidworksExporter/Private/Geometry/Vec3.cs
--- a/Engine/Source/Programs/Enterprise/Datasmith/DatasmithSolidworksExporter/Private/Geometry/Vec3.cs
+++ b/Engine/Source/Programs/Enterprise/Datasmith/DatasmithSolidworksExporter/Private/Geometry/Vec3.cs
@@ -61,10 +61,9 @@ namespace SolidworksDatasmith.Geometry
 
         public Vec3(MathVector swVec)
         {
-            double[] data = swVec.ArrayData();
-            x = (float)data[0];
-            y = (float)data[1];
-            z = (float)data[2];
+            x = (float)swVec.ArrayData[0];
+            y = (float)swVec.ArrayData[1];
+            z = (float)swVec.ArrayData[2];
         }
 
         public override bool Equals(object obj)
diff --git a/Engine/Source/Programs/Enterprise/Datasmith/DatasmithSolidworksExporter/Private/SwAddin.cs b/Engine/Source/Programs/Enterprise/Datasmith/DatasmithSolidworksExporter/Private/SwAddin.cs
--- a/Engine/Source/Programs/Enterprise/Datasmith/DatasmithSolidworksExporter/Private/SwAddin.cs
+++ b/Engine/Source/Programs/Enterprise/Datasmith/DatasmithSolidworksExporter/Private/SwAddin.cs
@@ -326,7 +326,7 @@ namespace SolidworksDatasmith
 
 			sFileName = ParseFilename(sFileName);
 
-			SwSingleton.CurrentScene.EvaluateScene();
+			SwSingleton.CurrentScene.EvaluateScene(false, false);
 
 			ExportCommand cmd = new ExportCommand();
 			cmd.Path = Path.GetDirectoryName(sFileName);
@@ -499,7 +499,7 @@ namespace SolidworksDatasmith
 		#region UI Callbacks
 		public void DirectLinkUpdate()
 		{
-			SwSingleton.CurrentScene.DirectLinkUpdate();
+			SwSingleton.CurrentScene.DirectLinkUpdate(true);
 		}
 
 		public int DirectLinkUpdateStatus()
@@ -519,8 +519,8 @@ namespace SolidworksDatasmith
 			if (SwSingleton.CurrentScene.bDirectLinkAutoSync)
 			{
 				// Run update on auto sync enable, in case we got changes
-				SwSingleton.CurrentScene.DirectLinkUpdate();
-		}
+				SwSingleton.CurrentScene.DirectLinkUpdate(false);
+			}
 		}
 
 		public int EnableDisableLinkStatus()
@@ -762,7 +762,7 @@ namespace SolidworksDatasmith
 			{
 				if (SwSingleton.CurrentScene.bIsDirty && SwSingleton.CurrentScene.bDirectLinkAutoSync)
 				{
-					SwSingleton.CurrentScene.DirectLinkUpdate();
+					SwSingleton.CurrentScene.DirectLinkUpdate(false);
 				}
 			}
 			return 0;
diff --git a/Engine/Source/Programs/Enterprise/Datasmith/DatasmithSolidworksExporter/Private/SwObjects/SwAnimation.cs b/Engine/Source/Programs/Enterprise/Datasmith/DatasmithSolidworksExporter/Private/SwObjects/SwAnimation.cs
--- a/Engine/Source/Programs/Enterprise/Datasmith/DatasmithSolidworksExporter/Private/SwObjects/SwAnimation.cs
+++ b/Engine/Source/Programs/Enterprise/Datasmith/DatasmithSolidworksExporter/Private/SwObjects/SwAnimation.cs
@@ -1,14 +1,8 @@
 ﻿// Copyright Epic Games, Inc. All Rights Reserved.
 
 using System;
-using System.Runtime.InteropServices;
-using System.IO;
-using System.Drawing;
 using System.Collections.Generic;
-using System.Linq;
-using System.Text;
 using SolidWorks.Interop.sldworks;
-using SolidWorks.Interop.swmotionstudy;
 using SolidworksDatasmith.Geometry;
 
 namespace SolidworksDatasmith.SwObjects
@@ -24,24 +18,24 @@ namespace SolidworksDatasmith.SwObjects
 			PATH_SCALE
 		}
 
-		private Dictionary<Component2, SwAnimationIntermediateChannel> node2IntermediateChannel = new Dictionary<Component2, SwAnimationIntermediateChannel>();
-
+		private Dictionary<string, SwAnimationIntermediateChannel> node2IntermediateChannel = new Dictionary<string, SwAnimationIntermediateChannel>();
 		public List<SwAnimationChannel> Channels = new List<SwAnimationChannel>();
 		public List<SwAnimationIntermediateChannel> IntermediateChannels = new List<SwAnimationIntermediateChannel>();
+		public int FPS {  get; set; } = 0;
 
-		public SwAnimationIntermediateChannel newIntermediateChannel(Component2 target)
+		public SwAnimationIntermediateChannel NewIntermediateChannel(Component2 target)
 		{
 			SwAnimationIntermediateChannel channel = new SwAnimationIntermediateChannel();
 			channel.Target = target;
 			IntermediateChannels.Add(channel);
-			node2IntermediateChannel.Add(target, channel);
+			node2IntermediateChannel.Add(target.Name2, channel);
 			return channel;
 		}
 
 		public SwAnimationIntermediateChannel getIntermediateChannel(Component2 target)
 		{
 			SwAnimationIntermediateChannel channel = null;
-			if (!node2IntermediateChannel.TryGetValue(target, out channel))
+			if (!node2IntermediateChannel.TryGetValue(target.Name2, out channel))
 				return null;
 			return channel;
 		}
@@ -84,6 +78,42 @@ namespace SolidworksDatasmith.SwObjects
 			return xform;
 		}
 
+		public FDatasmithFacadeLevelSequence Export()
+		{
+			FDatasmithFacadeLevelSequence LevelSeq = new FDatasmithFacadeLevelSequence(Name);
+
+			LevelSeq.SetFrameRate(FPS);
+
+			foreach (var NodePair in node2IntermediateChannel)
+			{
+				SwAnimationIntermediateChannel Chan = NodePair.Value;
+				Component2 Component = Chan.Target;
+				FDatasmithFacadeTransformAnimation Anim = new FDatasmithFacadeTransformAnimation(Component.Name2);
+
+				foreach (var Keyframe in Chan.keyframes)
+				{
+					Matrix4 LocalMatrix = Keyframe.LocalMatrix;
+
+					// Get euler angles in degrees
+					float X = Utility.Rad2Deg * (float)Math.Atan2(LocalMatrix[6], LocalMatrix[10]);
+					float Y = Utility.Rad2Deg * (float)Math.Atan2(-LocalMatrix[2], Math.Sqrt(LocalMatrix[6] * LocalMatrix[6] + LocalMatrix[10] * LocalMatrix[10]));
+					float Z = Utility.Rad2Deg * (float)Math.Atan2(LocalMatrix[1], LocalMatrix[0]);
+
+					float Scale = LocalMatrix[15];
+
+					Vec3 Translation = new Vec3(LocalMatrix[12], LocalMatrix[13], LocalMatrix[14]);
+
+					Anim.AddFrame(EDatasmithFacadeAnimationTransformType.Rotation, Keyframe.Step, X, -Y, -Z);
+					Anim.AddFrame(EDatasmithFacadeAnimationTransformType.Scale, Keyframe.Step, Scale, Scale, Scale);
+					Anim.AddFrame(EDatasmithFacadeAnimationTransformType.Translation, Keyframe.Step, Translation.x, -Translation.y, Translation.z);
+				}
+
+				LevelSeq.AddAnimation(Anim);
+			}
+
+			return LevelSeq;
+		}
+
 		// inject keyframes to pad
 		//
 		public void PadIntermediateFrames()
diff --git a/Engine/Source/Programs/Enterprise/Datasmith/DatasmithSolidworksExporter/Private/SwObjects/SwAnimationIntermediateChannel.cs b/Engine/Source/Programs/Enterprise/Datasmith/DatasmithSolidworksExporter/Private/SwObjects/SwAnimationIntermediateChannel.cs
--- a/Engine/Source/Programs/Enterprise/Datasmith/DatasmithSolidworksExporter/Private/SwObjects/SwAnimationIntermediateChannel.cs
+++ b/Engine/Source/Programs/Enterprise/Datasmith/DatasmithSolidworksExporter/Private/SwObjects/SwAnimationIntermediateChannel.cs
@@ -1,14 +1,7 @@
 ﻿// Copyright Epic Games, Inc. All Rights Reserved.
 
-using System;
-using System.Runtime.InteropServices;
-using System.IO;
-using System.Drawing;
 using System.Collections.Generic;
-using System.Linq;
-using System.Text;
 using SolidWorks.Interop.sldworks;
-using SolidWorks.Interop.swmotionstudy;
 using SolidworksDatasmith.Geometry;
 
 namespace SolidworksDatasmith.SwObjects
@@ -16,20 +9,18 @@ namespace SolidworksDatasmith.SwObjects
 	public class SwAnimationIntermediateChannel
 	{
 		private Dictionary<int, SwAnimationIntermediateKeyframe> id2key = new Dictionary<int, SwAnimationIntermediateKeyframe>();
+		public List<SwAnimationIntermediateKeyframe> keyframes = new List<SwAnimationIntermediateKeyframe>();
 		public Component2 Target;
 		public int interpolation;
-		public List<SwAnimationIntermediateKeyframe> keyframes;
 
-		public SwAnimationIntermediateChannel()
-		{
-		}
-
-		public SwAnimationIntermediateKeyframe newKeyframe(int step, double time = -1.0)
+		public SwAnimationIntermediateKeyframe NewKeyframe(int Step, MathTransform LocalTransform, double Time = -1.0)
 		{
 			SwAnimationIntermediateKeyframe key = new SwAnimationIntermediateKeyframe();
-			key.Step = step;
-			key.Time = time;
+			key.Step = Step;
+			key.Time = Time;
 			key.Owner = this;
+			key.LocalMatrix = new Matrix4(MathUtil.ConvertFromSolidworksTransform(LocalTransform));
+			
 			keyframes.Add(key);
 			id2key.Add(key.Step, key);
 			return key;
diff --git a/Engine/Source/Programs/Enterprise/Datasmith/DatasmithSolidworksExporter/Private/SwObjects/SwAnimationIntermediateKeyframe.cs b/Engine/Source/Programs/Enterprise/Datasmith/DatasmithSolidworksExporter/Private/SwObjects/SwAnimationIntermediateKeyframe.cs
--- a/Engine/Source/Programs/Enterprise/Datasmith/DatasmithSolidworksExporter/Private/SwObjects/SwAnimationIntermediateKeyframe.cs
+++ b/Engine/Source/Programs/Enterprise/Datasmith/DatasmithSolidworksExporter/Private/SwObjects/SwAnimationIntermediateKeyframe.cs
@@ -1,14 +1,6 @@
 ﻿// Copyright Epic Games, Inc. All Rights Reserved.
 
-using System;
-using System.Runtime.InteropServices;
-using System.IO;
-using System.Drawing;
-using System.Collections.Generic;
-using System.Linq;
-using System.Text;
 using SolidWorks.Interop.sldworks;
-using SolidWorks.Interop.swmotionstudy;
 using SolidworksDatasmith.Geometry;
 
 namespace SolidworksDatasmith.SwObjects
@@ -21,6 +13,7 @@ namespace SolidworksDatasmith.SwObjects
 		public MathTransform RelativeTm;
 		public SwAnimationIntermediateChannel Owner;
 		public bool IsPadding;
+		public Matrix4 LocalMatrix;
 
 		public SwAnimationIntermediateKeyframe()
 		{
diff --git a/Engine/Source/Programs/Enterprise/Datasmith/DatasmithSolidworksExporter/Private/SwObjects/SwScene.cs b/Engine/Source/Programs/Enterprise/Datasmith/DatasmithSolidworksExporter/Private/SwObjects/SwScene.cs
--- a/Engine/Source/Programs/Enterprise/Datasmith/DatasmithSolidworksExporter/Private/SwObjects/SwScene.cs
+++ b/Engine/Source/Programs/Enterprise/Datasmith/DatasmithSolidworksExporter/Private/SwObjects/SwScene.cs
@@ -274,20 +274,20 @@ namespace SolidworksDatasmith.SwObjects
 		}
 
 		// builds component tree and relative datasmith objects
-		public void EvaluateScene(bool fast = false)
+		public void EvaluateScene(bool bDirectLinkUpdate, bool bManualDirectLinkUpdate)
 		{
 			if (_doc != null)
 			{
 				if (_doc is AssemblyDoc)
 				{
-					EvaluateScene(_doc as AssemblyDoc, fast);
+					EvaluateScene(_doc as AssemblyDoc, bDirectLinkUpdate, bManualDirectLinkUpdate);
 				}
 				else if (_doc is PartDoc)
 				{
-					EvaluateScene(_doc as PartDoc, fast);
+					EvaluateScene(_doc as PartDoc, bDirectLinkUpdate);
 				}
 			}
-			if (!fast)
+			if (!bDirectLinkUpdate)
 			{
 				bool exportcurrentviewonly = false; //Waiting for Epic decision on this subject
 				if (exportcurrentviewonly == false)
@@ -1056,16 +1056,34 @@ namespace SolidworksDatasmith.SwObjects
 			Processor.AddCommand(cmd);
 		}
 
-		private void EvaluateScene(AssemblyDoc doc, bool fast = false)
+		private void EvaluateScene(AssemblyDoc doc, bool bDirectLinkUpdate, bool bManualDirectLinkUpdate)
 		{
 			var config = ((doc as ModelDoc2).GetActiveConfiguration() as Configuration);
 			var root = config.GetRootComponent3(true);
-			EvaluateSceneRecurse(root, null, fast);
-			if (!fast)
+			EvaluateSceneRecurse(root, null, bDirectLinkUpdate);
+
+			if (!bDirectLinkUpdate)
+			{
 				CollectConfigurationData(doc);
+			}
+
+			if (!bDirectLinkUpdate || bManualDirectLinkUpdate)
+			{
+				// Handle animations
+				List<SwAnimation> Anims = AnimationExtractor.ExtractAnimations(doc, root);
+				if (Anims != null)
+				{
+					foreach (SwAnimation Anim in Anims)
+					{
+						var Cmd = new AnimationCommand();
+						Cmd.Animation = Anim;
+						Processor.AddCommand(Cmd);
+					}
+				}
+			}
 		}
 
-		private void EvaluateScene(PartDoc doc, bool fast = false)
+		private void EvaluateScene(PartDoc doc, bool bDirectLinkUpdate)
 		{
 			foreach (var pp in Parts)
 			{
@@ -1076,12 +1094,15 @@ namespace SolidworksDatasmith.SwObjects
 				cmd.PartName = part.Name;
 				cmd.PartPath = part.PathName;
 				Processor.AddCommand(cmd);
-				if (!fast)
+
+				if (!bDirectLinkUpdate)
+				{
 					SendModelDocMetadataToProcessor(doc as ModelDoc2, part.Name, MetadataCommand.MetadataType.Actor);
+				}
 			}
 		}
 
-		public void DirectLinkUpdate()
+		public void DirectLinkUpdate(bool bManualUpdate)
 		{
 			foreach (var pp in Parts)
 			{
@@ -1089,7 +1110,7 @@ namespace SolidworksDatasmith.SwObjects
 			}
 
 			CollectMaterials();
-			EvaluateScene(true);
+			EvaluateScene(true, bManualUpdate);
 			var cmd = new LiveUpdateCommand();
 			Processor.AddCommand(cmd);
 			bIsDirty = false;
