Fixed warning in SketchUp exporter

diff --git a/Engine/Source/Programs/Enterprise/Datasmith/DatasmithSketchUpRubyExporter/Private/DatasmithSketchUpUtils.cpp b/Engine/Source/Programs/Enterprise/Datasmith/DatasmithSketchUpRubyExporter/Private/DatasmithSketchUpUtils.cpp
--- a/Engine/Source/Programs/Enterprise/Datasmith/DatasmithSketchUpRubyExporter/Private/DatasmithSketchUpUtils.cpp
+++ b/Engine/Source/Programs/Enterprise/Datasmith/DatasmithSketchUpRubyExporter/Private/DatasmithSketchUpUtils.cpp
@@ -226,22 +226,22 @@ void DatasmithSketchUpUtils::SetActorTransform(
 	Imath::Vec3<float> Axis = Quaternion.axis();
 	FQuat Rotation(FVector(Axis.x, Axis.y, Axis.z), Angle);
 
-	check(Rotation.IsNormalized());
+	ensure(Rotation.IsNormalized());
 
 	// Convert the SketchUp right-handed Z-up coordinate translation into an Unreal left-handed Z-up coordinate translation.
 	// To avoid perturbating X, which is forward in Unreal, the handedness conversion is done by flipping the side vector Y.
 	// SketchUp uses inches as internal system unit for all 3D coordinates in the model while Unreal uses centimeters.
 
-	FVector Translation = DatasmithSketchUpUtils::FromSketchUp::ConvertPosition(SMatrix[12], SMatrix[13], SMatrix[14]);
+	FVector3f Translation = DatasmithSketchUpUtils::FromSketchUp::ConvertPosition(SMatrix[12], SMatrix[13], SMatrix[14]);
 
 	// Set the world transform of the Datasmith actor.
-	InActorElement->SetScale(Scale.x, Scale.y, Scale.z, false);
-	InActorElement->SetRotation(Rotation.X, Rotation.Y, Rotation.Z, Rotation.W, false);
-#if UE_LARGE_WORLD_COORDINATES_DISABLED
 	InActorElement->SetTranslation(Translation.X, Translation.Y, Translation.Z, false);
+#if UE_LARGE_WORLD_COORDINATES_DISABLED
+	InActorElement->SetRotation(Rotation.X, Rotation.Y, Rotation.Z, Rotation.W, false);
 #else
-	InActorElement->SetTranslation((float)Translation.X, (float)Translation.Y, (float)Translation.Z, false);
+	InActorElement->SetRotation((float)Rotation.X, (float)Rotation.Y, (float)Rotation.Z, (float)Rotation.W, false);
 #endif
+	InActorElement->SetScale(Scale.x, Scale.y, Scale.z, false);
 }
 
 SUTransformation DatasmithSketchUpUtils::GetComponentInstanceTransform(SUComponentInstanceRef InComponentInstanceRef, SUTransformation const& InWorldTransform)
diff --git a/Engine/Source/Programs/Enterprise/Datasmith/DatasmithSketchUpRubyExporter/Private/DatasmithSketchUpUtils.h b/Engine/Source/Programs/Enterprise/Datasmith/DatasmithSketchUpRubyExporter/Private/DatasmithSketchUpUtils.h
--- a/Engine/Source/Programs/Enterprise/Datasmith/DatasmithSketchUpRubyExporter/Private/DatasmithSketchUpUtils.h
+++ b/Engine/Source/Programs/Enterprise/Datasmith/DatasmithSketchUpRubyExporter/Private/DatasmithSketchUpUtils.h
@@ -14,13 +14,13 @@ namespace DatasmithSketchUpUtils
 			return FVector(float(V.x), float(-V.y), float(V.z));
 		}
 
-		static FORCEINLINE FVector ConvertPosition(double X, double Y, double Z)
+		static FORCEINLINE FVector3f ConvertPosition(double X, double Y, double Z)
 		{
 			const float UnitScaleSketchupToUnreal = 2.54; // centimeters per inch
-			return FVector(float(X * UnitScaleSketchupToUnreal), float(-Y * UnitScaleSketchupToUnreal), float(Z * UnitScaleSketchupToUnreal));
+			return FVector3f(float(X * UnitScaleSketchupToUnreal), float(-Y * UnitScaleSketchupToUnreal), float(Z * UnitScaleSketchupToUnreal));
 		}
 
-		static FORCEINLINE FVector ConvertPosition(const SUPoint3D& V)
+		static FORCEINLINE FVector3f ConvertPosition(const SUPoint3D& V)
 		{
 			return ConvertPosition(V.x, V.y, V.z);
 		}
