Datasmith exporters: Build fix

diff --git a/Engine/Source/Programs/Enterprise/Datasmith/DatasmithFacadeCSharp/Private/DatasmithFacadeCSharp.cpp b/Engine/Source/Programs/Enterprise/Datasmith/DatasmithFacadeCSharp/Private/DatasmithFacadeCSharp.cpp
--- a/Engine/Source/Programs/Enterprise/Datasmith/DatasmithFacadeCSharp/Private/DatasmithFacadeCSharp.cpp
+++ b/Engine/Source/Programs/Enterprise/Datasmith/DatasmithFacadeCSharp/Private/DatasmithFacadeCSharp.cpp
@@ -4936,7 +4936,7 @@ SWIGEXPORT int SWIGSTDCALL CSharp_FDatasmithFacadeTransformAnimation_GetCurveInt
 }
 
 
-SWIGEXPORT void SWIGSTDCALL CSharp_FDatasmithFacadeTransformAnimation_GetFrame(void * jarg1, int jarg2, int jarg3, int * jarg4, void * jarg5, void * jarg6, void * jarg7) {
+SWIGEXPORT void SWIGSTDCALL CSharp_FDatasmithFacadeTransformAnimation_GetFrame(void * jarg1, int jarg2, int jarg3, int * jarg4, double * jarg5, double * jarg6, double * jarg7) {
   FDatasmithFacadeTransformAnimation *arg1 = (FDatasmithFacadeTransformAnimation *) 0 ;
   EDatasmithFacadeAnimationTransformType arg2 ;
   int32 arg3 ;
@@ -4949,21 +4949,9 @@ SWIGEXPORT void SWIGSTDCALL CSharp_FDatasmithFacadeTransformAnimation_GetFrame(v
   arg2 = (EDatasmithFacadeAnimationTransformType)jarg2; 
   arg3 = (int32)jarg3; 
   arg4 = (int32 *)jarg4; 
-  arg5 = (double *)jarg5;
-  if (!arg5) {
-    SWIG_CSharpSetPendingExceptionArgument(SWIG_CSharpArgumentNullException, "double & type is null", 0);
-    return ;
-  } 
-  arg6 = (double *)jarg6;
-  if (!arg6) {
-    SWIG_CSharpSetPendingExceptionArgument(SWIG_CSharpArgumentNullException, "double & type is null", 0);
-    return ;
-  } 
-  arg7 = (double *)jarg7;
-  if (!arg7) {
-    SWIG_CSharpSetPendingExceptionArgument(SWIG_CSharpArgumentNullException, "double & type is null", 0);
-    return ;
-  } 
+  arg5 = (double *)jarg5; 
+  arg6 = (double *)jarg6; 
+  arg7 = (double *)jarg7; 
   ((FDatasmithFacadeTransformAnimation const *)arg1)->GetFrame(arg2,arg3,*arg4,*arg5,*arg6,*arg7);
 }
 
diff --git a/Engine/Source/Programs/Enterprise/Datasmith/DatasmithFacadeCSharp/Public/DatasmithFacadeCSharpPINVOKE.cs b/Engine/Source/Programs/Enterprise/Datasmith/DatasmithFacadeCSharp/Public/DatasmithFacadeCSharpPINVOKE.cs
--- a/Engine/Source/Programs/Enterprise/Datasmith/DatasmithFacadeCSharp/Public/DatasmithFacadeCSharpPINVOKE.cs
+++ b/Engine/Source/Programs/Enterprise/Datasmith/DatasmithFacadeCSharp/Public/DatasmithFacadeCSharpPINVOKE.cs
@@ -1386,7 +1386,7 @@ class DatasmithFacadeCSharpPINVOKE {
   public static extern int FDatasmithFacadeTransformAnimation_GetCurveInterpMode(global::System.Runtime.InteropServices.HandleRef jarg1, int jarg2);
 
   [global::System.Runtime.InteropServices.DllImport("DatasmithFacadeCSharp", EntryPoint="CSharp_FDatasmithFacadeTransformAnimation_GetFrame")]
-  public static extern void FDatasmithFacadeTransformAnimation_GetFrame(global::System.Runtime.InteropServices.HandleRef jarg1, int jarg2, int jarg3, out int jarg4, global::System.Runtime.InteropServices.HandleRef jarg5, global::System.Runtime.InteropServices.HandleRef jarg6, global::System.Runtime.InteropServices.HandleRef jarg7);
+  public static extern void FDatasmithFacadeTransformAnimation_GetFrame(global::System.Runtime.InteropServices.HandleRef jarg1, int jarg2, int jarg3, out int jarg4, ref double jarg5, ref double jarg6, ref double jarg7);
 
   [global::System.Runtime.InteropServices.DllImport("DatasmithFacadeCSharp", EntryPoint="CSharp_FDatasmithFacadeTransformAnimation_RemoveFrame")]
   public static extern void FDatasmithFacadeTransformAnimation_RemoveFrame(global::System.Runtime.InteropServices.HandleRef jarg1, int jarg2, int jarg3);
diff --git a/Engine/Source/Programs/Enterprise/Datasmith/DatasmithFacadeCSharp/Public/FDatasmithFacadeTransformAnimation.cs b/Engine/Source/Programs/Enterprise/Datasmith/DatasmithFacadeCSharp/Public/FDatasmithFacadeTransformAnimation.cs
--- a/Engine/Source/Programs/Enterprise/Datasmith/DatasmithFacadeCSharp/Public/FDatasmithFacadeTransformAnimation.cs
+++ b/Engine/Source/Programs/Enterprise/Datasmith/DatasmithFacadeCSharp/Public/FDatasmithFacadeTransformAnimation.cs
@@ -56,9 +56,8 @@ public class FDatasmithFacadeTransformAnimation : FDatasmithFacadeBaseAnimation
     return ret;
   }
 
-  public void GetFrame(EDatasmithFacadeAnimationTransformType TransformType, int Index, out int OutFrameNumber, SWIGTYPE_p_double OutX, SWIGTYPE_p_double OutY, SWIGTYPE_p_double OutZ) {
-    DatasmithFacadeCSharpPINVOKE.FDatasmithFacadeTransformAnimation_GetFrame(swigCPtr, (int)TransformType, Index, out OutFrameNumber, SWIGTYPE_p_double.getCPtr(OutX), SWIGTYPE_p_double.getCPtr(OutY), SWIGTYPE_p_double.getCPtr(OutZ));
-    if (DatasmithFacadeCSharpPINVOKE.SWIGPendingException.Pending) throw DatasmithFacadeCSharpPINVOKE.SWIGPendingException.Retrieve();
+  public void GetFrame(EDatasmithFacadeAnimationTransformType TransformType, int Index, out int OutFrameNumber, ref double OutX, ref double OutY, ref double OutZ) {
+    DatasmithFacadeCSharpPINVOKE.FDatasmithFacadeTransformAnimation_GetFrame(swigCPtr, (int)TransformType, Index, out OutFrameNumber, ref OutX, ref OutY, ref OutZ);
   }
 
   public void RemoveFrame(EDatasmithFacadeAnimationTransformType TransformType, int Index) {
