Fix Issue #100298

diff --git a/Engine/Source/Programs/Enterprise/Datasmith/DatasmithCADWorker/Private/DatasmithCADWorker.cpp b/Engine/Source/Programs/Enterprise/Datasmith/DatasmithCADWorker/Private/DatasmithCADWorker.cpp
--- a/Engine/Source/Programs/Enterprise/Datasmith/DatasmithCADWorker/Private/DatasmithCADWorker.cpp
+++ b/Engine/Source/Programs/Enterprise/Datasmith/DatasmithCADWorker/Private/DatasmithCADWorker.cpp
@@ -6,6 +6,7 @@
 #include "RequiredProgramMainCPPInclude.h"
 #include "CADInterfacesModule.h"
 #include "CADToolsModule.h"
+#include "CoreTechTypes.h"
 
 #include "Modules/ModuleManager.h"
 
diff --git a/Engine/Source/Programs/Enterprise/Datasmith/DatasmithCADWorker/Private/DatasmithCADWorkerImpl.cpp b/Engine/Source/Programs/Enterprise/Datasmith/DatasmithCADWorker/Private/DatasmithCADWorkerImpl.cpp
--- a/Engine/Source/Programs/Enterprise/Datasmith/DatasmithCADWorker/Private/DatasmithCADWorkerImpl.cpp
+++ b/Engine/Source/Programs/Enterprise/Datasmith/DatasmithCADWorker/Private/DatasmithCADWorkerImpl.cpp
@@ -2,7 +2,8 @@
 
 #include "DatasmithCADWorkerImpl.h"
 
-#include "CoreTechFileParser.h"
+#include "CADFileReader.h"
+#include "CADOptions.h"
 #include "DatasmithCommands.h"
 #include "DatasmithDispatcherConfig.h"
 
@@ -138,7 +139,7 @@ uint64 DefineMaximumAllowedDuration(const CADLibrary::FFileDescription& FileDesc
 	{
 		MaxTimePerMb = 5e-7;
 	}
-	else if (FileDescription.Extension.StartsWith(TEXT("ig"))) // Iges
+	else if (FileDescription.Extension.StartsWith(TEXT("ig"))) // IGES
 	{
 		MaxTimePerMb = 1e-6;
 	}
@@ -149,7 +150,7 @@ uint64 DefineMaximumAllowedDuration(const CADLibrary::FFileDescription& FileDesc
 
 void FDatasmithCADWorkerImpl::ProcessCommand(const FRunTaskCommand& RunTaskCommand)
 {
-	const CADLibrary::FFileDescription& FileToProcess = RunTaskCommand.JobFileDescription;
+	CADLibrary::FFileDescription FileToProcess = RunTaskCommand.JobFileDescription;
 	UE_LOG(LogDatasmithCADWorker, Verbose, TEXT("Process %s %s"), *FileToProcess.Name, *FileToProcess.Configuration);
 
 	FCompletedTaskCommand CompletedTask;
@@ -158,8 +159,8 @@ void FDatasmithCADWorkerImpl::ProcessCommand(const FRunTaskCommand& RunTaskComma
 	int64 MaxDuration = DefineMaximumAllowedDuration(FileToProcess);
 	FThread TimeCheckerThread = FThread(TEXT("TimeCheckerThread"), [&]() { CheckDuration(FileToProcess, MaxDuration); });
 
-	CADLibrary::FCoreTechFileParser FileParser(ImportParameters, EnginePluginsPath, CachePath);
-	ETaskState ProcessResult = FileParser.ProcessFile(FileToProcess);
+	CADLibrary::FCADFileReader FileReader(ImportParameters, FileToProcess, EnginePluginsPath, CachePath);
+	ETaskState ProcessResult = FileReader.ProcessFile();
 
 	bProcessIsRunning = false;
 	TimeCheckerThread.Join();
@@ -168,10 +169,11 @@ void FDatasmithCADWorkerImpl::ProcessCommand(const FRunTaskCommand& RunTaskComma
 
 	if (CompletedTask.ProcessResult == ETaskState::ProcessOk)
 	{
-		CompletedTask.ExternalReferences = FileParser.GetExternalRefSet();
-		CompletedTask.SceneGraphFileName = FileParser.GetSceneGraphFile();
-		CompletedTask.GeomFileName = FileParser.GetMeshFileName();
-		CompletedTask.WarningMessages = FileParser.GetWarningMessages();
+		const CADLibrary::FCADFileData& CADFileData = FileReader.GetCADFileData();
+		CompletedTask.ExternalReferences = CADFileData.GetExternalRefSet();
+		CompletedTask.SceneGraphFileName = CADFileData.GetSceneGraphFileName();
+		CompletedTask.GeomFileName = CADFileData.GetMeshFileName();
+		CompletedTask.WarningMessages = CADFileData.GetWarningMessages();
 	}
 
 	CommandIO.SendCommand(CompletedTask, Config::SendCommandTimeout_s);
