Fix Issue 100924: Compile errors

diff --git a/Engine/Plugins/Enterprise/DatasmithCADImporter/Source/CADTools/Public/CADOptions.h b/Engine/Plugins/Enterprise/DatasmithCADImporter/Source/CADTools/Public/CADOptions.h
--- a/Engine/Plugins/Enterprise/DatasmithCADImporter/Source/CADTools/Public/CADOptions.h
+++ b/Engine/Plugins/Enterprise/DatasmithCADImporter/Source/CADTools/Public/CADOptions.h
@@ -67,13 +67,13 @@ namespace CADLibrary
 			uint32 Hash = 0; 
 			for (double Param : {MetricUnit, ScaleFactor, ChordTolerance, MaxEdgeLength, MaxNormalAngle})
 			{
-				Hash = HashCombine(Hash, GetTypeHash(Param));
+				Hash = HashCombine(Hash, ::GetTypeHash(Param));
 			}
 			for (uint32 Param : {uint32(StitchingTechnique), uint32(ModelCoordSys), uint32(Propagation), uint32(DisplayPreference)})
 			{
-				Hash = HashCombine(Hash, GetTypeHash(Param));
+				Hash = HashCombine(Hash, ::GetTypeHash(Param));
 			}
-			Hash = HashCombine(Hash, GetTypeHash(bScaleUVMap));
+			Hash = HashCombine(Hash, ::GetTypeHash(bScaleUVMap));
 			return Hash;
 		}
 
diff --git a/Engine/Source/Programs/Enterprise/Datasmith/DatasmithCADWorker/Private/DatasmithCADWorkerImpl.cpp b/Engine/Source/Programs/Enterprise/Datasmith/DatasmithCADWorker/Private/DatasmithCADWorkerImpl.cpp
--- a/Engine/Source/Programs/Enterprise/Datasmith/DatasmithCADWorker/Private/DatasmithCADWorkerImpl.cpp
+++ b/Engine/Source/Programs/Enterprise/Datasmith/DatasmithCADWorker/Private/DatasmithCADWorkerImpl.cpp
@@ -121,27 +121,27 @@ void FDatasmithCADWorkerImpl::ProcessCommand(const FImportParametersCommand& Imp
 	ImportParameters = ImportParametersCommand.ImportParameters;
 }
 
-uint64 DefineMaximumAllowedDuration(const CADLibrary::FFileDescription& FileDescription)
+uint64 DefineMaximumAllowedDuration(const CADLibrary::FFileDescriptor& FileDescriptor)
 {
-	FFileStatData FileStatData = IFileManager::Get().GetStatData(*FileDescription.Path);
+	FFileStatData FileStatData = IFileManager::Get().GetStatData(*FileDescriptor.GetSourcePath());
 	double MaxTimePerMb = 5e-6;
 	double SafetyCoeficient = 5;
 
-	if (FileDescription.Extension.StartsWith(TEXT("sld"))) // SW
+	CADLibrary::ECADFormat Format = FileDescriptor.GetFileFormat();
+	switch (Format)
 	{
+	case CADLibrary::ECADFormat::SOLIDWORKS:
+	case CADLibrary::ECADFormat::CATIA_3DXML:
 		MaxTimePerMb = 1e-5;
-	}
-	else if (FileDescription.Extension == TEXT("3dxml") || FileDescription.Extension == TEXT("3drep")) // Catia V5 CGR
-	{
-		MaxTimePerMb = 1e-5;
-	}
-	else if (FileDescription.Extension == TEXT("cgr"))
-	{
+		break;
+	case CADLibrary::ECADFormat::CATIA_CGR:
 		MaxTimePerMb = 5e-7;
-	}
-	else if (FileDescription.Extension.StartsWith(TEXT("ig"))) // IGES
-	{
+		break;
+	case CADLibrary::ECADFormat::IGES:
 		MaxTimePerMb = 1e-6;
+		break;
+	default:
+		break;
 	}
 
 	uint64 MaximumDuration = ((double)FileStatData.FileSize) * MaxTimePerMb * SafetyCoeficient;
@@ -150,8 +150,8 @@ uint64 DefineMaximumAllowedDuration(const CADLibrary::FFileDescription& FileDesc
 
 void FDatasmithCADWorkerImpl::ProcessCommand(const FRunTaskCommand& RunTaskCommand)
 {
-	CADLibrary::FFileDescription FileToProcess = RunTaskCommand.JobFileDescription;
-	UE_LOG(LogDatasmithCADWorker, Verbose, TEXT("Process %s %s"), *FileToProcess.Name, *FileToProcess.Configuration);
+	CADLibrary::FFileDescriptor FileToProcess = RunTaskCommand.JobFileDescription;
+	UE_LOG(LogDatasmithCADWorker, Verbose, TEXT("Process %s %s"), *FileToProcess.GetFileName(), *FileToProcess.GetConfiguration());
 
 	FCompletedTaskCommand CompletedTask;
 
@@ -178,10 +178,10 @@ void FDatasmithCADWorkerImpl::ProcessCommand(const FRunTaskCommand& RunTaskComma
 
 	CommandIO.SendCommand(CompletedTask, Config::SendCommandTimeout_s);
 
-	UE_LOG(LogDatasmithCADWorker, Verbose, TEXT("End of Process %s %s saved in %s"), *FileToProcess.Name, *FileToProcess.Configuration, *CompletedTask.GeomFileName);
+	UE_LOG(LogDatasmithCADWorker, Verbose, TEXT("End of Process %s %s saved in %s"), *FileToProcess.GetFileName(), *FileToProcess.GetConfiguration(), *CompletedTask.GeomFileName);
 }
 
-void FDatasmithCADWorkerImpl::CheckDuration(const CADLibrary::FFileDescription& FileToProcess, const int64 MaxDuration)
+void FDatasmithCADWorkerImpl::CheckDuration(const CADLibrary::FFileDescriptor& FileToProcess, const int64 MaxDuration)
 {
 	if (!ImportParameters.bEnableTimeControl)
 	{
@@ -196,7 +196,7 @@ void FDatasmithCADWorkerImpl::CheckDuration(const CADLibrary::FFileDescription&
 		FPlatformProcess::Sleep(0.1f);
 		if (FPlatformTime::Cycles64() > MaxCycles)
 		{
-			UE_LOG(LogDatasmithCADWorker, Verbose, TEXT("Time exceeded to process %s %s. The maximum allowed duration is %ld s"), *FileToProcess.Name, *FileToProcess.Configuration, MaxDuration);
+			UE_LOG(LogDatasmithCADWorker, Verbose, TEXT("Time exceeded to process %s %s. The maximum allowed duration is %ld s"), *FileToProcess.GetFileName(), *FileToProcess.GetConfiguration(), MaxDuration);
 			FPlatformMisc::RequestExit(true);
 		}
 	}
diff --git a/Engine/Source/Programs/Enterprise/Datasmith/DatasmithCADWorker/Private/DatasmithCADWorkerImpl.h b/Engine/Source/Programs/Enterprise/Datasmith/DatasmithCADWorker/Private/DatasmithCADWorkerImpl.h
--- a/Engine/Source/Programs/Enterprise/Datasmith/DatasmithCADWorker/Private/DatasmithCADWorkerImpl.h
+++ b/Engine/Source/Programs/Enterprise/Datasmith/DatasmithCADWorker/Private/DatasmithCADWorkerImpl.h
@@ -30,7 +30,7 @@ private:
 	/**
 	 * @param MaxDuration is defined in seconds.
 	 */
-	void CheckDuration(const CADLibrary::FFileDescription& FileToProcess, const int64 MaxDuration); 
+	void CheckDuration(const CADLibrary::FFileDescriptor& FileToProcess, const int64 MaxDuration); 
 
 private:
 	DatasmithDispatcher::FNetworkClientNode NetworkInterface;
