Fix build CADWorker

diff --git a/Engine/Source/Programs/Enterprise/Datasmith/DatasmithCADWorker/Private/DatasmithCADWorkerImpl.cpp b/Engine/Source/Programs/Enterprise/Datasmith/DatasmithCADWorker/Private/DatasmithCADWorkerImpl.cpp
--- a/Engine/Source/Programs/Enterprise/Datasmith/DatasmithCADWorker/Private/DatasmithCADWorkerImpl.cpp
+++ b/Engine/Source/Programs/Enterprise/Datasmith/DatasmithCADWorker/Private/DatasmithCADWorkerImpl.cpp
@@ -183,7 +183,7 @@ void FDatasmithCADWorkerImpl::ProcessCommand(const FRunTaskCommand& RunTaskComma
 
 void FDatasmithCADWorkerImpl::CheckDuration(const CADLibrary::FFileDescriptor& FileToProcess, const int64 MaxDuration)
 {
-	if (!ImportParameters.bEnableTimeControl)
+	if (!ImportParameters.bGEnableTimeControl)
 	{
 		return;
 	}
