3DS MAX DirectLink plugin - very WIP
---
Ho to use.

Build with environment variable set:

DATASMITH_3DSMAX_PLUGIN_WITH_DIRECTLINK=1

Note - it should build DatasmithMax2022.gup in Engine\Binaries\Win64\3DSMax\2022

Add that ..Engine\Binaries\Win64\3DSMax\2022 folder to Max pluging paths. Max menu - Customize/Configure User and System Paths/3rd Party Plugins

Execute Engine\Source\Programs\Enterprise\Datasmith\DatasmithMaxExporter\Extras\3dsMaxScripts\UI\DatasmithDialog.ms

    This will activate the plugin and create UI

Sync/Connections/Export buttons should work
---

Review disclaimer:

- This is very preliminary review!

- All in one file for now
- Some code duplicated from baseline - like exporting hierarchy. Syncing is quite different from straight export!
- Some baseline plugin code changed to be able to reuse - one which seemed to need little interference, less chance to break.

- Export - most hierarchy export, with transforms, refscene; geometry - basic GeomObj only, materials - already using full existing plugin's export code
- Missing export(to be on par with existing) - a lot! Other object types, collision, pivots, instances, forest..
- Change tracking - geometry/transforms/material assignment/material changes basically working. But lot to add/fix(e.g. hacky material re-export)
- Slate UI working
- Key functions exported to MaxScript

- Built/tested only in 2022

diff --git a/Engine/Source/Programs/Enterprise/Datasmith/DatasmithMaxExporter/DatasmithMax2017.Build.cs b/Engine/Source/Programs/Enterprise/Datasmith/DatasmithMaxExporter/DatasmithMax2017.Build.cs
--- a/Engine/Source/Programs/Enterprise/Datasmith/DatasmithMaxExporter/DatasmithMax2017.Build.cs
+++ b/Engine/Source/Programs/Enterprise/Datasmith/DatasmithMaxExporter/DatasmithMax2017.Build.cs
@@ -5,19 +5,52 @@ namespace UnrealBuildTool.Rules
 {
 	public abstract class DatasmithMaxBase : ModuleRules
 	{
+		// todo: duplicated in TargetRules
+		public bool bBuildExporterPlugin = true;
+
 		public DatasmithMaxBase(ReadOnlyTargetRules Target)
 			: base(Target)
 		{
+
 			bUseRTTI = true;
 
-			PrivateDependencyModuleNames.AddRange(
-				new string[]
-				{
+			// todo: duplicated in TargetRules
+			string BuildDirectLinkEnvVar = System.Environment.GetEnvironmentVariable("DATASMITH_3DSMAX_PLUGIN_WITH_DIRECTLINK");
+			if (BuildDirectLinkEnvVar != null && BuildDirectLinkEnvVar != "")
+			{
+				bBuildExporterPlugin  = false;
+			}
+
+			if (!bBuildExporterPlugin)
+            {
+				PublicDefinitions.Add("NEW_DIRECTLINK_PLUGIN=1");
+            }
+
+			if (bBuildExporterPlugin)
+			{
+				PrivateDependencyModuleNames.AddRange(
+					new string[]
+					{
 					"Core",
 					"DatasmithExporter",
 					"Projects",
-				}
-			);
+					}
+				);
+			}
+            else
+            {
+				PrivateDependencyModuleNames.AddRange(
+					new string[]
+					{
+					"DatasmithExporter",
+					"DatasmithExporterUI",
+
+					"UdpMessaging", // required for DirectLink networking
+					"UEOpenExr",
+					}
+				);
+			}
+
 
 			PrivateIncludePaths.AddRange( new string[] { "Runtime/Launch/Public", "Runtime/Launch/Private", ModuleDirectory } );
 
diff --git a/Engine/Source/Programs/Enterprise/Datasmith/DatasmithMaxExporter/DatasmithMax2017.Target.cs b/Engine/Source/Programs/Enterprise/Datasmith/DatasmithMaxExporter/DatasmithMax2017.Target.cs
--- a/Engine/Source/Programs/Enterprise/Datasmith/DatasmithMaxExporter/DatasmithMax2017.Target.cs
+++ b/Engine/Source/Programs/Enterprise/Datasmith/DatasmithMaxExporter/DatasmithMax2017.Target.cs
@@ -5,6 +5,8 @@ using UnrealBuildTool;
 [SupportedPlatforms("Win64")]
 public abstract class DatasmithMaxBaseTarget : TargetRules
 {
+	// todo: duplicated in ModuleRules
+	bool bBuildExporterPlugin = true;
 	public DatasmithMaxBaseTarget(TargetInfo Target)
 		: base(Target)
 	{
@@ -16,7 +18,22 @@ public abstract class DatasmithMaxBaseTarget : TargetRules
 		bShouldCompileAsDLL = true;
 		LinkType = TargetLinkType.Monolithic;
 
-		WindowsPlatform.ModuleDefinitionFile = "Programs/Enterprise/Datasmith/DatasmithMaxExporter/DatasmithMaxExporter.def";
+		// todo: duplicated in ModuleRules
+		string BuildDirectLinkEnvVar = System.Environment.GetEnvironmentVariable("DATASMITH_3DSMAX_PLUGIN_WITH_DIRECTLINK");
+		if (BuildDirectLinkEnvVar != null && BuildDirectLinkEnvVar != "")
+		{
+			bBuildExporterPlugin = false;
+		}
+
+		if (bBuildExporterPlugin)
+		{
+			WindowsPlatform.ModuleDefinitionFile = "Programs/Enterprise/Datasmith/DatasmithMaxExporter/DatasmithMaxExporter.def";
+		}
+		else
+		{
+			WindowsPlatform.ModuleDefinitionFile = "Programs/Enterprise/Datasmith/DatasmithMaxExporter/DatasmithMaxExporterWithDirectLink.def";
+		}
+
 		WindowsPlatform.bStrictConformanceMode = false;
 
 		bBuildDeveloperTools = false;
@@ -29,6 +46,14 @@ public abstract class DatasmithMaxBaseTarget : TargetRules
 
 		bHasExports = true;
 		bForceEnableExceptions = true;
+
+		if (!bBuildExporterPlugin)
+		{
+			GlobalDefinitions.Add("UE_EXTERNAL_PROFILING_ENABLED=0"); // For DirectLinkUI (see FDatasmithExporterManager::FInitOptions)
+		}
+
+		// todo: remove?
+		// bSupportEditAndContinue = true;
 	}
 
 	protected void AddCopyPostBuildStep(TargetInfo Target)
@@ -41,7 +66,17 @@ public abstract class DatasmithMaxBaseTarget : TargetRules
 		}
 
 		string SrcOutputFileName = string.Format(@"$(EngineDir)\Binaries\Win64\{0}\{1}.dll", ExeBinariesSubFolder, OutputName);
-		string DstOutputFileName = string.Format(@"$(EngineDir)\Binaries\Win64\{0}\{1}.dle", ExeBinariesSubFolder, OutputName);
+
+		string DstOutputFileName;
+
+		if (bBuildExporterPlugin)
+		{
+			DstOutputFileName = string.Format(@"$(EngineDir)\Binaries\Win64\{0}\{1}.dle", ExeBinariesSubFolder, OutputName);
+		}
+        else
+        {
+			DstOutputFileName = string.Format(@"$(EngineDir)\Binaries\Win64\{0}\{1}.gup", ExeBinariesSubFolder, OutputName);
+        }
 		PostBuildSteps.Add(string.Format("echo Copying {0} to {1}...", SrcOutputFileName, DstOutputFileName));
 		PostBuildSteps.Add(string.Format("copy /Y \"{0}\" \"{1}\" 1>nul", SrcOutputFileName, DstOutputFileName));
 	}
diff --git a/Engine/Source/Programs/Enterprise/Datasmith/DatasmithMaxExporter/DatasmithMax2022.Build.cs b/Engine/Source/Programs/Enterprise/Datasmith/DatasmithMaxExporter/DatasmithMax2022.Build.cs
--- a/Engine/Source/Programs/Enterprise/Datasmith/DatasmithMaxExporter/DatasmithMax2022.Build.cs
+++ b/Engine/Source/Programs/Enterprise/Datasmith/DatasmithMaxExporter/DatasmithMax2022.Build.cs
@@ -1,5 +1,7 @@
 // Copyright Epic Games, Inc. All Rights Reserved.
 
+using System.Diagnostics;
+
 namespace UnrealBuildTool.Rules
 {
 	[SupportedPlatforms("Win64")]
@@ -9,7 +11,6 @@ namespace UnrealBuildTool.Rules
 			: base(Target)
 		{
 		}
-
 		public override string GetMaxVersion() { return "2022"; }
 	}
 }
\ No newline at end of file
diff --git a/Engine/Source/Programs/Enterprise/Datasmith/DatasmithMaxExporter/DatasmithMaxExporterWithDirectLink.def b/Engine/Source/Programs/Enterprise/Datasmith/DatasmithMaxExporter/DatasmithMaxExporterWithDirectLink.def
new file mode 100644
--- /dev/null
+++ b/Engine/Source/Programs/Enterprise/Datasmith/DatasmithMaxExporter/DatasmithMaxExporterWithDirectLink.def
@@ -0,0 +1,9 @@
+LIBRARY
+EXPORTS
+   LibDescription       @1       PRIVATE
+   LibNumberClasses     @2       PRIVATE
+   LibClassDesc         @3       PRIVATE
+   LibVersion           @4       PRIVATE
+   LibInitialize        @6       PRIVATE
+SECTIONS
+   .data READ WRITE
diff --git a/Engine/Source/Programs/Enterprise/Datasmith/DatasmithMaxExporter/Private/DatasmithMaxExporter.cpp b/Engine/Source/Programs/Enterprise/Datasmith/DatasmithMaxExporter/Private/DatasmithMaxExporter.cpp
--- a/Engine/Source/Programs/Enterprise/Datasmith/DatasmithMaxExporter/Private/DatasmithMaxExporter.cpp
+++ b/Engine/Source/Programs/Enterprise/Datasmith/DatasmithMaxExporter/Private/DatasmithMaxExporter.cpp
@@ -2,6 +2,9 @@
 
 #include "DatasmithMaxExporter.h"
 
+#ifndef NEW_DIRECTLINK_PLUGIN
+
+
 #include "Modules/ModuleManager.h"
 #include "GenericPlatform/GenericPlatformFile.h"
 #include "HAL/PlatformFileManager.h"
@@ -11,6 +14,7 @@
 #include "Misc/MessageDialog.h"
 #include "Misc/Paths.h"
 
+
 #include "DatasmithCore.h"
 #include "DatasmithExporterManager.h"
 #include "DatasmithExportOptions.h"
@@ -1142,3 +1146,5 @@ bool FDatasmithExportImpl::Export(const TCHAR* Filename, BOOL bSuppressWarnings)
 }
 
 #include "Windows/HideWindowsPlatformTypes.h"
+
+#endif // NEW_DIRECTLINK_PLUGIN
diff --git a/Engine/Source/Programs/Enterprise/Datasmith/DatasmithMaxExporter/Private/DatasmithMaxExporterWithDirectLink.cpp b/Engine/Source/Programs/Enterprise/Datasmith/DatasmithMaxExporter/Private/DatasmithMaxExporterWithDirectLink.cpp
new file mode 100644
--- /dev/null
+++ b/Engine/Source/Programs/Enterprise/Datasmith/DatasmithMaxExporter/Private/DatasmithMaxExporterWithDirectLink.cpp
@@ -0,0 +1,2122 @@
+// Copyright Epic Games, Inc. All Rights Reserved.
+
+#include "DatasmithMaxExporter.h"
+
+#include "DatasmithMaxSceneExporter.h"
+#include "DatasmithMaxWriter.h"
+#include "DatasmithMaxClassIDs.h"
+
+#include "DatasmithMaxLogger.h"
+
+#ifdef NEW_DIRECTLINK_PLUGIN
+
+#include "Modules/ModuleManager.h"
+#include "GenericPlatform/GenericPlatformFile.h"
+#include "HAL/PlatformFilemanager.h"
+#include "HAL/PlatformProcess.h"
+#include "Misc/CommandLine.h"
+#include "Misc/ConfigCacheIni.h"
+#include "Misc/MessageDialog.h"
+#include "Misc/Paths.h"
+
+#include "DatasmithCore.h"
+#include "DatasmithExporterManager.h"
+#include "DatasmithExportOptions.h"
+#include "DatasmithSceneExporter.h"
+
+#include "DatasmithMesh.h"
+#include "DatasmithMeshExporter.h"
+
+#include "IDatasmithExporterUIModule.h"
+#include "IDirectLinkUI.h"
+
+
+
+#include "DatasmithSceneFactory.h"
+
+#include "DatasmithDirectLink.h"
+
+#include "HAL/PlatformTime.h"
+#include "Misc/DateTime.h"
+#include "Logging/LogMacros.h"
+
+#include "Misc/OutputDeviceRedirector.h" // For GLog
+
+#include "Async/Async.h"
+
+
+DECLARE_LOG_CATEGORY_EXTERN(LogDatasmithMaxExporter, Log, All);
+DEFINE_LOG_CATEGORY(LogDatasmithMaxExporter);
+
+#include "Windows/AllowWindowsPlatformTypes.h"
+MAX_INCLUDES_START
+	#include "Max.h"
+	#include "bitmap.h"
+	#include "gamma.h"
+	#include "IPathConfigMgr.h"
+	#include "maxheapdirect.h"
+	#include "maxscript/maxwrapper/mxsobjects.h"
+
+	#include "maxscript/maxscript.h"
+	#include "maxscript/foundation/numbers.h"
+	#include "maxscript/foundation/arrays.h"
+	#include "maxscript\macros\define_instantiation_functions.h"
+
+	#include "notify.h"
+
+	#include "decomp.h" // for matrix
+
+	#include "MeshNormalSpec.h"
+	#include "ISceneEventManager.h"
+
+
+	#include "IFileResolutionManager.h" // for GetActualPath
+
+	#include "maxicon.h" // for toolbar
+
+MAX_INCLUDES_END
+
+THIRD_PARTY_INCLUDES_START
+#include <locale.h>
+THIRD_PARTY_INCLUDES_END
+
+static FString OriginalLocale( _wsetlocale(LC_NUMERIC, nullptr) ); // Cache LC_NUMERIC locale before initialization of UE4
+static FString NewLocale = _wsetlocale(LC_NUMERIC, TEXT("C"));
+
+HINSTANCE HInstanceMax;
+
+__declspec( dllexport ) bool LibInitialize(void)
+{
+	// Restore LC_NUMERIC locale after initialization of UE4
+	_wsetlocale(LC_NUMERIC, *OriginalLocale);
+	return true;
+}
+
+__declspec(dllexport) const TCHAR* LibDescription()
+{
+	return TEXT("Unreal Datasmith Exporter With DirectLink Support");
+}
+
+// Return version so can detect obsolete DLLs
+__declspec(dllexport) ULONG LibVersion()
+{
+
+	return VERSION_3DSMAX;
+}
+
+__declspec(dllexport) int LibNumberClasses()
+{
+	return 0;
+}
+
+__declspec(dllexport) ClassDesc* LibClassDesc(int i)
+{
+	return nullptr;
+}
+
+/** public functions **/
+BOOL WINAPI DllMain(HINSTANCE hinstDLL, ULONG FdwReason, LPVOID LpvReserved)
+{
+	if (FdwReason == DLL_PROCESS_ATTACH)
+	{
+		MaxSDK::Util::UseLanguagePackLocale();
+		HInstanceMax = hinstDLL;
+		DisableThreadLibraryCalls(HInstanceMax);
+
+		UE_SET_LOG_VERBOSITY(LogDatasmithMaxExporter, Verbose);
+		UE_LOG(LogDatasmithMaxExporter, Error, TEXT("Hello!!!"));
+	}
+	return (TRUE);
+}
+
+///////////////////////////////////////////////////////////////////////////////////////////////////
+
+typedef NodeEventNamespace::NodeKey FNodeKey;
+typedef MtlBase* FMaterialKey;
+typedef Texmap* FTexmapKey;
+
+void LogFlush()
+{
+	Async(EAsyncExecution::TaskGraphMainThread,
+		[]()
+		{
+			GLog->FlushThreadedLogs();
+			GLog->Flush();
+		});
+}
+
+
+void LogDebug(const TCHAR* Msg)
+{
+	mprintf(L"[%s]%s\n", *FDateTime::UtcNow().ToString(TEXT("%Y.%m.%d-%H.%M.%S:%s")), Msg);
+	UE_LOG(LogDatasmithMaxExporter, Error, TEXT("%s"), Msg);
+	LogFlush();
+}
+
+void LogDebug(const FString& Msg)
+{
+	LogDebug(*Msg);
+}
+
+void LogInfo(const TCHAR* Msg)
+{
+	mprintf(L"[%s]%s\n", *FDateTime::UtcNow().ToString(TEXT("%Y.%m.%d-%H.%M.%S:%s")), Msg);
+	UE_LOG(LogDatasmithMaxExporter, Error, TEXT("%s"), Msg);
+}
+
+void LogInfo(const FString& Msg)
+{
+	LogInfo(*Msg);
+}
+
+void LogDebugNode(const TCHAR* Name, INode* Node)
+{
+	// todo: LogDebugNode is called a lot with many nodes - add extra level of debugging
+	return;
+
+
+	LogDebug(FString::Printf(TEXT("%s: %u %s(%d) - %s")
+		, Name
+		, NodeEventNamespace::GetKeyByNode(Node)
+		, Node ? Node->GetName() : L"<null>"
+		, Node ? Node->GetHandle() : 0
+		, (Node && Node->IsNodeHidden(TRUE))? TEXT("HIDDEN") : TEXT("")
+		));
+	if (Node)
+	{
+		LogDebug(FString::Printf(TEXT("    NumberOfChildren: %d "), Node->NumberOfChildren()));
+		
+		if (Object* ObjectRef = Node->GetObjectRef())
+		{
+			Class_ID ClassId = ObjectRef->ClassID();
+			LogDebug(FString::Printf(TEXT("    Class_ID: 0x%lx, 0x%lx "), ClassId.PartA(), ClassId.PartB()));
+		}
+	}
+}
+
+void LogNodeEvent(const MCHAR* Name, INodeEventCallback::NodeKeyTab& nodes)
+{
+	// todo: heavy logging - called a lot
+	return;
+
+	LogDebug(FString::Printf(TEXT("NodeEventCallback:%s"), Name));
+	for (int NodeIndex = 0; NodeIndex < nodes.Count(); ++NodeIndex)
+	{
+		FNodeKey NodeKey = nodes[NodeIndex];
+
+		Animatable* anim = Animatable::GetAnimByHandle(NodeKey);
+		INode* Node = NodeEventNamespace::GetNodeByKey(NodeKey);
+		if (Node) // Node sometimes is null. Not sure why
+		{
+			LogDebug(FString::Printf(TEXT("   %u %s(%d)"), NodeKey, Node->GetName(), Node->GetHandle()));
+		}
+		else
+		{
+			LogDebug(FString::Printf(TEXT("    %u <null>"), NodeKey));
+		}
+	}
+}
+
+
+class FDatasmith3dsMaxScene
+{
+public:
+	TSharedPtr<IDatasmithScene> DatasmithSceneRef;
+	TSharedPtr<FDatasmithSceneExporter> SceneExporterRef;
+
+	FDatasmith3dsMaxScene() 
+	{
+		Reset();
+	}
+
+	void Reset()
+	{
+		DatasmithSceneRef.Reset();
+		DatasmithSceneRef = FDatasmithSceneFactory::CreateScene(TEXT(""));
+		SceneExporterRef.Reset();
+		SceneExporterRef = MakeShared<FDatasmithSceneExporter>();
+
+		// todo: compute or pass from script
+		DatasmithSceneRef->SetProductName(TEXT("3dsmax"));
+		DatasmithSceneRef->SetHost(TEXT("3dsmax"));
+
+		// Set the vendor name of the application used to build the scene.
+		DatasmithSceneRef->SetVendor(TEXT("Autodesk"));
+
+		FString Version = FString::FromInt(MAX_VERSION_MAJOR) + TEXT(".") + FString::FromInt(MAX_VERSION_MINOR) + TEXT(".") + FString::FromInt(MAX_VERSION_POINT);
+		DatasmithSceneRef->SetProductVersion(*Version);
+
+		// XXX: PreExport needs to be called before DirectLink instance is constructed - 
+		// Reason - it calls initialization of FTaskGraphInterface. Callstack:
+		// PreExport:
+		//  - FDatasmithExporterManager::Initialize 
+		//	-- DatasmithGameThread::InitializeInCurrentThread
+		//  --- GEngineLoop.PreInit
+		//  ---- PreInitPreStartupScreen
+		//  ----- FTaskGraphInterface::Startup
+		PreExport();
+	}
+
+
+
+	TSharedRef<IDatasmithScene> GetDatasmithScene()
+	{
+		return DatasmithSceneRef.ToSharedRef();
+	}
+
+	FDatasmithSceneExporter& GetSceneExporter()
+	{
+		return *SceneExporterRef;
+	}
+
+	void SetName(const TCHAR* InName)
+	{
+		SceneExporterRef->SetName(InName);
+		DatasmithSceneRef->SetName(InName);
+		DatasmithSceneRef->SetLabel(InName);
+	}
+
+	void SetOutputPath(const TCHAR* InOutputPath)
+	{
+		// Set the output folder where this scene will be exported.
+		SceneExporterRef->SetOutputPath(InOutputPath);
+		DatasmithSceneRef->SetResourcePath(SceneExporterRef->GetOutputPath());
+	}
+
+	void PreExport()
+	{
+		// Create a Datasmith scene exporter.
+		SceneExporterRef->Reset();
+
+		// Start measuring the time taken to export the scene.
+		SceneExporterRef->PreExport();
+	}
+};
+
+// Material changes more precise tracking can be done with ReferenceMaker
+// INodeEventCallback's MaterialOtherEvent tracks that any change is done to the material assigned to node
+// When a submaterial of a multimat is changed MaterialOtherEvent is called /omitting/ details which submaterial is modified
+// ReferenceMaker on the other hand tracks individual (sub)material changes
+// todo: 
+// - stop observing material when not needed(i.e. it's not assigned, used as submaterial or something else(?)
+// - remove when deleted
+class FMaterialObserver: public ReferenceMaker
+{
+	typedef int FMaterialIndex;
+
+public:
+	~FMaterialObserver()
+	{
+		DeleteAllRefs(); // Required to be called in destructor
+	}
+
+	void Reset()
+	{
+		IndexToReferencedMaterial.Reset();
+		ReferencedMaterialToIndex.Reset();
+	}
+
+	RefResult NotifyRefChanged(const Interval& ChangeInterval, RefTargetHandle TargetHandle, PartID& PartId, RefMessage Message, BOOL propagate) override
+	{
+		// todo: remove material handling???
+		ensure(ReferencedMaterialToIndex.Contains(TargetHandle));
+
+
+		LogDebug(FString::Printf(TEXT("NotifyRefChanged: %s: %x"), dynamic_cast<Mtl*>(TargetHandle)->GetName().data(), Message));
+
+		return REF_SUCCEED;
+	}
+
+	void AddMaterial(Mtl* Material)
+	{
+		if (!ReferencedMaterialToIndex.Contains(Material))
+		{
+			ReplaceReference(NumRefs(), Material);
+		}
+	}
+
+	// todo: unused
+	// RECONSIDER: when this method is used - removed material reduces NumRefs result so adding new material will overwrite already existing reference 
+	// e.g. was two materials added, with index 0 and 1, material 0 removed, NumRefs becomes 1 so next call ReplaceReference(NumRefs(), Material) will replace material 1 in the map
+	void RemoveMaterial(Mtl* Material)
+	{
+		FMaterialIndex MaterialIndex;
+		if (ReferencedMaterialToIndex.RemoveAndCopyValue(Material, MaterialIndex))
+		{
+			IndexToReferencedMaterial.Remove(MaterialIndex);
+		}
+	}
+
+	int NumRefs() override
+	{
+		return IndexToReferencedMaterial.Num();
+	}
+
+	RefTargetHandle GetReference(int ReferenceIndex) override
+	{
+		return IndexToReferencedMaterial[ReferenceIndex];
+	}
+
+	void SetReference(int ReferenceIndex, RefTargetHandle TargetHandle) override
+	{
+		IndexToReferencedMaterial.Add(ReferenceIndex, TargetHandle);
+		ReferencedMaterialToIndex.Add(TargetHandle, ReferenceIndex);
+	}
+private:
+	TMap<FMaterialIndex, RefTargetHandle> IndexToReferencedMaterial;
+	TMap<RefTargetHandle, FMaterialIndex> ReferencedMaterialToIndex;
+
+};
+
+
+class FNodeObserver : public ReferenceMaker
+{
+	typedef int FItemIndex;
+
+public:
+	~FNodeObserver()
+	{
+		DeleteAllRefs(); // Required to be called in destructor
+	}
+
+	void Reset()
+	{
+		IndexToReferencedItem.Reset();
+		ReferencedItemToIndex.Reset();
+	}
+
+	RefResult NotifyRefChanged(const Interval& ChangeInterval, RefTargetHandle TargetHandle, PartID& PartId, RefMessage Message, BOOL propagate) override
+	{
+		// todo: remove material handling???
+		ensure(ReferencedItemToIndex.Contains(TargetHandle));
+
+
+		// todo: heavy logging - called a lot
+		// LogDebug(FString::Printf(TEXT("FNodeObserver::NotifyRefChanged: %s: %x"), dynamic_cast<INode*>(TargetHandle)->GetName(), Message));
+
+		return REF_SUCCEED;
+	}
+
+	void AddItem(INode* Node)
+	{
+		if (!ReferencedItemToIndex.Contains(Node))
+		{
+			ReplaceReference(NumRefs(), Node);
+		}
+	}
+
+	// todo: unused
+	// RECONSIDER: when this method is used - removed material reduces NumRefs result so adding new material will overwrite already existing reference 
+	// e.g. was two materials added, with index 0 and 1, material 0 removed, NumRefs becomes 1 so next call ReplaceReference(NumRefs(), Material) will replace material 1 in the map
+	void RemoveItem(Mtl* Node)
+	{
+		FItemIndex NodeIndex;
+		if (ReferencedItemToIndex.RemoveAndCopyValue(Node, NodeIndex))
+		{
+			IndexToReferencedItem.Remove(NodeIndex);
+		}
+	}
+
+	int NumRefs() override
+	{
+		return IndexToReferencedItem.Num();
+	}
+
+	RefTargetHandle GetReference(int ReferenceIndex) override
+	{
+		RefTargetHandle TargetHandle = IndexToReferencedItem[ReferenceIndex];
+		// todo: heavy logging - called a lot
+		// LogDebug(FString::Printf(TEXT("FNodeObserver::GetReference: %d, %s"), ReferenceIndex, TargetHandle ? dynamic_cast<INode*>(TargetHandle)->GetName() : TEXT("<null>")));
+		return TargetHandle;
+	}
+
+	void SetReference(int ReferenceIndex, RefTargetHandle TargetHandle) override
+	{
+		// todo: heavy logging - called a lot
+		// LogDebug(FString::Printf(TEXT("FNodeObserver::SetReference: %d, %s"), ReferenceIndex, TargetHandle?dynamic_cast<INode*>(TargetHandle)->GetName():TEXT("<null>")));
+
+		// todo: investigate why NodeEventNamespace::GetNodeByKey may stil return NULL
+		// testcase - add XRef Material - this will immediately have this 
+		// even though NOTIFY_SCENE_ADDED_NODE was called for node and NOTIFY_SCENE_PRE_DELETED_NODE wasn't!
+		// BUT SeetReference with NULL handle is called
+		// also REFMSG_REF_DELETED and TARGETMSG_DELETING_NODE messages are sent to NotifyRefChanged
+
+		check(!ReferencedItemToIndex.Contains(TargetHandle)); // Not expecting to have same handle under two indices(back-indexing breaks)
+
+		if (TargetHandle)
+		{
+			ReferencedItemToIndex.Add(TargetHandle, ReferenceIndex);
+		}
+
+		if (RefTargetHandle* HandlePtr = IndexToReferencedItem.Find(ReferenceIndex))
+		{
+			if (*HandlePtr)
+			{
+				ReferencedItemToIndex.Remove(*HandlePtr);
+			}
+			*HandlePtr = TargetHandle;
+		}
+		else
+		{
+			IndexToReferencedItem.Add(ReferenceIndex, TargetHandle);
+		}
+	}
+private:
+	TMap<FItemIndex, RefTargetHandle> IndexToReferencedItem;
+	TMap<RefTargetHandle, FItemIndex> ReferencedItemToIndex;
+};
+
+
+class FNodeTracker
+{
+public:
+	explicit FNodeTracker(INode* InNode) : Node(InNode) {}
+
+	void Invalidate()
+	{
+		bInvalidated = true;
+	}
+
+	bool IsInvalidated() const
+	{
+		return bInvalidated;
+	}
+
+	INode* Node = nullptr;
+	bool bInvalidated = true;
+
+	TSharedPtr<IDatasmithActorElement> DatasmithActorElement;
+
+	class FMaterialTracker* MaterialTracker = nullptr;
+
+	TSharedPtr<IDatasmithMeshActorElement> DatasmithMeshActor;
+	TSharedPtr<IDatasmithMeshElement> DatasmithMeshElement; //todo: this should go away to be separate from node(especially when instancing)
+};
+
+
+class FNodeTrackerHandle
+{
+public:
+	explicit FNodeTrackerHandle(INode* InNode) : Impl(MakeShared<FNodeTracker>(InNode)) {}
+
+	FNodeTracker* GetNodeTracker() const
+	{
+		return Impl.Get();
+	}
+
+private:
+	TSharedPtr<FNodeTracker> Impl;
+};
+
+class FNullView : public View
+{
+public:
+	FNullView()
+	{
+		worldToView.IdentityMatrix(); screenW = 640.0f; screenH = 480.0f;
+	}
+
+	virtual Point2 ViewToScreen(Point3 p) override
+	{
+		return Point2(p.x, p.y);
+	}
+};
+
+// todo: these converters from baseline plugin. Might extract and reuse in both places(here and FDatasmithMaxMeshExporter)?
+class FDatasmithConverter
+{
+	float UnitToCentimeter;
+public:
+
+	FDatasmithConverter()
+	{
+		UnitToCentimeter = FMath::Abs(GetSystemUnitScale(UNITS_CENTIMETERS));
+	}
+
+	FVector toDatasmithVector(Point3 Point) const
+	{
+		return FVector(UnitToCentimeter * Point.x,
+			UnitToCentimeter * -Point.y,
+			UnitToCentimeter * Point.z);
+	}
+
+	FColor toDatasmithColor(Point3& Point)
+	{
+		//The 3ds Max vertex colors are encoded from 0 to 1 in float
+		return FColor(Point.x * MAX_uint8, Point.y * MAX_uint8, Point.z * MAX_uint8);
+	}
+
+	void MaxToUnrealCoordinates(Matrix3 Matrix, FVector& Translation, FQuat& Rotation, FVector& Scale)
+	{
+		Point3 Pos = Matrix.GetTrans();
+		Translation.X = Pos.x * UnitToCentimeter;
+		Translation.Y = -Pos.y * UnitToCentimeter;
+		Translation.Z = Pos.z * UnitToCentimeter;
+
+		// Clear the transform on the matrix
+		Matrix.NoTrans();
+
+		// We're only doing Scale - save out the
+		// rotation so we can put it back
+		AffineParts Parts;
+		decomp_affine(Matrix, &Parts);
+		ScaleValue ScaleVal = ScaleValue(Parts.k * Parts.f, Parts.u);
+		Scale = FVector(ScaleVal.s.x, ScaleVal.s.y, ScaleVal.s.z);
+
+		Rotation = FQuat(Parts.q.x, -Parts.q.y, Parts.q.z, Parts.q.w);
+	}
+};
+
+
+// todo: copied from baseline plugin(it has dependencies on converters that are not static in FDatasmithMaxMeshExporter)
+void FillDatasmithMeshFromMaxMesh(FDatasmithMesh& DatasmithMesh, Mesh& MaxMesh, INode* ExportedNode, bool bForceSingleMat, TSet<uint16>& SupportedChannels, const TCHAR* MeshName, FTransform Pivot)
+{
+	FDatasmithConverter Converter;
+
+	const int NumFaces = MaxMesh.getNumFaces();
+	const int NumVerts = MaxMesh.getNumVerts();
+
+	DatasmithMesh.SetVerticesCount(NumVerts);
+	DatasmithMesh.SetFacesCount(NumFaces);
+
+	// Vertices
+	for (int i = 0; i < NumVerts; i++)
+	{
+		Point3 Point = MaxMesh.getVert(i);
+
+		FVector Vertex = Converter.toDatasmithVector(Point);
+		Vertex = Pivot.TransformPosition(Vertex); // Bake object-offset in the mesh data when possible
+
+		DatasmithMesh.SetVertex(i, Vertex.X, Vertex.Y, Vertex.Z);
+	}
+
+	// Vertex Colors
+	if (MaxMesh.curVCChan == 0 && MaxMesh.numCVerts > 0)
+	{
+		// Default vertex color channel
+		for (int32 i = 0; i < NumFaces; i++)
+		{
+			TVFace& Face = MaxMesh.vcFace[i];
+			DatasmithMesh.SetVertexColor(i * 3, Converter.toDatasmithColor(MaxMesh.vertCol[Face.t[0]]));
+			DatasmithMesh.SetVertexColor(i * 3 + 1, Converter.toDatasmithColor(MaxMesh.vertCol[Face.t[1]]));
+			DatasmithMesh.SetVertexColor(i * 3 + 2, Converter.toDatasmithColor(MaxMesh.vertCol[Face.t[2]]));
+		}
+	}
+
+	// UVs
+	TMap<int32, int32> UVChannelsMap;
+	TMap<uint32, int32> HashToChannel;
+	bool bIsFirstUVChannelValid = true;
+
+	for (int32 i = 1; i <= MaxMesh.getNumMaps(); ++i)
+	{
+		if (MaxMesh.mapSupport(i) == BOOL(true) && MaxMesh.getNumMapVerts(i) > 0)
+		{
+			DatasmithMesh.AddUVChannel();
+			const int32 UVChannelIndex = DatasmithMesh.GetUVChannelsCount() - 1;
+			const int32 UVsCount = MaxMesh.getNumMapVerts(i);
+
+			DatasmithMesh.SetUVCount(UVChannelIndex, UVsCount);
+
+			UVVert* Vertex = MaxMesh.mapVerts(i);
+
+			for (int32 j = 0; j < UVsCount; ++j)
+			{
+				const UVVert& MaxUV = Vertex[j];
+				DatasmithMesh.SetUV(UVChannelIndex, j, MaxUV.x, 1.f - MaxUV.y);
+			}
+
+			TVFace* Faces = MaxMesh.mapFaces(i);
+			for (int32 j = 0; j < MaxMesh.getNumFaces(); ++j)
+			{
+				DatasmithMesh.SetFaceUV(j, UVChannelIndex, Faces[j].t[0], Faces[j].t[1], Faces[j].t[2]);
+			}
+
+			if (UVChannelIndex == 0)
+			{
+				//Verifying that the UVs are properly unfolded, which is required to calculate the tangent in unreal.
+				bIsFirstUVChannelValid = FDatasmithMeshUtils::IsUVChannelValid(DatasmithMesh, UVChannelIndex);
+			}
+
+			uint32 Hash = DatasmithMesh.GetHashForUVChannel(UVChannelIndex);
+			int32* PointerToChannel = HashToChannel.Find(Hash);
+
+			if (PointerToChannel)
+			{
+				// Remove the channel because there is another one that is identical
+				DatasmithMesh.RemoveUVChannel();
+
+				// Map the user-specified UV Channel (in 3dsmax) to the actual UV channel that will be exported to Unreal
+				UVChannelsMap.Add(i - 1, *PointerToChannel);
+			}
+			else
+			{
+				// Map the user-specified UV Channel (in 3dsmax) to the actual UV channel that will be exported to Unreal
+				UVChannelsMap.Add(i - 1, UVChannelIndex);
+				HashToChannel.Add(Hash, UVChannelIndex);
+			}
+		}
+	}
+
+	if (!bIsFirstUVChannelValid)
+	{
+		//DatasmithMaxLogger::Get().AddGeneralError(*FString::Printf(TEXT("%s's UV channel #1 contains degenerated triangles, this can cause issues in Unreal. It is recommended to properly unfold and flatten exported UV data.")
+		//	, static_cast<const TCHAR*>(ExportedNode->GetName())));
+	}
+
+	if (MeshName != nullptr)
+	{
+		//MeshNamesToUVChannels.Add(MeshName, MoveTemp(UVChannelsMap));
+	}
+
+	// Faces
+	for (int i = 0; i < NumFaces; i++)
+	{
+		// Create polygons. Assign texture and texture UV indices.
+		// all faces of the cube have the same texture
+
+		Face& MaxFace = MaxMesh.faces[i];
+		int MaterialId = bForceSingleMat ? 0 : MaxFace.getMatID();
+
+		SupportedChannels.Add(MaterialId);
+
+		//Max's channel UI is not zero-based, so we register an incremented ChannelID for better visual consistency after importing in Unreal.
+		DatasmithMesh.SetFace(i, MaxFace.getVert(0), MaxFace.getVert(1), MaxFace.getVert(2), MaterialId + 1);
+		DatasmithMesh.SetFaceSmoothingMask(i, (uint32)MaxFace.getSmGroup());
+	}
+
+	//Normals
+
+	MaxMesh.SpecifyNormals();
+	MeshNormalSpec* Normal = MaxMesh.GetSpecifiedNormals();
+	Normal->MakeNormalsExplicit(false);
+	Normal->CheckNormals();
+
+	Matrix3 RotationMatrix;
+	RotationMatrix.IdentityMatrix();
+	Quat ObjectOffsetRotation = ExportedNode->GetObjOffsetRot();
+	RotateMatrix(RotationMatrix, ObjectOffsetRotation);
+
+	Point3 Point;
+
+	for (int i = 0; i < NumFaces; i++)
+	{
+		Point = Normal->GetNormal(i, 0).Normalize() * RotationMatrix;
+		FVector NormalVector = Converter.toDatasmithVector(Point);
+		DatasmithMesh.SetNormal(i * 3, NormalVector.X, NormalVector.Y, NormalVector.Z);
+
+		Point = Normal->GetNormal(i, 1).Normalize() * RotationMatrix;
+		NormalVector = Converter.toDatasmithVector(Point);
+		DatasmithMesh.SetNormal(i * 3 + 1, NormalVector.X, NormalVector.Y, NormalVector.Z);
+
+		Point = Normal->GetNormal(i, 2).Normalize() * RotationMatrix;
+		NormalVector = Converter.toDatasmithVector(Point);
+		DatasmithMesh.SetNormal(i * 3 + 2, NormalVector.X, NormalVector.Y, NormalVector.Z);
+	}
+}
+
+
+class FMaterialTracker
+{
+public:
+	explicit FMaterialTracker(Mtl* InMaterial) : Material(InMaterial) {}
+
+
+	Mtl* Material;
+
+	TArray<Mtl*> Materials; // Actual materials used for this assigned material
+	TArray<Texmap*> Textures;
+
+	bool bInvalidated = true;
+
+	TArray<Mtl*>& GetActualMaterials()
+	{
+		return Materials;
+	}
+
+	void ResetActualMaterialAndTextures()
+	{
+		Materials.Reset();
+		Textures.Reset(); //todo: unregister textures
+	}
+
+	void AddActualMaterial(Mtl* ActualMaterial)
+	{
+		Materials.Add(ActualMaterial);
+	}
+
+	void AddActualTexture(Texmap* Texture)
+	{
+		Textures.Add(Texture);
+	}
+};
+
+class FMaterialTrackerHandle
+{
+public:
+
+	explicit FMaterialTrackerHandle(Mtl* InMaterial) : Impl(MakeShared<FMaterialTracker>(InMaterial)) {}
+
+	FMaterialTracker* GetMaterialTracker()
+	{
+		return Impl.Get();
+	}
+
+private:
+	TSharedPtr<FMaterialTracker> Impl; // todo: reuse material tracker objects(e.g. make a pool)
+
+	FMaterialTracker& GetImpl()
+	{
+		return *Impl;
+	}
+};
+
+class FMaterialsTracker
+{
+public:
+	TSet<Mtl*> EncounteredMaterials;
+	TSet<Texmap*> EncounteredTextures;
+
+	TArray<FString> MaterialNames;
+
+	TMap<Mtl*, TSet<FMaterialTracker*>> UsedMaterialToMaterialTracker; // Materials uses by nodes keep set of assigned materials they are used for
+	TMap<Mtl*, TSharedPtr<IDatasmithBaseMaterialElement>> UsedMaterialToDatasmithMaterial;
+
+	FDatasmith3dsMaxScene& ExportedScene;
+
+	FMaterialsTracker(FDatasmith3dsMaxScene& InExportedScene) : ExportedScene(InExportedScene) {}
+
+	void Reset()
+	{
+		EncounteredMaterials.Reset();
+		EncounteredTextures.Reset();
+		MaterialNames.Reset();
+
+		UsedMaterialToMaterialTracker.Reset();
+
+		UsedMaterialToDatasmithMaterial.Reset();
+	}
+
+	void SetDatasmithMaterial(Mtl* ActualMaterial, TSharedPtr<IDatasmithBaseMaterialElement> DatastmihMaterial)
+	{
+		UsedMaterialToDatasmithMaterial.Add(ActualMaterial, DatastmihMaterial);
+	}
+
+	void RegisterMaterialTracker(FMaterialTracker& MaterialTracker)
+	{
+		for (Mtl* Material : MaterialTracker.GetActualMaterials())
+		{
+			UsedMaterialToMaterialTracker.FindOrAdd(Material).Add(&MaterialTracker);
+		}
+
+		//todo: register textures
+	}
+
+	void UnregisterMaterialTracker(FMaterialTracker& MaterialTracker)
+	{
+		for (Mtl* Material: MaterialTracker.GetActualMaterials())
+		{
+			TSet<FMaterialTracker*>& MaterialTrackersForMaterial = UsedMaterialToMaterialTracker[Material];
+			MaterialTrackersForMaterial.Remove(&MaterialTracker);
+			if (!MaterialTrackersForMaterial.Num())
+			{
+				UsedMaterialToMaterialTracker.Remove(Material);
+
+				TSharedPtr<IDatasmithBaseMaterialElement> DatasmithMaterial;
+				if(UsedMaterialToDatasmithMaterial.RemoveAndCopyValue(Material, DatasmithMaterial))
+				{
+					ExportedScene.DatasmithSceneRef->RemoveMaterial(DatasmithMaterial);
+				}
+			}
+		}
+
+		MaterialTracker.ResetActualMaterialAndTextures();
+	}
+};
+
+// Copied from
+// FDatasmithMaxSceneParser::MaterialEnum
+// FDatasmithMaxSceneParser::TexEnum
+// Collects actual materials that are used by the top-level material(assigned to node)
+class FMaterialEnum
+{
+public:
+	FMaterialsTracker& MaterialsTracker;
+	FMaterialTracker& MaterialTracker;
+
+	FMaterialEnum(FMaterialsTracker& InMaterialsTracker, FMaterialTracker& InMaterialTracker): MaterialsTracker(InMaterialsTracker), MaterialTracker(InMaterialTracker) {}
+
+	void MaterialEnum(Mtl* Material, bool bAddMaterial)
+	{
+		if (Material == NULL)
+		{
+			return;
+		}
+
+		if (FDatasmithMaxMatHelper::GetMaterialClass(Material) == EDSMaterialType::XRefMat)
+		{
+			MaterialEnum(FDatasmithMaxMatHelper::GetRenderedXRefMaterial(Material), true);
+		}
+		else if (FDatasmithMaxMatHelper::GetMaterialClass(Material) == EDSMaterialType::MultiMat)
+		{
+			for (int i = 0; i < Material->NumSubMtls(); i++)
+			{
+				MaterialEnum(Material->GetSubMtl(i), true);
+			}
+		}
+		else
+		{
+			if (bAddMaterial)
+			{
+				if (!MaterialsTracker.EncounteredMaterials.Contains(Material))
+				{
+					int DuplicateCount = 0;
+					FString ProposedName = Material->GetName().data();
+					// todo: fix this without changing max material name? Btw - this requires changing all material export functions for all types of materials(those functions tied to Mtl->GetName())
+					// todo: revert material names after export
+					MaterialsTracker.MaterialNames.Add(*ProposedName);
+
+					// Make unique material name
+					FDatasmithUtils::SanitizeNameInplace(ProposedName);
+					for (Mtl* OtherMaterial: MaterialsTracker.EncounteredMaterials)
+					{
+						if (ProposedName == FDatasmithUtils::SanitizeName(OtherMaterial->GetName().data()))
+						{
+							DuplicateCount++;
+							ProposedName = FDatasmithUtils::SanitizeName(Material->GetName().data()) + TEXT("_(") + FString::FromInt(DuplicateCount) + TEXT(")");
+						}
+					}
+					Material->SetName(*ProposedName);
+					MaterialsTracker.EncounteredMaterials.Add(Material);
+				}
+				MaterialTracker.AddActualMaterial(Material);
+			}
+
+			bool bAddRecursively = Material->ClassID() == THEARANDOMCLASS || Material->ClassID() == VRAYBLENDMATCLASS || Material->ClassID() == CORONALAYERMATCLASS;
+			for (int i = 0; i < Material->NumSubMtls(); i++)
+			{
+				MaterialEnum(Material->GetSubMtl(i), bAddRecursively);
+			}
+
+			for (int i = 0; i < Material->NumSubTexmaps(); i++)
+			{
+				Texmap* SubTexture = Material->GetSubTexmap(i);
+				if (SubTexture != NULL)
+				{
+					TexEnum(SubTexture);
+				}
+			}
+		}
+	}
+
+	void TexEnum(Texmap* Texture)
+	{
+		if (Texture == NULL)
+		{
+			return;
+		}
+
+		if (!MaterialsTracker.EncounteredTextures.Contains(Texture))
+		{
+			MaterialsTracker.EncounteredTextures.Add(Texture);
+		}
+
+		for (int i = 0; i < Texture->NumSubTexmaps(); i++)
+		{
+			Texmap* SubTexture = Texture->GetSubTexmap(i);
+			if (SubTexture != NULL)
+			{
+				TexEnum(SubTexture);
+			}
+		}
+		MaterialTracker.AddActualTexture(Texture);
+	}
+};
+
+
+// Holds states of entities for syncronization and handles change events
+class FSceneTracker
+{
+public:
+	FSceneTracker(FDatasmith3dsMaxScene& InExportedScene) : ExportedScene(InExportedScene), MaterialsTracker(InExportedScene) {}
+
+	bool ParseScene()
+	{
+		INode* Node = GetCOREInterface()->GetRootNode();
+		return ParseScene(Node, nullptr);
+	}
+
+	bool ParseScene(INode* SceneRootNode, IDatasmithActorElement* ParentElement)
+	{
+		// todo: do we need Root Datasmith node of scene/XRefScene in the hierarchy?
+		// is there anything we need to handle for main file root node?
+		// for XRefScene? Maybe addition/removal? Do we need one node to consolidate XRefScene under?
+
+		// nodes comming from XRef Scenes/Objects could be null
+		if (SceneRootNode == NULL)
+		{
+			return false;
+		}
+
+		// Parse XRefScenes
+		for (int XRefChild = 0; XRefChild < SceneRootNode->GetXRefFileCount(); ++XRefChild)
+		{
+			DWORD XRefFlags = SceneRootNode->GetXRefFlags(XRefChild);
+
+			// XRef is disabled - not shown in viewport/render. Not loaded.
+			if (XRefFlags & XREF_DISABLED)
+			{
+				// todo: baseline - doesn't check this - exports even disabled and XREF_HIDDEN scenes
+				continue;
+			}
+
+			FString Path = FDatasmithMaxSceneExporter::GetActualPath(SceneRootNode->GetXRefFile(XRefChild).GetFileName());
+			if (FPaths::FileExists(Path) == false)
+			{
+				FString Error = FString("XRefScene file \"") + FPaths::GetCleanFilename(*Path) + FString("\" cannot be found");
+				// todo: logging
+				// DatasmithMaxLogger::Get().AddMissingAssetError(*Error);
+			}
+			else
+			{
+				ParseScene(SceneRootNode->GetXRefTree(XRefChild), ParentElement);
+			}
+		}
+
+		int32 ChildNum = SceneRootNode->NumberOfChildren();
+		for (int32 ChildIndex = 0; ChildIndex < ChildNum; ++ChildIndex)
+		{
+			ParseNode(SceneRootNode->GetChildNode(ChildIndex));
+		}
+		return true;
+	}
+
+	void ParseNode(INode* Node)
+	{
+		// todo: 
+		// Node->IsNodeHidden(TRUE)
+		// Node->GetXRefFileCount()
+
+		BOOL bIsNodeHidden = Node->IsNodeHidden(TRUE);
+
+		if (bIsNodeHidden == false)
+		{
+			//todo: when a refernced file is not found XRefObj is not resolved and it's kept as XREFOBJ_CLASS_ID
+			// instead of resolved class that it references
+//			if (ObjState.obj->ClassID() == XREFOBJ_CLASS_ID)
+//			{
+//				//max2017 and newer versions allow developers to get the file, with previous versions you only can retrieve the active one
+//#ifdef MAX_RELEASE_R19
+//				FString Path = FDatasmithMaxSceneExporter::GetActualPath(((IXRefObject8*)ObjState.obj)->GetFile(FALSE).GetFileName());
+//#else
+//				FString Path = FDatasmithMaxSceneExporter::GetActualPath(((IXRefObject8*)ObjState.obj)->GetActiveFile().GetFileName());
+//#endif
+//				if (FPaths::FileExists(Path) == false)
+//				{
+//					FString Error = FString("XRefObj file \"") + FPaths::GetCleanFilename(*Path) + FString("\" cannot be found");
+//					DatasmithMaxLogger::Get().AddMissingAssetError(*Error);
+//					bIsNodeHidden = true;
+//				}
+//			}
+		}
+
+
+		FNodeKey NodeKey = NodeEventNamespace::GetKeyByNode(Node);
+
+		FNodeTrackerHandle& NodeTracker = AddNode(NodeKey, Node);
+
+		// Parse children
+		int32 ChildNum = Node->NumberOfChildren();
+		for (int32 ChildIndex = 0; ChildIndex < ChildNum; ++ChildIndex)
+		{
+			ParseNode(Node->GetChildNode(ChildIndex));
+		}
+	}
+
+
+	void Reset()
+	{
+		NodeObserver.Reset();
+		MaterialObserver.Reset();
+		NodeTrackers.Reset();
+		InvalidatedNodeTrackers.Reset();
+		MaterialTrackers.Reset();
+		InvalidatedMaterialTrackers.Reset();
+		MaterialsTracker.Reset();
+	}
+
+	// Applies all recorded changes to Datasmith scene
+	void Update()
+	{
+		DatasmithMaxLogger::Get().Purge();
+		for (FNodeTracker* NodeTracker : InvalidatedNodeTrackers)
+		{
+			ConvertNode(*NodeTracker);
+		}
+		InvalidatedNodeTrackers.Reset();
+
+		TSet<Mtl*> ActualMaterialToUpdate;
+		TSet<Texmap*> ActualTexmapsToUpdate;
+		for (FMaterialTracker* MaterialTracker : InvalidatedMaterialTrackers)
+		{
+			MaterialsTracker.UnregisterMaterialTracker(*MaterialTracker);
+			FMaterialEnum(MaterialsTracker, *MaterialTracker).MaterialEnum(MaterialTracker->Material, true);
+			MaterialsTracker.RegisterMaterialTracker(*MaterialTracker);
+
+			for (Mtl* ActualMaterial : MaterialTracker->GetActualMaterials())
+			{
+				ActualMaterialToUpdate.Add(ActualMaterial);
+			}
+			MaterialTracker->bInvalidated = false;
+			for (Texmap* Texture: MaterialTracker->Textures)
+			{
+				ActualTexmapsToUpdate.Add(Texture);
+			}
+		}
+		InvalidatedMaterialTrackers.Reset();
+
+		for (Texmap* Texture : ActualTexmapsToUpdate)
+		{
+			FDatasmithMaxMatExport::GetXMLTexture(ExportedScene.GetDatasmithScene(), Texture, ExportedScene.GetSceneExporter().GetAssetsOutputPath());
+		}
+
+		for (Mtl* ActualMaterial : ActualMaterialToUpdate)
+		{
+			// todo: make sure not reexport submaterial more than once - i.e. when a submaterial is used in two composite material
+			FDatasmithMaxMatExport::bForceReexport = true;
+			TSharedPtr<IDatasmithBaseMaterialElement> DatastmihMaterial = FDatasmithMaxMatExport::ExportUniqueMaterial(ExportedScene.GetDatasmithScene(), ActualMaterial, ExportedScene.GetSceneExporter().GetAssetsOutputPath());
+
+			MaterialsTracker.SetDatasmithMaterial(ActualMaterial, DatastmihMaterial);
+		}
+
+		// todo: removes textures that were added again(materials were updated). Need to fix this by identifying exactly which textures are being updated and removing ahead
+		//TMap<FString, TSharedPtr<IDatasmithTextureElement>> TexturesAdded;
+		//TArray<TSharedPtr<IDatasmithTextureElement>> TexturesToRemove;
+		//for(int32 TextureIndex = 0; TextureIndex < ExportedScene.GetDatasmithScene()->GetTexturesCount(); ++TextureIndex )
+		//{
+		//	TSharedPtr<IDatasmithTextureElement> TextureElement = ExportedScene.GetDatasmithScene()->GetTexture(TextureIndex);
+		//	FString Name = TextureElement->GetName();
+		//	if (TexturesAdded.Contains(Name))
+		//	{
+		//		TexturesToRemove.Add(TexturesAdded[Name]);
+		//		TexturesAdded[Name] = TextureElement;
+		//	}
+		//	else
+		//	{
+		//		TexturesAdded.Add(Name, TextureElement);
+		//	}
+		//}
+
+		//for(TSharedPtr<IDatasmithTextureElement> Texture: TexturesToRemove)
+		//{
+		//	ExportedScene.GetDatasmithScene()->RemoveTexture(Texture);
+		//}
+
+	}
+
+	FORCENOINLINE
+	FNodeTrackerHandle& AddNode(FNodeKey NodeKey, INode* Node)
+	{
+		FNodeTrackerHandle& NodeTracker = NodeTrackers.Emplace(NodeKey, Node);
+		InvalidatedNodeTrackers.Add(NodeTracker.GetNodeTracker()); // Add
+		return NodeTracker;
+	}
+
+	// todo: make fine invalidates - full only something like geometry change, but finer for transform, name change and more
+	void InvalidateNode(FNodeKey NodeKey)
+	{
+		if (FNodeTrackerHandle* NodeTrackerHandle = NodeTrackers.Find(NodeKey))
+		{
+			FNodeTracker* NodeTracker = NodeTrackerHandle->GetNodeTracker();
+			NodeTracker->Invalidate();
+			InvalidatedNodeTrackers.Add(NodeTracker);
+		}
+	}
+
+	bool IsNodeInvalidated(const FNodeTrackerHandle& NodeTracker)
+	{
+		return NodeTracker.GetNodeTracker()->IsInvalidated();
+	}
+
+	void ConvertNode(FNodeTracker& NodeTracker)
+	{
+		INode* Node = NodeTracker.Node;
+
+		// Initialize actor(reset hierarchy if it was already created before) and set label
+		if (NodeTracker.DatasmithActorElement)
+		{
+			if (TSharedPtr<IDatasmithActorElement> ParentActor = NodeTracker.DatasmithActorElement->GetParentActor())
+			{
+				ParentActor->RemoveChild(NodeTracker.DatasmithActorElement);
+			}
+			else
+			{
+				ExportedScene.DatasmithSceneRef->RemoveActor(NodeTracker.DatasmithActorElement, EDatasmithActorRemovalRule::RemoveChildren);
+			}
+		}
+		else
+		{
+			// note: this is how baseline exporter derives names
+			FString UniqueName = FString::FromInt(Node->GetHandle());
+			NodeTracker.DatasmithActorElement = FDatasmithSceneFactory::CreateActor((const TCHAR*)*UniqueName);
+		}
+		NodeTracker.DatasmithActorElement->SetLabel(Node->GetName());
+
+		// Add to parent
+		FNodeKey ParentNodeKey = NodeEventNamespace::GetKeyByNode(Node->GetParentNode());
+		if (FNodeTrackerHandle* ParentNodeTrackerHandle = NodeTrackers.Find(ParentNodeKey))
+		{
+			// Add to parent Datasmith actor if it has been updated already, if not parent will add it
+			if (!IsNodeInvalidated(*ParentNodeTrackerHandle))
+			{
+				FNodeTracker* ParentNodeTracker = ParentNodeTrackerHandle->GetNodeTracker();
+				ParentNodeTracker->DatasmithActorElement->AddChild(NodeTracker.DatasmithActorElement);
+			}
+		}
+		else
+		{
+			// If there's no parent node registered assume it's at root
+			ExportedScene.GetDatasmithScene()->AddActor(NodeTracker.DatasmithActorElement);
+		}
+
+		// Add attach datasmith actors of child nodes
+		int32 ChildNum = Node->NumberOfChildren();
+		for (int32 ChildIndex = 0; ChildIndex < ChildNum; ++ChildIndex)
+		{
+
+			FNodeKey ChildNodeKey = NodeEventNamespace::GetKeyByNode(Node->GetChildNode(ChildIndex));
+			if (FNodeTrackerHandle* ChildNodeTrackerHandle = NodeTrackers.Find(ChildNodeKey))
+			{
+				// Add child datasmith actor if child is updated, if not child will add itself(it will be updated further in the queue)
+				if (!IsNodeInvalidated(*ChildNodeTrackerHandle))
+				{
+					NodeTracker.DatasmithActorElement->AddChild(ChildNodeTrackerHandle->GetNodeTracker()->DatasmithActorElement);
+				}
+			}
+		}
+
+		ConvertNodeTransform(NodeTracker);
+		ConvertNodeGeometry(NodeTracker);
+
+		// Mark node as updated as soon as it is - in order for next nodes to be able to use its DatasmithActor
+		NodeTracker.bInvalidated = false;
+	}
+
+	void ConvertNodeGeometry(FNodeTracker& NodeTracker)
+	{
+		if (NodeTracker.DatasmithMeshElement)
+		{
+			ExportedScene.DatasmithSceneRef->RemoveMesh(NodeTracker.DatasmithMeshElement);
+			// todo: mesh is removed from the scene but not deallocated for reuse
+			// OR it will stay taking up memory if not reused(e.g. node has no valid geometry now)
+		}
+
+
+		ObjectState ObjState = NodeTracker.Node->EvalWorldState(0);
+		Object* Obj = ObjState.obj;
+
+		if (!Obj)
+		{
+			return;
+		}
+
+		switch (Obj->SuperClassID())
+		{
+		case GEOMOBJECT_CLASS_ID:
+		{
+			// todo: reuse meshe element(make sure to reset all)
+			TSharedPtr<IDatasmithMeshElement> DatasmithMeshElement;
+			ConvertGeomObjToDatasmithMesh(NodeTracker, Obj, DatasmithMeshElement);
+			AssignDatasmithMeshToNodeTracker(NodeTracker, DatasmithMeshElement);
+			break;
+		}
+		// todo: other object types besides geometry
+		default:;
+		}
+	}
+
+	void AssignDatasmithMeshToNodeTracker(FNodeTracker& NodeTracker, const TSharedPtr<IDatasmithMeshElement>& DatasmithMeshElement)
+	{
+		NodeTracker.DatasmithMeshElement = DatasmithMeshElement;
+
+		if (DatasmithMeshElement)
+		{
+			if (!NodeTracker.DatasmithMeshActor)
+			{
+				FString MeshActorName = FString::FromInt(NodeTracker.Node->GetHandle() ) + TEXT("_Mesh");
+				FString MeshActorLabel = FString((const TCHAR*)NodeTracker.Node->GetName());
+				TSharedPtr<IDatasmithMeshActorElement> DatasmithMeshActor = FDatasmithSceneFactory::CreateMeshActor(*MeshActorName);
+				DatasmithMeshActor->SetLabel(*MeshActorLabel);
+
+				NodeTracker.DatasmithActorElement->AddChild(DatasmithMeshActor, EDatasmithActorAttachmentRule::KeepRelativeTransform);
+				NodeTracker.DatasmithMeshActor = DatasmithMeshActor;
+			}
+
+			NodeTracker.DatasmithMeshActor->SetStaticMeshPathName(DatasmithMeshElement->GetName());
+		}
+		else
+		{
+			if (NodeTracker.DatasmithMeshActor)
+			{
+				if (TSharedPtr<IDatasmithActorElement> DatasmithParentActorElement = NodeTracker.DatasmithMeshActor->GetParentActor())
+				{
+					DatasmithParentActorElement->RemoveChild(NodeTracker.DatasmithMeshActor);
+				}
+			}
+		}
+	}
+
+	// todo: change this function to create mesh element and schedule ExportToUObject threading
+	bool ConvertGeomObjToDatasmithMesh(FNodeTracker& NodeTracker, Object* Obj, TSharedPtr<IDatasmithMeshElement>& DatasmithMeshElement)
+	{
+		bool bResult = false;
+
+		INode* Node = NodeTracker.Node;
+
+		// todo: baseline exporter uses GetBaseObject which takes result of EvalWorldState
+		// and searched down DerivedObject pipeline(by taking GetObjRef) 
+		// This is STRANGE as EvalWorldState shouldn't return DerivedObject in the first place(it should return result of pipeline evaluation)
+
+		GeomObject* GeomObj = dynamic_cast<GeomObject*>(Obj);
+
+		FNullView View;
+		BOOL bNeedsDelete;
+		TimeValue Time = GetCOREInterface()->GetTime();
+		Mesh* RenderMesh = GeomObj->GetRenderMesh(Time, Node, View, bNeedsDelete);
+
+		if (!RenderMesh)
+		{
+			return bResult;
+		}
+
+		if (RenderMesh->getNumFaces())
+		{
+			// Copy mesh to clean it before filling Datasmith mesh from it
+			Mesh CachedMesh;
+			CachedMesh.DeepCopy(RenderMesh, TOPO_CHANNEL | GEOM_CHANNEL | TEXMAP_CHANNEL | VERTCOLOR_CHANNEL);
+
+			CachedMesh.DeleteIsoVerts();
+			CachedMesh.RemoveDegenerateFaces();
+			CachedMesh.RemoveIllegalFaces();
+
+			// Need to invalidate/rebuild strips/edges after topology change(removing bad verts/faces)
+			CachedMesh.InvalidateStrips();
+			CachedMesh.BuildStripsAndEdges();
+
+			if (CachedMesh.getNumFaces() > 0)
+			{
+				FDatasmithMesh DatasmithMesh;
+				TSet<uint16> SupportedChannels;
+				const TCHAR* MeshName = (const TCHAR*)Node->GetName();
+				// todo: pivot
+				FillDatasmithMeshFromMaxMesh(DatasmithMesh, CachedMesh, Node, false, SupportedChannels, MeshName, FTransform::Identity);
+
+				FDatasmithMeshExporter DatasmithMeshExporter;
+
+				DatasmithMeshElement = FDatasmithSceneFactory::CreateMesh(TEXT(""));
+				DatasmithMeshElement->SetName(MeshName);
+				DatasmithMeshElement->SetLabel(MeshName);
+
+				ExportedScene.GetDatasmithScene()->AddMesh(DatasmithMeshElement);
+
+				// todo: remove material if null
+				if (Mtl* Material = Node->GetMtl())
+				{
+					// todo: move to header
+					void AssignMeshMaterials(TSharedPtr<IDatasmithMeshElement>&MeshElement, Mtl * Material, const TSet<uint16>&SupportedChannels);
+
+					AssignMeshMaterials(DatasmithMeshElement, Material, SupportedChannels);
+
+					if (!NodeTracker.MaterialTracker || NodeTracker.MaterialTracker->Material != Material)
+					{
+						// Release old material
+						if (NodeTracker.MaterialTracker)
+						{
+							// Release material assignment
+							MaterialsAssignedToNodes[NodeTracker.MaterialTracker].Remove(&NodeTracker);
+
+							// Clean tracker if it's not used aby any node
+							if (!MaterialsAssignedToNodes[NodeTracker.MaterialTracker].Num())
+							{
+								MaterialsTracker.UnregisterMaterialTracker(*NodeTracker.MaterialTracker);
+								MaterialTrackers.Remove(NodeTracker.MaterialTracker->Material);
+							}
+						}
+
+						if (!MaterialTrackers.Contains(Material))
+						{
+							// Track material if not yet
+							FMaterialTrackerHandle& MaterialTrackerHandle = MaterialTrackers.Emplace(Material, Material);
+							InvalidatedMaterialTrackers.Add(MaterialTrackerHandle.GetMaterialTracker());
+						}
+
+						// Store new 
+						NodeTracker.MaterialTracker = MaterialTrackers.Find(Material)->GetMaterialTracker();
+						MaterialsAssignedToNodes.FindOrAdd(NodeTracker.MaterialTracker).Add(&NodeTracker);
+					}
+				}
+				else
+				{
+					// Release old material
+					if (NodeTracker.MaterialTracker)
+					{
+						MaterialsAssignedToNodes[NodeTracker.MaterialTracker].Remove(&NodeTracker);
+						if (!MaterialsAssignedToNodes[NodeTracker.MaterialTracker].Num())
+						{
+							MaterialsTracker.UnregisterMaterialTracker(*NodeTracker.MaterialTracker);
+							MaterialTrackers.Remove(Material);
+						}
+					}
+					NodeTracker.MaterialTracker = nullptr;
+				}
+
+				bResult = true; // Set to true, don't care what ExportToUObject does here - we need to move it to a thread anyway
+
+				// todo: parallelize this
+				if (DatasmithMeshExporter.ExportToUObject(DatasmithMeshElement, ExportedScene.GetSceneExporter().GetAssetsOutputPath(), DatasmithMesh, nullptr, FDatasmithExportOptions::LightmapUV))
+				{
+					// todo: handle error exporting mesh?
+				}
+			}
+
+			CachedMesh.FreeAll();
+		}
+		if (bNeedsDelete)
+		{
+			RenderMesh->DeleteThis();
+		}
+
+		return bResult;
+	}
+
+	void ConvertNodeTransform(FNodeTracker& NodeTracker)
+	{
+		FVector Translation, Scale;
+		FQuat Rotation;
+
+		FDatasmithConverter Converter;
+
+		// todo: do we really need to call GetObjectTM if there's no WSM attached? Maybe just call GetObjTMAfterWSM always?
+		if (NodeTracker.Node->GetWSMDerivedObject() != nullptr)
+		{
+
+			Converter.MaxToUnrealCoordinates(NodeTracker.Node->GetObjTMAfterWSM(GetCOREInterface()->GetTime()), Translation, Rotation, Scale);
+		}
+		else
+		{
+			Converter.MaxToUnrealCoordinates(NodeTracker.Node->GetObjectTM(GetCOREInterface()->GetTime()), Translation, Rotation, Scale);
+		}
+
+		Rotation.Normalize();
+
+		TSharedPtr<IDatasmithActorElement> DatasmithActorElement = NodeTracker.DatasmithActorElement;
+
+		if (DatasmithActorElement)
+		{
+			DatasmithActorElement->SetTranslation(Translation);
+			DatasmithActorElement->SetScale(Scale);
+			DatasmithActorElement->SetRotation(Rotation);
+		}
+	}
+
+	/******************* Events *****************************/
+
+	void NodeAdded(INode* Node)
+	{
+		// Node sometimes is null. 'Added' NodeEvent might come after node was actually deleted(immediately after creation)
+		// e.g.[mxs]: b = box(); delete b 
+		// NodeEvents are delayed(not executed in the same stack frame as command that causes them) so they come later. 
+		if (!Node)
+		{
+			return;
+		}
+
+		ParseNode(Node);
+	}
+
+	void NodeDeleted(INode* Node)
+	{
+		// todo: check for null
+
+		FNodeKey NodeKey = NodeEventNamespace::GetKeyByNode(Node);
+
+		if (FNodeTrackerHandle* NodeTrackerHandle = NodeTrackers.Find(NodeKey))
+		{
+			// todo: schedule for delete on Update?
+			FNodeTracker* NodeTracker = NodeTrackerHandle->GetNodeTracker();
+			InvalidatedNodeTrackers.Remove(NodeTracker);
+			NodeTrackers.Remove(NodeKey);
+
+			if (TSharedPtr<IDatasmithActorElement> DatasmithActorElement = NodeTracker->DatasmithActorElement)
+			{
+				const TSharedPtr<IDatasmithActorElement>& ParentActor = DatasmithActorElement->GetParentActor();
+
+				if (ParentActor)
+				{
+					// todo: 
+				}
+				else
+				{
+					// todo: remove children??? check that when node is deleted - it's deleted with all its children either:
+					// - children are deleted prior to deleting parent
+					// - in other order - then make sure not to confuse - need to remove NodeTrackers for children or leave them dangling and remove when their event comes - 
+					//    IMPORTANT - we are testing that a datasmith actor is at root by its parent AND the dangling datasmith actors will have no parent... Change by adding 'root' actor flag?
+					//
+					ExportedScene.DatasmithSceneRef->RemoveActor(DatasmithActorElement, EDatasmithActorRemovalRule::RemoveChildren);
+				}
+			}
+
+			//todo: remove ActorMesh
+			//todo: remove Mesh(if not used)
+		}
+
+	}
+
+	void NodeTransformChanged(FNodeKey NodeKey)
+	{
+		// todo: invalidate transform only
+
+		// todo: grouping makes this crash. Need to handle event before?
+		InvalidateNode(NodeKey);
+
+		// ControllerOtherEvent sent only for top actors in hierarchy when it's moved
+		INode* Node = NodeEventNamespace::GetNodeByKey(NodeKey);
+		if (Node)
+		{
+			int32 ChildNum = Node->NumberOfChildren();
+			for (int32 ChildIndex = 0; ChildIndex < ChildNum; ++ChildIndex)
+			{
+				// todo: pass INode to NodeTransformChanged to remove redundant get
+				NodeTransformChanged(NodeEventNamespace::GetKeyByNode(Node->GetChildNode(ChildIndex)));
+			}
+		}
+	}
+
+	void NodeMaterialAssignmentChanged(FNodeKey NodeKey)
+	{
+		//todo: handle more precisely
+		InvalidateNode(NodeKey);
+	}
+
+	void NodeMaterialGraphModified(FNodeKey NodeKey)
+	{
+		//identify material tree and update all materials
+		//todo: possible to handle this more precisely(only refresh changed materials) - see FMaterialObserver
+
+		if (FNodeTrackerHandle* NodeTracker = NodeTrackers.Find(NodeKey))
+		{
+			// todo: investigate why NodeEventNamespace::GetNodeByKey may stil return NULL
+			// testcase - add XRef Material - this will immediately have this 
+			// even though NOTIFY_SCENE_ADDED_NODE was called for node and NOTIFY_SCENE_PRE_DELETED_NODE wasn't!
+			INode* Node = NodeEventNamespace::GetNodeByKey(NodeKey);
+			if (Node)
+			{
+				if (Mtl* Material = Node->GetMtl())
+				{
+					if (FMaterialTrackerHandle* MaterialTrackerHandle = MaterialTrackers.Find(Material))
+					{
+						InvalidatedMaterialTrackers.Add(MaterialTrackerHandle->GetMaterialTracker());
+					}
+				}
+			}
+		}
+	}
+
+	void NodeGeometryChanged(FNodeKey NodeKey)
+	{
+		// GeometryChanged is executed to handle:
+		// - actual geometry modification(in any way)
+		// - change of baseObject
+
+		// todo: how this could be?
+		ensure(NodeTrackers.Contains(NodeKey));
+
+		InvalidateNode(NodeKey);
+	}
+	///////////////////////////////////////////////
+
+
+	FDatasmith3dsMaxScene& ExportedScene;
+	TMap<FNodeKey, FNodeTrackerHandle> NodeTrackers;
+
+	TSet<FNodeTracker*> InvalidatedNodeTrackers;
+
+	FNodeObserver NodeObserver;
+	FMaterialObserver MaterialObserver;
+
+	FMaterialsTracker MaterialsTracker;
+	TMap<FMaterialKey, FMaterialTrackerHandle> MaterialTrackers;
+	TSet<FMaterialTracker*> InvalidatedMaterialTrackers;
+
+	TMap<FMaterialTracker*, TSet<FNodeTracker*>> MaterialsAssignedToNodes;
+};
+
+
+// This is used to handle some of change events
+class FNodeEventCallback : public INodeEventCallback
+{
+	FSceneTracker& SceneTracker;
+
+public:
+	FNodeEventCallback(FSceneTracker& InSceneTracker) : SceneTracker(InSceneTracker)
+	{
+	}
+
+	virtual BOOL VerboseDeleted() { return TRUE; }
+
+	virtual void GeometryChanged(NodeKeyTab& nodes) override
+	{
+		LogNodeEvent(L"GeometryChanged", nodes);
+
+		for (int NodeIndex = 0; NodeIndex < nodes.Count(); ++NodeIndex)
+		{
+			SceneTracker.NodeGeometryChanged(nodes[NodeIndex]);
+		}
+	}
+
+	// Fired when node transform changes
+	virtual void ControllerOtherEvent(NodeKeyTab& nodes) override
+	{
+		LogNodeEvent(L"ControllerOtherEvent", nodes);
+
+		for (int NodeIndex = 0; NodeIndex < nodes.Count(); ++NodeIndex)
+		{
+			SceneTracker.NodeTransformChanged(nodes[NodeIndex]);
+		}
+	}
+
+	// Tracks material assignment on node
+	virtual void MaterialStructured(NodeKeyTab& nodes) override
+	{
+		LogNodeEvent(L"MaterialStructured", nodes);
+		for (int NodeIndex = 0; NodeIndex < nodes.Count(); ++NodeIndex)
+		{
+			SceneTracker.NodeMaterialAssignmentChanged(nodes[NodeIndex]);
+		}
+	}
+
+	// Tracks node's material parameter change(even if it's a submaterial of multimat that is assigned)
+	virtual void MaterialOtherEvent(NodeKeyTab& nodes) override
+	{
+		LogNodeEvent(L"MaterialOtherEvent", nodes);
+		for (int NodeIndex = 0; NodeIndex < nodes.Count(); ++NodeIndex)
+		{
+			SceneTracker.NodeMaterialGraphModified(nodes[NodeIndex]);
+		}
+	}
+
+	// Not used:
+
+	virtual void Added(NodeKeyTab& nodes) override
+	{
+		LogNodeEvent(L"Added", nodes);
+	}
+
+	virtual void Deleted(NodeKeyTab& nodes) override
+	{
+		LogNodeEvent(L"Deleted", nodes);
+	}
+
+	virtual void LinkChanged(NodeKeyTab& nodes) override
+	{
+		LogNodeEvent(L"LinkChanged", nodes);
+	}
+
+	virtual void LayerChanged(NodeKeyTab& nodes) override
+	{
+		LogNodeEvent(L"LayerChanged", nodes);
+	}
+
+	virtual void GroupChanged(NodeKeyTab& nodes) override
+	{
+		LogNodeEvent(L"GroupChanged", nodes);
+	}
+
+	virtual void HierarchyOtherEvent(NodeKeyTab& nodes) override
+	{
+		LogNodeEvent(L"HierarchyOtherEvent", nodes);
+	}
+
+	virtual void ModelStructured(NodeKeyTab& nodes) override
+	{
+		LogNodeEvent(L"ModelStructured", nodes);
+	}
+
+	virtual void TopologyChanged(NodeKeyTab& nodes) override
+	{
+		LogNodeEvent(L"TopologyChanged", nodes);
+	}
+
+	virtual void MappingChanged(NodeKeyTab& nodes) override
+	{
+		LogNodeEvent(L"MappingChanged", nodes);
+	}
+
+	virtual void ExtentionChannelChanged(NodeKeyTab& nodes) override
+	{
+		LogNodeEvent(L"ExtentionChannelChanged", nodes);
+	}
+
+	virtual void ModelOtherEvent(NodeKeyTab& nodes) override
+	{
+		LogNodeEvent(L"ModelOtherEvent", nodes);
+	}
+
+	virtual void ControllerStructured(NodeKeyTab& nodes) override
+	{
+		LogNodeEvent(L"ControllerStructured", nodes);
+	}
+
+	virtual void NameChanged(NodeKeyTab& nodes) override
+	{
+		LogNodeEvent(L"NameChanged", nodes);
+	}
+
+	virtual void WireColorChanged(NodeKeyTab& nodes) override
+	{
+		LogNodeEvent(L"WireColorChanged", nodes);
+	}
+
+	virtual void RenderPropertiesChanged(NodeKeyTab& nodes) override
+	{
+		LogNodeEvent(L"RenderPropertiesChanged", nodes);
+	}
+
+	virtual void DisplayPropertiesChanged(NodeKeyTab& nodes) override
+	{
+		LogNodeEvent(L"DisplayPropertiesChanged", nodes);
+	}
+
+	virtual void UserPropertiesChanged(NodeKeyTab& nodes) override
+	{
+		LogNodeEvent(L"UserPropertiesChanged", nodes);
+	}
+
+	virtual void PropertiesOtherEvent(NodeKeyTab& nodes) override
+	{
+		LogNodeEvent(L"PropertiesOtherEvent", nodes);
+	}
+
+	virtual void SubobjectSelectionChanged(NodeKeyTab& nodes) override
+	{
+		LogNodeEvent(L"SubobjectSelectionChanged", nodes);
+	}
+
+	virtual void SelectionChanged(NodeKeyTab& nodes) override
+	{
+		LogNodeEvent(L"SelectionChanged", nodes);
+	}
+
+	virtual void HideChanged(NodeKeyTab& nodes) override
+	{
+		LogNodeEvent(L"HideChanged", nodes);
+	}
+
+	virtual void FreezeChanged(NodeKeyTab& nodes) override
+	{
+		LogNodeEvent(L"FreezeChanged", nodes);
+	}
+
+	virtual void DisplayOtherEvent(NodeKeyTab& nodes) override
+	{
+		LogNodeEvent(L"DisplayOtherEvent", nodes);
+	}
+
+	virtual void CallbackBegin() override
+	{
+
+		// todo: heavy logging - called a lot
+		return;
+		LogDebug(L"NodeEventCallback: CallbackBegin\n");
+	}
+
+	virtual void CallbackEnd() override
+	{
+
+		// todo: heavy logging - called a lot
+		return;
+		LogDebug(L"NodeEventCallback: CallbackEnd\n");
+	}
+};
+
+
+class FExporter
+{
+public:
+
+	FExporter(): SceneTracker(ExportedScene), NodeEventCallback(SceneTracker) {}
+
+	void SetOutputPath(const TCHAR* Path)
+	{
+		OutputPath = Path;
+		ExportedScene.SetOutputPath(*OutputPath);
+	}
+
+	// Just export, parsing scene from scratch
+	bool Export()
+	{
+		SceneTracker.ParseScene();
+		SceneTracker.Update();
+		ExportedScene.GetSceneExporter().Export(ExportedScene.GetDatasmithScene(), false);
+		return true;
+	}
+
+	// Install change notification systems
+	void StartSceneChangeTracking()
+	{
+		// Build todo: remove strings, for debug/logging
+#pragma warning(push)
+#pragma warning(disable:4995) // disable error on deprecated events, we just assign handlers not firing them
+		int Codes[] = { NOTIFY_UNITS_CHANGE, NOTIFY_TIMEUNITS_CHANGE, NOTIFY_VIEWPORT_CHANGE, NOTIFY_SPACEMODE_CHANGE, NOTIFY_SYSTEM_PRE_RESET, NOTIFY_SYSTEM_POST_RESET, NOTIFY_SYSTEM_PRE_NEW, NOTIFY_SYSTEM_POST_NEW, NOTIFY_FILE_PRE_OPEN, NOTIFY_FILE_POST_OPEN, NOTIFY_FILE_PRE_MERGE, NOTIFY_FILE_POST_MERGE, NOTIFY_FILE_PRE_SAVE, NOTIFY_FILE_POST_SAVE, NOTIFY_FILE_OPEN_FAILED, NOTIFY_FILE_PRE_SAVE_OLD, NOTIFY_FILE_POST_SAVE_OLD, NOTIFY_SELECTIONSET_CHANGED, NOTIFY_BITMAP_CHANGED, NOTIFY_PRE_RENDER, NOTIFY_POST_RENDER, NOTIFY_PRE_RENDERFRAME, NOTIFY_POST_RENDERFRAME, NOTIFY_PRE_IMPORT, NOTIFY_POST_IMPORT, NOTIFY_IMPORT_FAILED, NOTIFY_PRE_EXPORT, NOTIFY_POST_EXPORT, NOTIFY_EXPORT_FAILED, NOTIFY_NODE_RENAMED, NOTIFY_PRE_PROGRESS, NOTIFY_POST_PROGRESS, NOTIFY_MODPANEL_SEL_CHANGED, NOTIFY_RENDPARAM_CHANGED, NOTIFY_MATLIB_PRE_OPEN, NOTIFY_MATLIB_POST_OPEN, NOTIFY_MATLIB_PRE_SAVE, NOTIFY_MATLIB_POST_SAVE, NOTIFY_MATLIB_PRE_MERGE, NOTIFY_MATLIB_POST_MERGE, NOTIFY_FILELINK_BIND_FAILED, NOTIFY_FILELINK_DETACH_FAILED, NOTIFY_FILELINK_RELOAD_FAILED, NOTIFY_FILELINK_ATTACH_FAILED, NOTIFY_FILELINK_PRE_BIND, NOTIFY_FILELINK_POST_BIND, NOTIFY_FILELINK_PRE_DETACH, NOTIFY_FILELINK_POST_DETACH, NOTIFY_FILELINK_PRE_RELOAD, NOTIFY_FILELINK_POST_RELOAD, NOTIFY_FILELINK_PRE_ATTACH, NOTIFY_FILELINK_POST_ATTACH, NOTIFY_RENDER_PREEVAL, NOTIFY_NODE_CREATED, NOTIFY_NODE_LINKED, NOTIFY_NODE_UNLINKED, NOTIFY_NODE_HIDE, NOTIFY_NODE_UNHIDE, NOTIFY_NODE_FREEZE, NOTIFY_NODE_UNFREEZE, NOTIFY_NODE_PRE_MTL, NOTIFY_NODE_POST_MTL, NOTIFY_SCENE_ADDED_NODE, NOTIFY_SCENE_PRE_DELETED_NODE, NOTIFY_SCENE_POST_DELETED_NODE, NOTIFY_SEL_NODES_PRE_DELETE, NOTIFY_SEL_NODES_POST_DELETE, NOTIFY_WM_ENABLE, NOTIFY_SYSTEM_SHUTDOWN, NOTIFY_SYSTEM_STARTUP, NOTIFY_PLUGIN_LOADED, NOTIFY_SYSTEM_SHUTDOWN2, NOTIFY_ANIMATE_ON, NOTIFY_ANIMATE_OFF, NOTIFY_COLOR_CHANGE, NOTIFY_PRE_EDIT_OBJ_CHANGE, NOTIFY_POST_EDIT_OBJ_CHANGE, NOTIFY_RADIOSITYPROCESS_STARTED, NOTIFY_RADIOSITYPROCESS_STOPPED, NOTIFY_RADIOSITYPROCESS_RESET, NOTIFY_RADIOSITYPROCESS_DONE, NOTIFY_LIGHTING_UNIT_DISPLAY_SYSTEM_CHANGE, NOTIFY_BEGIN_RENDERING_REFLECT_REFRACT_MAP, NOTIFY_BEGIN_RENDERING_ACTUAL_FRAME, NOTIFY_BEGIN_RENDERING_TONEMAPPING_IMAGE, NOTIFY_RADIOSITY_PLUGIN_CHANGED, NOTIFY_SCENE_UNDO, NOTIFY_SCENE_REDO, NOTIFY_MANIPULATE_MODE_OFF, NOTIFY_MANIPULATE_MODE_ON, NOTIFY_SCENE_XREF_PRE_MERGE, NOTIFY_SCENE_XREF_POST_MERGE, NOTIFY_OBJECT_XREF_PRE_MERGE, NOTIFY_OBJECT_XREF_POST_MERGE, NOTIFY_PRE_MIRROR_NODES, NOTIFY_POST_MIRROR_NODES, NOTIFY_NODE_CLONED, NOTIFY_PRE_NOTIFYDEPENDENTS, NOTIFY_POST_NOTIFYDEPENDENTS, NOTIFY_MTL_REFDELETED, NOTIFY_TIMERANGE_CHANGE, NOTIFY_PRE_MODIFIER_ADDED, NOTIFY_POST_MODIFIER_ADDED, NOTIFY_PRE_MODIFIER_DELETED, NOTIFY_POST_MODIFIER_DELETED, NOTIFY_FILELINK_POST_RELOAD_PRE_PRUNE, NOTIFY_PRE_NODES_CLONED, NOTIFY_POST_NODES_CLONED, NOTIFY_SYSTEM_PRE_DIR_CHANGE, NOTIFY_SYSTEM_POST_DIR_CHANGE, NOTIFY_SV_SELECTIONSET_CHANGED, NOTIFY_SV_DOUBLECLICK_GRAPHNODE, NOTIFY_PRE_RENDERER_CHANGE, NOTIFY_POST_RENDERER_CHANGE, NOTIFY_SV_PRE_LAYOUT_CHANGE, NOTIFY_SV_POST_LAYOUT_CHANGE, NOTIFY_BY_CATEGORY_DISPLAY_FILTER_CHANGED, NOTIFY_CUSTOM_DISPLAY_FILTER_CHANGED, NOTIFY_LAYER_CREATED, NOTIFY_LAYER_DELETED, NOTIFY_NODE_LAYER_CHANGED, NOTIFY_TABBED_DIALOG_CREATED, NOTIFY_TABBED_DIALOG_DELETED, NOTIFY_NODE_NAME_SET, NOTIFY_HW_TEXTURE_CHANGED, NOTIFY_MXS_STARTUP, NOTIFY_MXS_POST_STARTUP, NOTIFY_ACTION_ITEM_HOTKEY_PRE_EXEC, NOTIFY_ACTION_ITEM_HOTKEY_POST_EXEC, NOTIFY_SCENESTATE_PRE_SAVE, NOTIFY_SCENESTATE_POST_SAVE, NOTIFY_SCENESTATE_PRE_RESTORE, NOTIFY_SCENESTATE_POST_RESTORE, NOTIFY_SCENESTATE_DELETE, NOTIFY_SCENESTATE_RENAME, NOTIFY_SCENE_PRE_UNDO, NOTIFY_SCENE_PRE_REDO, NOTIFY_SCENE_POST_UNDO, NOTIFY_SCENE_POST_REDO, NOTIFY_MXS_SHUTDOWN, NOTIFY_D3D_PRE_DEVICE_RESET, NOTIFY_D3D_POST_DEVICE_RESET, NOTIFY_TOOLPALETTE_MTL_SUSPEND, NOTIFY_TOOLPALETTE_MTL_RESUME, NOTIFY_CLASSDESC_REPLACED, NOTIFY_FILE_PRE_OPEN_PROCESS, NOTIFY_FILE_POST_OPEN_PROCESS, NOTIFY_FILE_PRE_SAVE_PROCESS, NOTIFY_FILE_POST_SAVE_PROCESS, NOTIFY_CLASSDESC_LOADED, NOTIFY_TOOLBARS_PRE_LOAD, NOTIFY_TOOLBARS_POST_LOAD, NOTIFY_ATS_PRE_REPATH_PHASE, NOTIFY_ATS_POST_REPATH_PHASE, NOTIFY_PROXY_TEMPORARY_DISABLE_START, NOTIFY_PROXY_TEMPORARY_DISABLE_END, NOTIFY_FILE_CHECK_STATUS, NOTIFY_NAMED_SEL_SET_CREATED, NOTIFY_NAMED_SEL_SET_DELETED, NOTIFY_NAMED_SEL_SET_RENAMED, NOTIFY_NAMED_SEL_SET_PRE_MODIFY, NOTIFY_NAMED_SEL_SET_POST_MODIFY, NOTIFY_MODPANEL_SUBOBJECTLEVEL_CHANGED, NOTIFY_FAILED_DIRECTX_MATERIAL_TEXTURE_LOAD, NOTIFY_RENDER_PREEVAL_FRAMEINFO, NOTIFY_POST_SCENE_RESET, NOTIFY_ANIM_LAYERS_ENABLED, NOTIFY_ANIM_LAYERS_DISABLED, NOTIFY_ACTION_ITEM_PRE_START_OVERRIDE, NOTIFY_ACTION_ITEM_POST_START_OVERRIDE, NOTIFY_ACTION_ITEM_PRE_END_OVERRIDE, NOTIFY_ACTION_ITEM_POST_END_OVERRIDE, NOTIFY_PRE_NODE_GENERAL_PROP_CHANGED, NOTIFY_POST_NODE_GENERAL_PROP_CHANGED, NOTIFY_PRE_NODE_GI_PROP_CHANGED, NOTIFY_POST_NODE_GI_PROP_CHANGED, NOTIFY_PRE_NODE_MENTALRAY_PROP_CHANGED, NOTIFY_POST_NODE_MENTALRAY_PROP_CHANGED, NOTIFY_PRE_NODE_BONE_PROP_CHANGED, NOTIFY_POST_NODE_BONE_PROP_CHANGED, NOTIFY_PRE_NODE_USER_PROP_CHANGED, NOTIFY_POST_NODE_USER_PROP_CHANGED, NOTIFY_PRE_NODE_RENDER_PROP_CHANGED, NOTIFY_POST_NODE_RENDER_PROP_CHANGED, NOTIFY_PRE_NODE_DISPLAY_PROP_CHANGED, NOTIFY_POST_NODE_DISPLAY_PROP_CHANGED, NOTIFY_PRE_NODE_BASIC_PROP_CHANGED, NOTIFY_POST_NODE_BASIC_PROP_CHANGED, NOTIFY_SELECTION_LOCK, NOTIFY_SELECTION_UNLOCK, NOTIFY_PRE_IMAGE_VIEWER_DISPLAY, NOTIFY_POST_IMAGE_VIEWER_DISPLAY, NOTIFY_IMAGE_VIEWER_UPDATE, NOTIFY_CUSTOM_ATTRIBUTES_ADDED, NOTIFY_CUSTOM_ATTRIBUTES_REMOVED, NOTIFY_OS_THEME_CHANGED, NOTIFY_ACTIVE_VIEWPORT_CHANGED, NOTIFY_PRE_MAXMAINWINDOW_SHOW, NOTIFY_POST_MAXMAINWINDOW_SHOW, NOTIFY_CLASSDESC_ADDED, NOTIFY_OBJECT_DEFINITION_CHANGE_BEGIN, NOTIFY_OBJECT_DEFINITION_CHANGE_END, NOTIFY_MTLBASE_PARAMDLG_PRE_OPEN, NOTIFY_MTLBASE_PARAMDLG_POST_CLOSE, NOTIFY_PRE_APP_FRAME_THEME_CHANGED, NOTIFY_APP_FRAME_THEME_CHANGED, NOTIFY_PRE_VIEWPORT_DELETE, NOTIFY_PRE_WORKSPACE_CHANGE, NOTIFY_POST_WORKSPACE_CHANGE, NOTIFY_PRE_WORKSPACE_COLLECTION_CHANGE, NOTIFY_POST_WORKSPACE_COLLECTION_CHANGE, NOTIFY_KEYBOARD_SETTING_CHANGED, NOTIFY_MOUSE_SETTING_CHANGED, NOTIFY_TOOLBARS_PRE_SAVE, NOTIFY_TOOLBARS_POST_SAVE, NOTIFY_APP_ACTIVATED, NOTIFY_APP_DEACTIVATED, NOTIFY_CUI_MENUS_UPDATED, NOTIFY_CUI_MENUS_PRE_SAVE, NOTIFY_CUI_MENUS_POST_SAVE, NOTIFY_VIEWPORT_SAFEFRAME_TOGGLE, NOTIFY_PLUGINS_PRE_SHUTDOWN, NOTIFY_PLUGINS_PRE_UNLOAD, NOTIFY_CUI_MENUS_POST_LOAD, NOTIFY_LAYER_PARENT_CHANGED, NOTIFY_ACTION_ITEM_EXECUTION_STARTED, NOTIFY_ACTION_ITEM_EXECUTION_ENDED, NOTIFY_INTERACTIVE_PLUGIN_INSTANCE_CREATION_STARTED, NOTIFY_INTERACTIVE_PLUGIN_INSTANCE_CREATION_ENDED, NOTIFY_FILE_POST_MERGE2, NOTIFY_POST_NODE_SELECT_OPERATION, NOTIFY_PRE_VIEWPORT_TOOLTIP, NOTIFY_WELCOMESCREEN_DONE, NOTIFY_PLAYBACK_START, NOTIFY_PLAYBACK_END, NOTIFY_SCENE_EXPLORER_NEEDS_UPDATE, NOTIFY_FILE_POST_OPEN_PROCESS_FINALIZED, NOTIFY_FILE_POST_MERGE_PROCESS_FINALIZED, NOTIFY_PRE_PROJECT_FOLDER_CHANGE, NOTIFY_POST_PROJECT_FOLDER_CHANGE, NOTIFY_PRE_MXS_STARTUP_SCRIPT_LOAD, NOTIFY_ACTIVESHADE_IN_VIEWPORT_TOGGLED, NOTIFY_SYSTEM_SHUTDOWN_CHECK, NOTIFY_SYSTEM_SHUTDOWN_CHECK_FAILED, NOTIFY_SYSTEM_SHUTDOWN_CHECK_PASSED, NOTIFY_FILE_POST_MERGE3, NOTIFY_ACTIVESHADE_IN_FRAMEBUFFER_TOGGLED, NOTIFY_PRE_ACTIVESHADE_IN_VIEWPORT_TOGGLED, NOTIFY_POST_ACTIVESHADE_IN_VIEWPORT_TOGGLED, NOTIFY_INTERNAL_USE_START };
+		FString Strings[] = { TEXT("NOTIFY_UNITS_CHANGE"), TEXT("NOTIFY_TIMEUNITS_CHANGE"), TEXT("NOTIFY_VIEWPORT_CHANGE"), TEXT("NOTIFY_SPACEMODE_CHANGE"), TEXT("NOTIFY_SYSTEM_PRE_RESET"), TEXT("NOTIFY_SYSTEM_POST_RESET"), TEXT("NOTIFY_SYSTEM_PRE_NEW"), TEXT("NOTIFY_SYSTEM_POST_NEW"), TEXT("NOTIFY_FILE_PRE_OPEN"), TEXT("NOTIFY_FILE_POST_OPEN"), TEXT("NOTIFY_FILE_PRE_MERGE"), TEXT("NOTIFY_FILE_POST_MERGE"), TEXT("NOTIFY_FILE_PRE_SAVE"), TEXT("NOTIFY_FILE_POST_SAVE"), TEXT("NOTIFY_FILE_OPEN_FAILED"), TEXT("NOTIFY_FILE_PRE_SAVE_OLD"), TEXT("NOTIFY_FILE_POST_SAVE_OLD"), TEXT("NOTIFY_SELECTIONSET_CHANGED"), TEXT("NOTIFY_BITMAP_CHANGED"), TEXT("NOTIFY_PRE_RENDER"), TEXT("NOTIFY_POST_RENDER"), TEXT("NOTIFY_PRE_RENDERFRAME"), TEXT("NOTIFY_POST_RENDERFRAME"), TEXT("NOTIFY_PRE_IMPORT"), TEXT("NOTIFY_POST_IMPORT"), TEXT("NOTIFY_IMPORT_FAILED"), TEXT("NOTIFY_PRE_EXPORT"), TEXT("NOTIFY_POST_EXPORT"), TEXT("NOTIFY_EXPORT_FAILED"), TEXT("NOTIFY_NODE_RENAMED"), TEXT("NOTIFY_PRE_PROGRESS"), TEXT("NOTIFY_POST_PROGRESS"), TEXT("NOTIFY_MODPANEL_SEL_CHANGED"), TEXT("NOTIFY_RENDPARAM_CHANGED"), TEXT("NOTIFY_MATLIB_PRE_OPEN"), TEXT("NOTIFY_MATLIB_POST_OPEN"), TEXT("NOTIFY_MATLIB_PRE_SAVE"), TEXT("NOTIFY_MATLIB_POST_SAVE"), TEXT("NOTIFY_MATLIB_PRE_MERGE"), TEXT("NOTIFY_MATLIB_POST_MERGE"), TEXT("NOTIFY_FILELINK_BIND_FAILED"), TEXT("NOTIFY_FILELINK_DETACH_FAILED"), TEXT("NOTIFY_FILELINK_RELOAD_FAILED"), TEXT("NOTIFY_FILELINK_ATTACH_FAILED"), TEXT("NOTIFY_FILELINK_PRE_BIND"), TEXT("NOTIFY_FILELINK_POST_BIND"), TEXT("NOTIFY_FILELINK_PRE_DETACH"), TEXT("NOTIFY_FILELINK_POST_DETACH"), TEXT("NOTIFY_FILELINK_PRE_RELOAD"), TEXT("NOTIFY_FILELINK_POST_RELOAD"), TEXT("NOTIFY_FILELINK_PRE_ATTACH"), TEXT("NOTIFY_FILELINK_POST_ATTACH"), TEXT("NOTIFY_RENDER_PREEVAL"), TEXT("NOTIFY_NODE_CREATED"), TEXT("NOTIFY_NODE_LINKED"), TEXT("NOTIFY_NODE_UNLINKED"), TEXT("NOTIFY_NODE_HIDE"), TEXT("NOTIFY_NODE_UNHIDE"), TEXT("NOTIFY_NODE_FREEZE"), TEXT("NOTIFY_NODE_UNFREEZE"), TEXT("NOTIFY_NODE_PRE_MTL"), TEXT("NOTIFY_NODE_POST_MTL"), TEXT("NOTIFY_SCENE_ADDED_NODE"), TEXT("NOTIFY_SCENE_PRE_DELETED_NODE"), TEXT("NOTIFY_SCENE_POST_DELETED_NODE"), TEXT("NOTIFY_SEL_NODES_PRE_DELETE"), TEXT("NOTIFY_SEL_NODES_POST_DELETE"), TEXT("NOTIFY_WM_ENABLE"), TEXT("NOTIFY_SYSTEM_SHUTDOWN"), TEXT("NOTIFY_SYSTEM_STARTUP"), TEXT("NOTIFY_PLUGIN_LOADED"), TEXT("NOTIFY_SYSTEM_SHUTDOWN2"), TEXT("NOTIFY_ANIMATE_ON"), TEXT("NOTIFY_ANIMATE_OFF"), TEXT("NOTIFY_COLOR_CHANGE"), TEXT("NOTIFY_PRE_EDIT_OBJ_CHANGE"), TEXT("NOTIFY_POST_EDIT_OBJ_CHANGE"), TEXT("NOTIFY_RADIOSITYPROCESS_STARTED"), TEXT("NOTIFY_RADIOSITYPROCESS_STOPPED"), TEXT("NOTIFY_RADIOSITYPROCESS_RESET"), TEXT("NOTIFY_RADIOSITYPROCESS_DONE"), TEXT("NOTIFY_LIGHTING_UNIT_DISPLAY_SYSTEM_CHANGE"), TEXT("NOTIFY_BEGIN_RENDERING_REFLECT_REFRACT_MAP"), TEXT("NOTIFY_BEGIN_RENDERING_ACTUAL_FRAME"), TEXT("NOTIFY_BEGIN_RENDERING_TONEMAPPING_IMAGE"), TEXT("NOTIFY_RADIOSITY_PLUGIN_CHANGED"), TEXT("NOTIFY_SCENE_UNDO"), TEXT("NOTIFY_SCENE_REDO"), TEXT("NOTIFY_MANIPULATE_MODE_OFF"), TEXT("NOTIFY_MANIPULATE_MODE_ON"), TEXT("NOTIFY_SCENE_XREF_PRE_MERGE"), TEXT("NOTIFY_SCENE_XREF_POST_MERGE"), TEXT("NOTIFY_OBJECT_XREF_PRE_MERGE"), TEXT("NOTIFY_OBJECT_XREF_POST_MERGE"), TEXT("NOTIFY_PRE_MIRROR_NODES"), TEXT("NOTIFY_POST_MIRROR_NODES"), TEXT("NOTIFY_NODE_CLONED"), TEXT("NOTIFY_PRE_NOTIFYDEPENDENTS"), TEXT("NOTIFY_POST_NOTIFYDEPENDENTS"), TEXT("NOTIFY_MTL_REFDELETED"), TEXT("NOTIFY_TIMERANGE_CHANGE"), TEXT("NOTIFY_PRE_MODIFIER_ADDED"), TEXT("NOTIFY_POST_MODIFIER_ADDED"), TEXT("NOTIFY_PRE_MODIFIER_DELETED"), TEXT("NOTIFY_POST_MODIFIER_DELETED"), TEXT("NOTIFY_FILELINK_POST_RELOAD_PRE_PRUNE"), TEXT("NOTIFY_PRE_NODES_CLONED"), TEXT("NOTIFY_POST_NODES_CLONED"), TEXT("NOTIFY_SYSTEM_PRE_DIR_CHANGE"), TEXT("NOTIFY_SYSTEM_POST_DIR_CHANGE"), TEXT("NOTIFY_SV_SELECTIONSET_CHANGED"), TEXT("NOTIFY_SV_DOUBLECLICK_GRAPHNODE"), TEXT("NOTIFY_PRE_RENDERER_CHANGE"), TEXT("NOTIFY_POST_RENDERER_CHANGE"), TEXT("NOTIFY_SV_PRE_LAYOUT_CHANGE"), TEXT("NOTIFY_SV_POST_LAYOUT_CHANGE"), TEXT("NOTIFY_BY_CATEGORY_DISPLAY_FILTER_CHANGED"), TEXT("NOTIFY_CUSTOM_DISPLAY_FILTER_CHANGED"), TEXT("NOTIFY_LAYER_CREATED"), TEXT("NOTIFY_LAYER_DELETED"), TEXT("NOTIFY_NODE_LAYER_CHANGED"), TEXT("NOTIFY_TABBED_DIALOG_CREATED"), TEXT("NOTIFY_TABBED_DIALOG_DELETED"), TEXT("NOTIFY_NODE_NAME_SET"), TEXT("NOTIFY_HW_TEXTURE_CHANGED"), TEXT("NOTIFY_MXS_STARTUP"), TEXT("NOTIFY_MXS_POST_STARTUP"), TEXT("NOTIFY_ACTION_ITEM_HOTKEY_PRE_EXEC"), TEXT("NOTIFY_ACTION_ITEM_HOTKEY_POST_EXEC"), TEXT("NOTIFY_SCENESTATE_PRE_SAVE"), TEXT("NOTIFY_SCENESTATE_POST_SAVE"), TEXT("NOTIFY_SCENESTATE_PRE_RESTORE"), TEXT("NOTIFY_SCENESTATE_POST_RESTORE"), TEXT("NOTIFY_SCENESTATE_DELETE"), TEXT("NOTIFY_SCENESTATE_RENAME"), TEXT("NOTIFY_SCENE_PRE_UNDO"), TEXT("NOTIFY_SCENE_PRE_REDO"), TEXT("NOTIFY_SCENE_POST_UNDO"), TEXT("NOTIFY_SCENE_POST_REDO"), TEXT("NOTIFY_MXS_SHUTDOWN"), TEXT("NOTIFY_D3D_PRE_DEVICE_RESET"), TEXT("NOTIFY_D3D_POST_DEVICE_RESET"), TEXT("NOTIFY_TOOLPALETTE_MTL_SUSPEND"), TEXT("NOTIFY_TOOLPALETTE_MTL_RESUME"), TEXT("NOTIFY_CLASSDESC_REPLACED"), TEXT("NOTIFY_FILE_PRE_OPEN_PROCESS"), TEXT("NOTIFY_FILE_POST_OPEN_PROCESS"), TEXT("NOTIFY_FILE_PRE_SAVE_PROCESS"), TEXT("NOTIFY_FILE_POST_SAVE_PROCESS"), TEXT("NOTIFY_CLASSDESC_LOADED"), TEXT("NOTIFY_TOOLBARS_PRE_LOAD"), TEXT("NOTIFY_TOOLBARS_POST_LOAD"), TEXT("NOTIFY_ATS_PRE_REPATH_PHASE"), TEXT("NOTIFY_ATS_POST_REPATH_PHASE"), TEXT("NOTIFY_PROXY_TEMPORARY_DISABLE_START"), TEXT("NOTIFY_PROXY_TEMPORARY_DISABLE_END"), TEXT("NOTIFY_FILE_CHECK_STATUS"), TEXT("NOTIFY_NAMED_SEL_SET_CREATED"), TEXT("NOTIFY_NAMED_SEL_SET_DELETED"), TEXT("NOTIFY_NAMED_SEL_SET_RENAMED"), TEXT("NOTIFY_NAMED_SEL_SET_PRE_MODIFY"), TEXT("NOTIFY_NAMED_SEL_SET_POST_MODIFY"), TEXT("NOTIFY_MODPANEL_SUBOBJECTLEVEL_CHANGED"), TEXT("NOTIFY_FAILED_DIRECTX_MATERIAL_TEXTURE_LOAD"), TEXT("NOTIFY_RENDER_PREEVAL_FRAMEINFO"), TEXT("NOTIFY_POST_SCENE_RESET"), TEXT("NOTIFY_ANIM_LAYERS_ENABLED"), TEXT("NOTIFY_ANIM_LAYERS_DISABLED"), TEXT("NOTIFY_ACTION_ITEM_PRE_START_OVERRIDE"), TEXT("NOTIFY_ACTION_ITEM_POST_START_OVERRIDE"), TEXT("NOTIFY_ACTION_ITEM_PRE_END_OVERRIDE"), TEXT("NOTIFY_ACTION_ITEM_POST_END_OVERRIDE"), TEXT("NOTIFY_PRE_NODE_GENERAL_PROP_CHANGED"), TEXT("NOTIFY_POST_NODE_GENERAL_PROP_CHANGED"), TEXT("NOTIFY_PRE_NODE_GI_PROP_CHANGED"), TEXT("NOTIFY_POST_NODE_GI_PROP_CHANGED"), TEXT("NOTIFY_PRE_NODE_MENTALRAY_PROP_CHANGED"), TEXT("NOTIFY_POST_NODE_MENTALRAY_PROP_CHANGED"), TEXT("NOTIFY_PRE_NODE_BONE_PROP_CHANGED"), TEXT("NOTIFY_POST_NODE_BONE_PROP_CHANGED"), TEXT("NOTIFY_PRE_NODE_USER_PROP_CHANGED"), TEXT("NOTIFY_POST_NODE_USER_PROP_CHANGED"), TEXT("NOTIFY_PRE_NODE_RENDER_PROP_CHANGED"), TEXT("NOTIFY_POST_NODE_RENDER_PROP_CHANGED"), TEXT("NOTIFY_PRE_NODE_DISPLAY_PROP_CHANGED"), TEXT("NOTIFY_POST_NODE_DISPLAY_PROP_CHANGED"), TEXT("NOTIFY_PRE_NODE_BASIC_PROP_CHANGED"), TEXT("NOTIFY_POST_NODE_BASIC_PROP_CHANGED"), TEXT("NOTIFY_SELECTION_LOCK"), TEXT("NOTIFY_SELECTION_UNLOCK"), TEXT("NOTIFY_PRE_IMAGE_VIEWER_DISPLAY"), TEXT("NOTIFY_POST_IMAGE_VIEWER_DISPLAY"), TEXT("NOTIFY_IMAGE_VIEWER_UPDATE"), TEXT("NOTIFY_CUSTOM_ATTRIBUTES_ADDED"), TEXT("NOTIFY_CUSTOM_ATTRIBUTES_REMOVED"), TEXT("NOTIFY_OS_THEME_CHANGED"), TEXT("NOTIFY_ACTIVE_VIEWPORT_CHANGED"), TEXT("NOTIFY_PRE_MAXMAINWINDOW_SHOW"), TEXT("NOTIFY_POST_MAXMAINWINDOW_SHOW"), TEXT("NOTIFY_CLASSDESC_ADDED"), TEXT("NOTIFY_OBJECT_DEFINITION_CHANGE_BEGIN"), TEXT("NOTIFY_OBJECT_DEFINITION_CHANGE_END"), TEXT("NOTIFY_MTLBASE_PARAMDLG_PRE_OPEN"), TEXT("NOTIFY_MTLBASE_PARAMDLG_POST_CLOSE"), TEXT("NOTIFY_PRE_APP_FRAME_THEME_CHANGED"), TEXT("NOTIFY_APP_FRAME_THEME_CHANGED"), TEXT("NOTIFY_PRE_VIEWPORT_DELETE"), TEXT("NOTIFY_PRE_WORKSPACE_CHANGE"), TEXT("NOTIFY_POST_WORKSPACE_CHANGE"), TEXT("NOTIFY_PRE_WORKSPACE_COLLECTION_CHANGE"), TEXT("NOTIFY_POST_WORKSPACE_COLLECTION_CHANGE"), TEXT("NOTIFY_KEYBOARD_SETTING_CHANGED"), TEXT("NOTIFY_MOUSE_SETTING_CHANGED"), TEXT("NOTIFY_TOOLBARS_PRE_SAVE"), TEXT("NOTIFY_TOOLBARS_POST_SAVE"), TEXT("NOTIFY_APP_ACTIVATED"), TEXT("NOTIFY_APP_DEACTIVATED"), TEXT("NOTIFY_CUI_MENUS_UPDATED"), TEXT("NOTIFY_CUI_MENUS_PRE_SAVE"), TEXT("NOTIFY_CUI_MENUS_POST_SAVE"), TEXT("NOTIFY_VIEWPORT_SAFEFRAME_TOGGLE"), TEXT("NOTIFY_PLUGINS_PRE_SHUTDOWN"), TEXT("NOTIFY_PLUGINS_PRE_UNLOAD"), TEXT("NOTIFY_CUI_MENUS_POST_LOAD"), TEXT("NOTIFY_LAYER_PARENT_CHANGED"), TEXT("NOTIFY_ACTION_ITEM_EXECUTION_STARTED"), TEXT("NOTIFY_ACTION_ITEM_EXECUTION_ENDED"), TEXT("NOTIFY_INTERACTIVE_PLUGIN_INSTANCE_CREATION_STARTED"), TEXT("NOTIFY_INTERACTIVE_PLUGIN_INSTANCE_CREATION_ENDED"), TEXT("NOTIFY_FILE_POST_MERGE2"), TEXT("NOTIFY_POST_NODE_SELECT_OPERATION"), TEXT("NOTIFY_PRE_VIEWPORT_TOOLTIP"), TEXT("NOTIFY_WELCOMESCREEN_DONE"), TEXT("NOTIFY_PLAYBACK_START"), TEXT("NOTIFY_PLAYBACK_END"), TEXT("NOTIFY_SCENE_EXPLORER_NEEDS_UPDATE"), TEXT("NOTIFY_FILE_POST_OPEN_PROCESS_FINALIZED"), TEXT("NOTIFY_FILE_POST_MERGE_PROCESS_FINALIZED"), TEXT("NOTIFY_PRE_PROJECT_FOLDER_CHANGE"), TEXT("NOTIFY_POST_PROJECT_FOLDER_CHANGE"), TEXT("NOTIFY_PRE_MXS_STARTUP_SCRIPT_LOAD"), TEXT("NOTIFY_ACTIVESHADE_IN_VIEWPORT_TOGGLED"), TEXT("NOTIFY_SYSTEM_SHUTDOWN_CHECK"), TEXT("NOTIFY_SYSTEM_SHUTDOWN_CHECK_FAILED"), TEXT("NOTIFY_SYSTEM_SHUTDOWN_CHECK_PASSED"), TEXT("NOTIFY_FILE_POST_MERGE3"), TEXT("NOTIFY_ACTIVESHADE_IN_FRAMEBUFFER_TOGGLED"), TEXT("NOTIFY_PRE_ACTIVESHADE_IN_VIEWPORT_TOGGLED"), TEXT("NOTIFY_POST_ACTIVESHADE_IN_VIEWPORT_TOGGLED"), TEXT("NOTIFY_INTERNAL_USE_START") };
+#pragma warning(pop)
+
+		int i = 0;
+		for (int Code : Codes)
+		{
+			RegisterNotification(On3dsMaxNotification, this, Code);
+			NotificationCodetoString.Add(Code, Strings[i]);
+			++i;
+		}
+
+		// Setup Node Event System callback
+		// https://help.autodesk.com/view/3DSMAX/2018/ENU/?guid=__files_GUID_7C91D285_5683_4606_9F7C_B8D3A7CA508B_htm
+		GetISceneEventManager()->RegisterCallback(&NodeEventCallback);
+	}
+
+	bool UpdateScene()
+	{
+		SceneTracker.Update();
+
+		return true;
+	}
+
+	void Reset()
+	{
+		ExportedScene.Reset();
+
+		// todo: control output path from somewhere else?
+		if (!OutputPath.IsEmpty())
+		{
+			ExportedScene.SetOutputPath(*OutputPath);
+		}
+
+		SceneTracker.Reset();
+
+		if (DirectLinkImpl)
+		{
+			DirectLinkImpl.Reset();
+			DirectLinkImpl.Reset(new FDatasmithDirectLink);
+			DirectLinkImpl->InitializeForScene(ExportedScene.GetDatasmithScene());
+		}
+	}
+
+	static void On3dsMaxNotification(void* param, NotifyInfo* info)
+	{
+		FExporter* Exporter = reinterpret_cast<FExporter*>(param);
+		FString* Str = Exporter->NotificationCodetoString.Find(info->intcode);
+
+		const TCHAR* StrValue = Str ? **Str : TEXT("<unknown>");
+
+		switch (info->intcode)
+		{
+			// Skip some events to display(spamming tests)
+		case NOTIFY_VIEWPORT_CHANGE:
+		case NOTIFY_PRE_RENDERER_CHANGE:
+		case NOTIFY_POST_RENDERER_CHANGE:
+		case NOTIFY_CUSTOM_ATTRIBUTES_ADDED:
+		case NOTIFY_CUSTOM_ATTRIBUTES_REMOVED:
+		case NOTIFY_MTL_REFDELETED:
+			break;
+
+			// This one crashes when calling LogInfo
+		case NOTIFY_PLUGINS_PRE_SHUTDOWN:
+			break;
+		default:
+			// todo: heavy logging - called a lot
+			// LogDebug(FString(TEXT("Notify: ")) + StrValue);
+			;
+		};
+
+
+		switch (info->intcode)
+		{
+		case NOTIFY_NODE_POST_MTL:
+			// todo: Event - node got a new material
+			break;
+
+		case NOTIFY_SCENE_ADDED_NODE:
+		{
+			// note: INodeEventCallback::Added/Deleted is not used because there's a test case when it fails:
+			//   When a box is being created(dragging corners using mouse interface) and then cancelled during creation(RMB pressed)
+			//   INodeEventCallback::Deleted event is not fired by Max, although Added was called(along with other change events during creation)
+
+			INode* Node = reinterpret_cast<INode*>(info->callParam);
+
+			Exporter->SceneTracker.NodeObserver.AddItem(Node);
+
+			LogDebugNode(StrValue, Node);
+			Exporter->SceneTracker.NodeAdded(Node);
+
+			break;
+		}
+
+		case NOTIFY_SCENE_PRE_DELETED_NODE:
+		{
+			// note: INodeEventCallback::Deleted is not called when object creation was cancelled in the process
+
+			INode* Node = reinterpret_cast<INode*>(info->callParam);
+			LogDebugNode(StrValue, Node);
+
+			Exporter->SceneTracker.NodeDeleted(reinterpret_cast<INode*>(info->callParam));
+			break;
+		}
+
+		case NOTIFY_SYSTEM_POST_RESET:
+			Exporter->Reset();
+			Exporter->SceneTracker.ParseScene();
+			break;
+
+		case NOTIFY_FILE_POST_OPEN:
+			Exporter->Reset();
+			Exporter->SceneTracker.ParseScene();
+			break;
+
+		}
+	}
+
+	FDatasmith3dsMaxScene ExportedScene;
+	TUniquePtr<FDatasmithDirectLink> DirectLinkImpl;
+	FString OutputPath;
+
+	FSceneTracker SceneTracker;
+	FNodeEventCallback NodeEventCallback;
+
+	TMap<int, FString> NotificationCodetoString; // todo: remove, just for debug to output strings for notification codes
+};
+
+
+/************************************* MaxScript exports *********************************/
+
+static FExporter* Exporter = nullptr;
+
+Value* OnLoad_cf(Value**, int);
+Primitive OnLoad_pf(_M("Datasmith_OnLoad"), OnLoad_cf);
+
+Value* OnLoad_cf(Value **arg_list, int count)
+{
+	check_arg_count(OnLoad, 2, count);
+	Value* pEnableUI= arg_list[0];
+	Value* pEnginePath = arg_list[1];
+
+	bool bEnableUI = pEnableUI->to_bool();
+
+	const TCHAR* EnginePathUnreal = (const TCHAR*)pEnginePath->to_string();
+
+	FDatasmithExporterManager::FInitOptions Options;
+	Options.bEnableMessaging = true; // DirectLink requires the Messaging service.
+	Options.bSuppressLogs = false;   // Log are useful, don't suppress them
+	Options.bUseDatasmithExporterUI = bEnableUI;
+	Options.RemoteEngineDirPath = EnginePathUnreal;
+
+	if (!FDatasmithExporterManager::Initialize(Options))
+	{
+		return &false_value;
+	}
+
+	if (int32 ErrorCode = FDatasmithDirectLink::ValidateCommunicationSetup())
+	{
+		return &false_value;
+	}
+
+	Exporter = new FExporter;
+	Exporter->SceneTracker.ParseScene();
+
+	return bool_result(true);
+}
+
+Value* OnUnload_cf(Value**, int);
+Primitive OnUnload_pf(_M("Datasmith_OnUnload"), OnUnload_cf);
+
+Value* OnUnload_cf(Value **arg_list, int count)
+{
+	check_arg_count(OnUnload, 0, count);
+
+	delete Exporter;
+
+	FDatasmithDirectLink::Shutdown();
+	FDatasmithExporterManager::Shutdown();
+
+	return bool_result(true);
+}
+
+Value* SetOutputPath_cf(Value**, int);
+Primitive SetOutputPath_pf(_M("Datasmith_SetOutputPath"), SetOutputPath_cf);
+
+Value* SetOutputPath_cf(Value** arg_list, int count)
+{
+	check_arg_count(CreateScene, 1, count);
+	Value* pOutputPath = arg_list[0];
+
+	Exporter->SetOutputPath(pOutputPath->to_string());
+
+	return bool_result(true);
+}
+
+Value* CreateScene_cf(Value**, int);
+Primitive CreateScene_pf(_M("Datasmith_CreateScene"), CreateScene_cf);
+
+Value* CreateScene_cf(Value** arg_list, int count)
+{
+	check_arg_count(CreateScene, 1, count);
+	Value* pName = arg_list[0];
+
+	Exporter->ExportedScene.SetName(pName->to_string());
+
+	return bool_result(true);
+}
+
+Value* UpdateScene_cf(Value**, int);
+Primitive UpdateScene_pf(_M("Datasmith_UpdateScene"), UpdateScene_cf);
+
+Value* UpdateScene_cf(Value** arg_list, int count)
+{
+	check_arg_count(UpdateScene, 0, count);
+
+	if(!Exporter)
+	{
+		return bool_result(false);
+	}
+
+	bool bResult = Exporter->UpdateScene();
+	return bool_result(bResult);
+}
+
+
+Value* Export_cf(Value**, int);
+Primitive Export_pf(_M("Datasmith_Export"), Export_cf);
+
+Value* Export_cf(Value** arg_list, int count)
+{
+	check_arg_count(Export, 2, count);
+	Value* pName = arg_list[0];
+	Value* pOutputPath = arg_list[1];
+
+
+	FExporter TempExporter;
+	TempExporter.ExportedScene.SetName(pName->to_string());
+	TempExporter.SetOutputPath(pOutputPath->to_string());
+
+	bool bResult = TempExporter.Export();
+	return bool_result(bResult);
+}
+
+
+Value* Reset_cf(Value**, int);
+Primitive Reset_pf(_M("Datasmith_Reset"), Reset_cf);
+
+Value* Reset_cf(Value** arg_list, int count)
+{
+	check_arg_count(Reset, 0, count);
+
+	if (!Exporter)
+	{
+		return bool_result(false);
+	}
+
+	Exporter->Reset();
+	return bool_result(true);
+}
+
+Value* StartSceneChangeTracking_cf(Value**, int);
+Primitive StartSceneChangeTracking_pf(_M("Datasmith_StartSceneChangeTracking"), StartSceneChangeTracking_cf);
+
+Value* StartSceneChangeTracking_cf(Value** arg_list, int count)
+{
+	check_arg_count(StartSceneChangeTracking, 0, count);
+
+	Exporter->StartSceneChangeTracking();
+
+	return bool_result(true);
+}
+
+Value* DirectLinkInitializeForScene_cf(Value** arg_list, int count)
+{
+	check_arg_count(DirectLinkInitializeForScene, 0, count);
+
+	Exporter->DirectLinkImpl.Reset(new FDatasmithDirectLink);
+	Exporter->DirectLinkImpl->InitializeForScene(Exporter->ExportedScene.GetDatasmithScene());
+
+	return bool_result(true);
+}
+Primitive DirectLinkInitializeForScene_pf(_M("Datasmith_DirectLinkInitializeForScene"), DirectLinkInitializeForScene_cf);
+
+
+Value* DirectLinkUpdateScene_cf(Value** arg_list, int count)
+{
+	check_arg_count(DirectLinkUpdateScene, 0, count);
+
+	Exporter->DirectLinkImpl->UpdateScene(Exporter->ExportedScene.GetDatasmithScene());
+
+	return bool_result(true);
+}
+Primitive DirectLinkUpdateScene_pf(_M("Datasmith_DirectLinkUpdateScene"), DirectLinkUpdateScene_cf);
+
+Value* OpenDirectlinkUi_cf(Value** arg_list, int count) 
+{
+	check_arg_count(OpenDirectlinkUi, 0, count);
+	if (IDatasmithExporterUIModule* Module = IDatasmithExporterUIModule::Get())
+	{
+		if (IDirectLinkUI* UI = Module->GetDirectLinkExporterUI())
+		{
+			UI->OpenDirectLinkStreamWindow();
+			return &true_value;
+		}
+	}
+	return &false_value;
+}
+Primitive OpenDirectlinkUi_pf(_M("Datasmith_OpenDirectlinkUi"), OpenDirectlinkUi_cf);
+
+
+Value* GetDirectlinkCacheDirectory_cf(Value** arg_list, int count)
+{
+	check_arg_count(GetDirectlinkCacheDirectory, 0, count);
+	if (IDatasmithExporterUIModule* Module = IDatasmithExporterUIModule::Get())
+	{
+		if (IDirectLinkUI* UI = Module->GetDirectLinkExporterUI())
+		{
+			return new String(UI->GetDirectLinkCacheDirectory());
+		}
+	}
+
+	return &undefined;
+}
+
+Primitive GetDirectlinkCacheDirectory_pf(_M("Datasmith_GetDirectlinkCacheDirectory"), GetDirectlinkCacheDirectory_cf);
+
+
+Value* LogFlush_cf(Value** arg_list, int count)
+{
+	LogFlush();
+	return &undefined;
+}
+
+Primitive LogFlush_pf(_M("Datasmith_LogFlush"), LogFlush_cf);
+
+
+Value* Crash_cf(Value** arg_list, int count)
+{
+	volatile int* P;
+	P = nullptr;
+
+	*P = 666;
+	return &undefined;
+}
+
+Primitive Crash_pf(_M("Datasmith_Crash"), Crash_cf);
+
+
+#include "Windows/HideWindowsPlatformTypes.h"
+
+
+#endif // NEW_DIRECTLINK_PLUGIN
diff --git a/Engine/Source/Programs/Enterprise/Datasmith/DatasmithMaxExporter/Private/DatasmithMaxWriter.cpp b/Engine/Source/Programs/Enterprise/Datasmith/DatasmithMaxExporter/Private/DatasmithMaxWriter.cpp
--- a/Engine/Source/Programs/Enterprise/Datasmith/DatasmithMaxExporter/Private/DatasmithMaxWriter.cpp
+++ b/Engine/Source/Programs/Enterprise/Datasmith/DatasmithMaxExporter/Private/DatasmithMaxWriter.cpp
@@ -481,6 +481,8 @@ bool FDatasmithMaxMatHelper::IsSRGB(Bitmap& InBitmap)
 	return true;
 }
 
+bool FDatasmithMaxMatExport::bForceReexport = false;
+
 void FDatasmithMaxMatExport::GetXMLTexture(TSharedRef< IDatasmithScene > DatasmithScene, Texmap* InTexMap, const TCHAR* AssetsPath)
 {
 	switch (FDatasmithMaxMatHelper::GetTextureClass(InTexMap))
@@ -559,13 +561,29 @@ bool FDatasmithMaxMatExport::UseFirstSubMapOnly(EDSMaterialType MaterialType, Mt
 
 TSharedPtr< IDatasmithBaseMaterialElement > FDatasmithMaxMatExport::ExportUniqueMaterial(TSharedRef< IDatasmithScene > DatasmithScene, Mtl* Material, const TCHAR* AssetsPath)
 {
-	// Names should be unique prior to this
-	FString MaterialName(FDatasmithUtils::SanitizeObjectName(Material->GetName().data()));
-	for (int i = 0; i < DatasmithScene->GetMaterialsCount(); i++)
+	if (!bForceReexport)
+	{
+		// Names should be unique prior to this
+		FString MaterialName(FDatasmithUtils::SanitizeObjectName(Material->GetName().data()));
+		for (int i = 0; i < DatasmithScene->GetMaterialsCount(); i++)
+		{
+			if (FString(DatasmithScene->GetMaterial(i)->GetName()) == MaterialName)
+			{
+				return DatasmithScene->GetMaterial(i);
+			}
+		}
+	}
+	else
 	{
-		if (FString(DatasmithScene->GetMaterial(i)->GetName()) == MaterialName)
+		// Names should be unique prior to this
+		FString MaterialName(FDatasmithUtils::SanitizeObjectName(Material->GetName().data()));
+		for (int i = 0; i < DatasmithScene->GetMaterialsCount(); i++)
 		{
-			return DatasmithScene->GetMaterial(i);
+			if (FString(DatasmithScene->GetMaterial(i)->GetName()) == MaterialName)
+			{
+				DatasmithScene->RemoveMaterial(DatasmithScene->GetMaterial(i));
+				break;
+			}
 		}
 	}
 
diff --git a/Engine/Source/Programs/Enterprise/Datasmith/DatasmithMaxExporter/Public/DatasmithMaxWriter.h b/Engine/Source/Programs/Enterprise/Datasmith/DatasmithMaxExporter/Public/DatasmithMaxWriter.h
--- a/Engine/Source/Programs/Enterprise/Datasmith/DatasmithMaxExporter/Public/DatasmithMaxWriter.h
+++ b/Engine/Source/Programs/Enterprise/Datasmith/DatasmithMaxExporter/Public/DatasmithMaxWriter.h
@@ -110,6 +110,8 @@ public:
 	static TSharedPtr< IDatasmithBaseMaterialElement > ExportUniqueMaterial(TSharedRef< IDatasmithScene > DatasmithScene, Mtl* Material, const TCHAR* AssetsPath);
 	static void WriteXMLMaterial(TSharedRef< IDatasmithScene > DatasmithScene, TSharedPtr< IDatasmithMaterialElement > MaterialElement, Mtl* Material);
 
+	static bool bForceReexport; // todo: ExportUniqueMaterial calls go deep into composite materials export - need to fix this for DirectLink Sync
+
 private:
 	static bool UseFirstSubMapOnly(EDSMaterialType MaterialType, Mtl* Material);
 };
