Fixed build breakage introduced by CL 17614269

diff --git a/Engine/Source/Programs/Enterprise/Datasmith/DatasmithARCHICADExporter/DatsmithUE4ArchiCAD/DatasmithUE4ArchiCAD.Target.cs b/Engine/Source/Programs/Enterprise/Datasmith/DatasmithARCHICADExporter/DatsmithUE4ArchiCAD/DatasmithUE4ArchiCAD.Target.cs
--- a/Engine/Source/Programs/Enterprise/Datasmith/DatasmithARCHICADExporter/DatsmithUE4ArchiCAD/DatasmithUE4ArchiCAD.Target.cs
+++ b/Engine/Source/Programs/Enterprise/Datasmith/DatasmithARCHICADExporter/DatsmithUE4ArchiCAD/DatasmithUE4ArchiCAD.Target.cs
@@ -26,7 +26,7 @@ public class DatasmithUE4ArchiCADTarget : TargetRules
 
 		bBuildDeveloperTools = false;
 		bUseMallocProfiler = false;
-		bBuildWithEditorOnlyData = true;
+		bBuildWithEditorOnlyData = false;
 		bCompileAgainstEngine = false;
 		bCompileAgainstCoreUObject = true;
 		bCompileICU = false;
diff --git a/Engine/Source/Programs/Enterprise/Datasmith/DatasmithARCHICADExporter/UBTModules/DatasmithArchiCAD.Target.cs b/Engine/Source/Programs/Enterprise/Datasmith/DatasmithARCHICADExporter/UBTModules/DatasmithArchiCAD.Target.cs
--- a/Engine/Source/Programs/Enterprise/Datasmith/DatasmithARCHICADExporter/UBTModules/DatasmithArchiCAD.Target.cs
+++ b/Engine/Source/Programs/Enterprise/Datasmith/DatasmithARCHICADExporter/UBTModules/DatasmithArchiCAD.Target.cs
@@ -21,7 +21,7 @@ public class DatasmithArchiCADTarget : TargetRules
 
 		bBuildDeveloperTools = false;
 		bUseMallocProfiler = false;
-		bBuildWithEditorOnlyData = true;
+		bBuildWithEditorOnlyData = false;
 		bCompileAgainstEngine = false;
 		bCompileAgainstCoreUObject = true;
 		bCompileICU = false;
diff --git a/Engine/Source/Programs/Enterprise/Datasmith/DatasmithCADWorker/DatasmithCADWorker.Target.cs b/Engine/Source/Programs/Enterprise/Datasmith/DatasmithCADWorker/DatasmithCADWorker.Target.cs
--- a/Engine/Source/Programs/Enterprise/Datasmith/DatasmithCADWorker/DatasmithCADWorker.Target.cs
+++ b/Engine/Source/Programs/Enterprise/Datasmith/DatasmithCADWorker/DatasmithCADWorker.Target.cs
@@ -21,7 +21,7 @@ public class DatasmithCADWorkerTarget : TargetRules
 		bUseMallocProfiler = false;
 
 		// Editor-only data, however, is needed
-		bBuildWithEditorOnlyData = true;
+		bBuildWithEditorOnlyData = false;
 
 		// Currently this app is not linking against the engine, so we'll compile out references from Core to the rest of the engine
 		bCompileAgainstEngine = false;
diff --git a/Engine/Source/Programs/Enterprise/Datasmith/DatasmithFacadeCSharp/DatasmithFacadeCSharp.Target.cs b/Engine/Source/Programs/Enterprise/Datasmith/DatasmithFacadeCSharp/DatasmithFacadeCSharp.Target.cs
--- a/Engine/Source/Programs/Enterprise/Datasmith/DatasmithFacadeCSharp/DatasmithFacadeCSharp.Target.cs
+++ b/Engine/Source/Programs/Enterprise/Datasmith/DatasmithFacadeCSharp/DatasmithFacadeCSharp.Target.cs
@@ -22,7 +22,7 @@ public class DatasmithFacadeCSharpTarget : TargetRules
 
 		bBuildDeveloperTools = false;
 		bUseMallocProfiler = false;
-		bBuildWithEditorOnlyData = true;
+		bBuildWithEditorOnlyData = false;
 		bCompileAgainstEngine = false;
 		bCompileAgainstCoreUObject = true;
 		bCompileICU = false;
diff --git a/Engine/Source/Programs/Enterprise/Datasmith/DatasmithMaxExporter/DatasmithMax2017.Target.cs b/Engine/Source/Programs/Enterprise/Datasmith/DatasmithMaxExporter/DatasmithMax2017.Target.cs
--- a/Engine/Source/Programs/Enterprise/Datasmith/DatasmithMaxExporter/DatasmithMax2017.Target.cs
+++ b/Engine/Source/Programs/Enterprise/Datasmith/DatasmithMaxExporter/DatasmithMax2017.Target.cs
@@ -38,7 +38,7 @@ public abstract class DatasmithMaxBaseTarget : TargetRules
 
 		bBuildDeveloperTools = false;
 		bUseMallocProfiler = false;
-		bBuildWithEditorOnlyData = true;
+		bBuildWithEditorOnlyData = false;
 		bCompileAgainstEngine = false;
 		bCompileAgainstCoreUObject = true;
 		bCompileICU = false;
diff --git a/Engine/Source/Programs/Enterprise/Datasmith/DatasmithNavisworksExporter/UBTModules/DatasmithNavisworks2020.Target.cs b/Engine/Source/Programs/Enterprise/Datasmith/DatasmithNavisworksExporter/UBTModules/DatasmithNavisworks2020.Target.cs
--- a/Engine/Source/Programs/Enterprise/Datasmith/DatasmithNavisworksExporter/UBTModules/DatasmithNavisworks2020.Target.cs
+++ b/Engine/Source/Programs/Enterprise/Datasmith/DatasmithNavisworksExporter/UBTModules/DatasmithNavisworks2020.Target.cs
@@ -23,7 +23,7 @@ public abstract class DatasmithNavisworksBaseTarget : TargetRules
 
 		bBuildDeveloperTools = false;
 		bUseMallocProfiler = false;
-		bBuildWithEditorOnlyData = true;
+		bBuildWithEditorOnlyData = false;
 		bCompileAgainstEngine = false;
 		bCompileAgainstCoreUObject = false;
 		bCompileICU = false;
diff --git a/Engine/Source/Programs/Enterprise/Datasmith/DatasmithRevitExporter/DatasmithRevit2018.Target.cs b/Engine/Source/Programs/Enterprise/Datasmith/DatasmithRevitExporter/DatasmithRevit2018.Target.cs
--- a/Engine/Source/Programs/Enterprise/Datasmith/DatasmithRevitExporter/DatasmithRevit2018.Target.cs
+++ b/Engine/Source/Programs/Enterprise/Datasmith/DatasmithRevitExporter/DatasmithRevit2018.Target.cs
@@ -25,7 +25,7 @@ public abstract class DatasmithRevitBaseTarget : TargetRules
 
 		bBuildDeveloperTools = false;
 		bUseMallocProfiler = false;
-		bBuildWithEditorOnlyData = true;
+		bBuildWithEditorOnlyData = false;
 		bCompileAgainstEngine = false;
 		bCompileAgainstCoreUObject = true;
 		bCompileICU = false;
diff --git a/Engine/Source/Programs/Enterprise/Datasmith/DatasmithRhinoExporter/DatasmithRhino6.Target.cs b/Engine/Source/Programs/Enterprise/Datasmith/DatasmithRhinoExporter/DatasmithRhino6.Target.cs
--- a/Engine/Source/Programs/Enterprise/Datasmith/DatasmithRhinoExporter/DatasmithRhino6.Target.cs
+++ b/Engine/Source/Programs/Enterprise/Datasmith/DatasmithRhinoExporter/DatasmithRhino6.Target.cs
@@ -34,7 +34,7 @@ public abstract class DatasmithRhinoBaseTarget : TargetRules
 
 		bBuildDeveloperTools = false;
 		bUseMallocProfiler = false;
-		bBuildWithEditorOnlyData = true;
+		bBuildWithEditorOnlyData = false;
 		bCompileAgainstEngine = false;
 		bCompileAgainstCoreUObject = true;
 		bCompileICU = false;
diff --git a/Engine/Source/Programs/Enterprise/Datasmith/DatasmithSDK/DatasmithSDK.Target.cs b/Engine/Source/Programs/Enterprise/Datasmith/DatasmithSDK/DatasmithSDK.Target.cs
--- a/Engine/Source/Programs/Enterprise/Datasmith/DatasmithSDK/DatasmithSDK.Target.cs
+++ b/Engine/Source/Programs/Enterprise/Datasmith/DatasmithSDK/DatasmithSDK.Target.cs
@@ -21,7 +21,7 @@ public class DatasmithSDKTarget : TargetRules
 
 		bBuildDeveloperTools = false;
 		bUseMallocProfiler = false;
-		bBuildWithEditorOnlyData = true;
+		bBuildWithEditorOnlyData = false;
 		bCompileAgainstEngine = false;
 		bCompileAgainstCoreUObject = true;
 		bCompileICU = false;
diff --git a/Engine/Source/Programs/Enterprise/Datasmith/DatasmithSketchUpRubyExporter/DatasmithSketchUpRuby2020.Target.cs b/Engine/Source/Programs/Enterprise/Datasmith/DatasmithSketchUpRubyExporter/DatasmithSketchUpRuby2020.Target.cs
--- a/Engine/Source/Programs/Enterprise/Datasmith/DatasmithSketchUpRubyExporter/DatasmithSketchUpRuby2020.Target.cs
+++ b/Engine/Source/Programs/Enterprise/Datasmith/DatasmithSketchUpRubyExporter/DatasmithSketchUpRuby2020.Target.cs
@@ -19,7 +19,7 @@ public abstract class DatasmithSketchUpRubyBaseTarget : TargetRules
 
 		bBuildDeveloperTools = false;
 		bUseMallocProfiler = false;
-		bBuildWithEditorOnlyData = true;
+		bBuildWithEditorOnlyData = false;
 		bCompileAgainstEngine = false;
 		bCompileAgainstCoreUObject = true;
 		bCompileICU = false;
diff --git a/Engine/Source/Programs/Enterprise/Datasmith/DatasmithSolidworksExporter/DatasmithSolidworks2021.Target.cs b/Engine/Source/Programs/Enterprise/Datasmith/DatasmithSolidworksExporter/DatasmithSolidworks2021.Target.cs
--- a/Engine/Source/Programs/Enterprise/Datasmith/DatasmithSolidworksExporter/DatasmithSolidworks2021.Target.cs
+++ b/Engine/Source/Programs/Enterprise/Datasmith/DatasmithSolidworksExporter/DatasmithSolidworks2021.Target.cs
@@ -54,7 +54,7 @@ public abstract class DatasmithSolidworksBaseTarget : TargetRules
 		LinkType = TargetLinkType.Monolithic;
 		bBuildDeveloperTools = false;
 		bUseMallocProfiler = false;
-		bBuildWithEditorOnlyData = true;
+		bBuildWithEditorOnlyData = false;
 		bCompileAgainstEngine = false;
 		bCompileAgainstCoreUObject = false;
 		bCompileICU = false;
