3ds max: Instances tracking

diff --git a/Engine/Source/Programs/Enterprise/Datasmith/DatasmithMaxExporter/Private/DatasmithMaxExporterWithDirectLink.cpp b/Engine/Source/Programs/Enterprise/Datasmith/DatasmithMaxExporter/Private/DatasmithMaxExporterWithDirectLink.cpp
--- a/Engine/Source/Programs/Enterprise/Datasmith/DatasmithMaxExporter/Private/DatasmithMaxExporterWithDirectLink.cpp
+++ b/Engine/Source/Programs/Enterprise/Datasmith/DatasmithMaxExporter/Private/DatasmithMaxExporterWithDirectLink.cpp
@@ -169,11 +169,18 @@ void LogInfo(const FString& Msg)
 	LogInfo(*Msg);
 }
 
+// Log all messages
+// #define LOG_DEBUG_HEAVY_ENABLE 1
+
+#ifdef LOG_DEBUG_HEAVY_ENABLE
+	#define LOG_DEBUG_HEAVY(message) LogDebug(message)
+#else
+	#define LOG_DEBUG_HEAVY(message) 
+#endif
+
 void LogDebugNode(const TCHAR* Name, INode* Node)
 {
-	// todo: LogDebugNode is called a lot with many nodes - add extra level of debugging
-	return;
-
+#ifdef LOG_DEBUG_HEAVY_ENABLE
 
 	LogDebug(FString::Printf(TEXT("%s: %u %s(%d) - %s")
 		, Name
@@ -192,13 +199,12 @@ void LogDebugNode(const TCHAR* Name, INode* Node)
 			LogDebug(FString::Printf(TEXT("    Class_ID: 0x%lx, 0x%lx "), ClassId.PartA(), ClassId.PartB()));
 		}
 	}
+#endif
 }
 
 void LogNodeEvent(const MCHAR* Name, INodeEventCallback::NodeKeyTab& nodes)
 {
-	// todo: heavy logging - called a lot
-	return;
-
+#ifdef LOG_DEBUG_HEAVY_ENABLE
 	LogDebug(FString::Printf(TEXT("NodeEventCallback:%s"), Name));
 	for (int NodeIndex = 0; NodeIndex < nodes.Count(); ++NodeIndex)
 	{
@@ -215,6 +221,7 @@ void LogNodeEvent(const MCHAR* Name, INodeEventCallback::NodeKeyTab& nodes)
 			LogDebug(FString::Printf(TEXT("    %u <null>"), NodeKey));
 		}
 	}
+#endif
 }
 
 
@@ -390,10 +397,7 @@ public:
 		// todo: remove material handling???
 		ensure(ReferencedItemToIndex.Contains(TargetHandle));
 
-
-		// todo: heavy logging - called a lot
-		// LogDebug(FString::Printf(TEXT("FNodeObserver::NotifyRefChanged: %s: %x"), dynamic_cast<INode*>(TargetHandle)->GetName(), Message));
-
+		LOG_DEBUG_HEAVY(FString::Printf(TEXT("FNodeObserver::NotifyRefChanged: %s: %x"), dynamic_cast<INode*>(TargetHandle)->GetName(), Message)); // heavy logging - called a lot
 		return REF_SUCCEED;
 	}
 
@@ -425,15 +429,13 @@ public:
 	RefTargetHandle GetReference(int ReferenceIndex) override
 	{
 		RefTargetHandle TargetHandle = IndexToReferencedItem[ReferenceIndex];
-		// todo: heavy logging - called a lot
-		// LogDebug(FString::Printf(TEXT("FNodeObserver::GetReference: %d, %s"), ReferenceIndex, TargetHandle ? dynamic_cast<INode*>(TargetHandle)->GetName() : TEXT("<null>")));
+		LOG_DEBUG_HEAVY(FString::Printf(TEXT("FNodeObserver::GetReference: %d, %s"), ReferenceIndex, TargetHandle ? dynamic_cast<INode*>(TargetHandle)->GetName() : TEXT("<null>")));
 		return TargetHandle;
 	}
 
 	void SetReference(int ReferenceIndex, RefTargetHandle TargetHandle) override
 	{
-		// todo: heavy logging - called a lot
-		// LogDebug(FString::Printf(TEXT("FNodeObserver::SetReference: %d, %s"), ReferenceIndex, TargetHandle?dynamic_cast<INode*>(TargetHandle)->GetName():TEXT("<null>")));
+		LOG_DEBUG_HEAVY(FString::Printf(TEXT("FNodeObserver::SetReference: %d, %s"), ReferenceIndex, TargetHandle?dynamic_cast<INode*>(TargetHandle)->GetName():TEXT("<null>")));
 
 		// todo: investigate why NodeEventNamespace::GetNodeByKey may stil return NULL
 		// testcase - add XRef Material - this will immediately have this 
@@ -483,6 +485,7 @@ public:
 	}
 
 	INode* Node = nullptr;
+	AnimHandle InstanceHandle = 0; // todo: rename - this is handle for object this node is instance of
 	bool bInvalidated = true;
 
 	TSharedPtr<IDatasmithActorElement> DatasmithActorElement;
@@ -490,7 +493,11 @@ public:
 	class FMaterialTracker* MaterialTracker = nullptr;
 
 	TSharedPtr<IDatasmithMeshActorElement> DatasmithMeshActor;
-	TSharedPtr<IDatasmithMeshElement> DatasmithMeshElement; //todo: this should go away to be separate from node(especially when instancing)
+
+	bool IsInstance()
+	{
+		return InstanceHandle != 0;
+	}
 };
 
 
@@ -934,6 +941,20 @@ public:
 };
 
 
+// Every node which is resolved to the same object is considered an instance
+// This class holds all this nodes and the object they resolve to
+struct FInstances
+{
+	Object* EvaluatedObj = nullptr;
+
+	TSet<class FNodeTracker*> NodeTrackers;
+
+	// Mesh conversion results
+	TSet<uint16> SupportedChannels;
+	TSharedPtr<IDatasmithMeshElement> DatasmithMeshElement;
+};
+
+
 // Holds states of entities for syncronization and handles change events
 class FSceneTracker
 {
@@ -1040,9 +1061,12 @@ public:
 		MaterialObserver.Reset();
 		NodeTrackers.Reset();
 		InvalidatedNodeTrackers.Reset();
+		InvalidatedInstances.Reset();
 		MaterialTrackers.Reset();
 		InvalidatedMaterialTrackers.Reset();
 		MaterialsTracker.Reset();
+
+		InstancesForAnimHandle.Reset();
 	}
 
 	// Applies all recorded changes to Datasmith scene
@@ -1055,6 +1079,12 @@ public:
 		}
 		InvalidatedNodeTrackers.Reset();
 
+		for (FInstances* Instances : InvalidatedInstances)
+		{
+			UpdateInstances(*Instances);
+		}
+		InvalidatedInstances.Reset();
+
 		TSet<Mtl*> ActualMaterialToUpdate;
 		TSet<Texmap*> ActualTexmapsToUpdate;
 		for (FMaterialTracker* MaterialTracker : InvalidatedMaterialTrackers)
@@ -1204,14 +1234,24 @@ public:
 
 	void ConvertNodeGeometry(FNodeTracker& NodeTracker)
 	{
-		if (NodeTracker.DatasmithMeshElement)
+		if (NodeTracker.IsInstance())
 		{
-			ExportedScene.DatasmithSceneRef->RemoveMesh(NodeTracker.DatasmithMeshElement);
+			if (TUniquePtr<FInstances>* InstancesPtr = InstancesForAnimHandle.Find(NodeTracker.InstanceHandle))
+			{
+				FInstances& Instances = **InstancesPtr;
+				Instances.NodeTrackers.Remove(&NodeTracker);
+				if (!Instances.NodeTrackers.Num())
+				{
+					ExportedScene.DatasmithSceneRef->RemoveMesh(Instances.DatasmithMeshElement);
+					InstancesForAnimHandle.Remove(NodeTracker.InstanceHandle);
+					InvalidatedInstances.Remove(&Instances);
+				}
+			}
+			// ExportedScene.DatasmithSceneRef->RemoveMesh(NodeTracker.DatasmithMeshElement);
 			// todo: mesh is removed from the scene but not deallocated for reuse
 			// OR it will stay taking up memory if not reused(e.g. node has no valid geometry now)
 		}
 
-
 		ObjectState ObjState = NodeTracker.Node->EvalWorldState(0);
 		Object* Obj = ObjState.obj;
 
@@ -1224,10 +1264,8 @@ public:
 		{
 		case GEOMOBJECT_CLASS_ID:
 		{
-			// todo: reuse meshe element(make sure to reset all)
-			TSharedPtr<IDatasmithMeshElement> DatasmithMeshElement;
-			ConvertGeomObjToDatasmithMesh(NodeTracker, Obj, DatasmithMeshElement);
-			AssignDatasmithMeshToNodeTracker(NodeTracker, DatasmithMeshElement);
+			// todo: reuse mesh element(make sure to reset all)
+			ConvertGeomObjToDatasmithMesh(NodeTracker, Obj);
 			break;
 		}
 		// todo: other object types besides geometry
@@ -1235,9 +1273,44 @@ public:
 		}
 	}
 
-	void AssignDatasmithMeshToNodeTracker(FNodeTracker& NodeTracker, const TSharedPtr<IDatasmithMeshElement>& DatasmithMeshElement)
+	void InvalidateInstances(FInstances& Instances)
+	{
+		for (FNodeTracker* NodeTracker : Instances.NodeTrackers)
+		{
+			InvalidatedNodeTrackers.Add(NodeTracker);
+		}
+	}
+
+	void UpdateInstances(FInstances& Instances)
+	{
+		 if (!Instances.NodeTrackers.IsEmpty())
+		 {
+			 // todo: determine before converting geometry if:
+			 // - there's multimat among instances
+			 // - there's one instance only(can just assing material to mesh instead of actor mesh overrides)
+
+			 bool bConverted = false; // Use first node to extract information from evaluated object(e.g. GetRendedrMesh needs it)
+
+			 bool bAssignToStaticMesh = true; // assign to static mesh for the first instance
+			 for(FNodeTracker* NodeTrackerPtr: Instances.NodeTrackers)
+			 {
+				 FNodeTracker& NodeTracker = *NodeTrackerPtr;
+				 if (!bConverted)
+				 {
+					 ConvertInstancesGeometry(Instances, NodeTracker);
+					 bConverted = true;
+				 }
+
+
+				 AssignDatasmithMeshToNodeTracker(NodeTracker, Instances, bAssignToStaticMesh);
+				 bAssignToStaticMesh = false;
+			 }
+		 }
+	}
+
+	void AssignDatasmithMeshToNodeTracker(FNodeTracker& NodeTracker, FInstances& Instances, bool bAssignToStaticMesh)
 	{
-		NodeTracker.DatasmithMeshElement = DatasmithMeshElement;
+		const TSharedPtr<IDatasmithMeshElement>& DatasmithMeshElement = Instances.DatasmithMeshElement;
 
 		if (DatasmithMeshElement)
 		{
@@ -1253,6 +1326,80 @@ public:
 			}
 
 			NodeTracker.DatasmithMeshActor->SetStaticMeshPathName(DatasmithMeshElement->GetName());
+
+			{
+				// todo: might assign one instance's material to static mesh when there are other instances - this way static mesh would be 
+
+				if (Mtl* Material = NodeTracker.Node->GetMtl())
+				{
+					if (!NodeTracker.MaterialTracker || NodeTracker.MaterialTracker->Material != Material)
+					{
+						// Release old material
+						if (NodeTracker.MaterialTracker)
+						{
+							// Release material assignment
+							MaterialsAssignedToNodes[NodeTracker.MaterialTracker].Remove(&NodeTracker);
+
+							// Clean tracker if it's not used aby any node
+							if (!MaterialsAssignedToNodes[NodeTracker.MaterialTracker].Num())
+							{
+								MaterialsTracker.UnregisterMaterialTracker(*NodeTracker.MaterialTracker);
+								MaterialTrackers.Remove(NodeTracker.MaterialTracker->Material);
+							}
+						}
+
+						if (!MaterialTrackers.Contains(Material))
+						{
+							// Track material if not yet
+							FMaterialTrackerHandle& MaterialTrackerHandle = MaterialTrackers.Emplace(Material, Material);
+							InvalidatedMaterialTrackers.Add(MaterialTrackerHandle.GetMaterialTracker());
+						}
+
+						// Store new 
+						NodeTracker.MaterialTracker = MaterialTrackers.Find(Material)->GetMaterialTracker();
+						MaterialsAssignedToNodes.FindOrAdd(NodeTracker.MaterialTracker).Add(&NodeTracker);
+					}
+
+					// Clear previous material overrides
+					NodeTracker.DatasmithMeshActor->ResetMaterialOverrides(); 
+
+					// Assign materials
+					if (bAssignToStaticMesh)
+					{
+						// todo: move to header
+						void AssignMeshMaterials(TSharedPtr<IDatasmithMeshElement>&MeshElement, Mtl * Material, const TSet<uint16>&SupportedChannels);
+
+						AssignMeshMaterials(Instances.DatasmithMeshElement, Material, Instances.SupportedChannels);
+					}
+					else // Assign material overrides to meshactor
+					{
+						
+						TSharedRef<IDatasmithMeshActorElement> DatasmithMeshActor = NodeTracker.DatasmithMeshActor.ToSharedRef();
+						FDatasmithMaxSceneExporter::ParseMaterialForMeshActor(NodeTracker.MaterialTracker->Material, DatasmithMeshActor, Instances.SupportedChannels, NodeTracker.DatasmithMeshActor->GetTranslation());
+					}
+				}
+				else
+				{
+					// Release old material
+					if (NodeTracker.MaterialTracker)
+					{
+						MaterialsAssignedToNodes[NodeTracker.MaterialTracker].Remove(&NodeTracker);
+						if (!MaterialsAssignedToNodes[NodeTracker.MaterialTracker].Num())
+						{
+							MaterialsTracker.UnregisterMaterialTracker(*NodeTracker.MaterialTracker);
+							MaterialTrackers.Remove(Material);
+						}
+					}
+					NodeTracker.MaterialTracker = nullptr;
+					NodeTracker.DatasmithMeshActor->ResetMaterialOverrides();
+				}
+			}
+
+			// todo: test mesh becoming empty/invalid/not created - what happens?
+			// todo: test multimaterial changes
+			// todo: check other material permutations
+
+
 		}
 		else
 		{
@@ -1266,12 +1413,10 @@ public:
 		}
 	}
 
-	// todo: change this function to create mesh element and schedule ExportToUObject threading
-	bool ConvertGeomObjToDatasmithMesh(FNodeTracker& NodeTracker, Object* Obj, TSharedPtr<IDatasmithMeshElement>& DatasmithMeshElement)
+	bool ConvertInstancesGeometry(FInstances& Instances, FNodeTracker& NodeTracker)
 	{
-		bool bResult = false;
-
 		INode* Node = NodeTracker.Node;
+		Object* Obj = Instances.EvaluatedObj;
 
 		// todo: baseline exporter uses GetBaseObject which takes result of EvalWorldState
 		// and searched down DerivedObject pipeline(by taking GetObjRef) 
@@ -1284,6 +1429,8 @@ public:
 		TimeValue Time = GetCOREInterface()->GetTime();
 		Mesh* RenderMesh = GeomObj->GetRenderMesh(Time, Node, View, bNeedsDelete);
 
+		bool bResult = false;
+
 		if (!RenderMesh)
 		{
 			return bResult;
@@ -1306,85 +1453,65 @@ public:
 			if (CachedMesh.getNumFaces() > 0)
 			{
 				FDatasmithMesh DatasmithMesh;
-				TSet<uint16> SupportedChannels;
+
 				const TCHAR* MeshName = (const TCHAR*)Node->GetName();
 				// todo: pivot
-				FillDatasmithMeshFromMaxMesh(DatasmithMesh, CachedMesh, Node, false, SupportedChannels, MeshName, FTransform::Identity);
+				FillDatasmithMeshFromMaxMesh(DatasmithMesh, CachedMesh, Node, false, Instances.SupportedChannels, MeshName, FTransform::Identity);
 
 				FDatasmithMeshExporter DatasmithMeshExporter;
 
-				DatasmithMeshElement = FDatasmithSceneFactory::CreateMesh(TEXT(""));
-				DatasmithMeshElement->SetName(MeshName);
-				DatasmithMeshElement->SetLabel(MeshName);
-
-				ExportedScene.GetDatasmithScene()->AddMesh(DatasmithMeshElement);
-
-				// todo: remove material if null
-				if (Mtl* Material = Node->GetMtl())
+				if (Instances.DatasmithMeshElement)
 				{
-					// todo: move to header
-					void AssignMeshMaterials(TSharedPtr<IDatasmithMeshElement>&MeshElement, Mtl * Material, const TSet<uint16>&SupportedChannels);
+					// todo: potential mesh reuse - when DatasmithMeshElement allows to reset materials(as well as other params)
+					ExportedScene.GetDatasmithScene()->RemoveMesh(Instances.DatasmithMeshElement);
+				}
 
-					AssignMeshMaterials(DatasmithMeshElement, Material, SupportedChannels);
+				Instances.DatasmithMeshElement = FDatasmithSceneFactory::CreateMesh(TEXT(""));
+				Instances.DatasmithMeshElement->SetName(MeshName);
+				Instances.DatasmithMeshElement->SetLabel(MeshName);
 
-					if (!NodeTracker.MaterialTracker || NodeTracker.MaterialTracker->Material != Material)
-					{
-						// Release old material
-						if (NodeTracker.MaterialTracker)
-						{
-							// Release material assignment
-							MaterialsAssignedToNodes[NodeTracker.MaterialTracker].Remove(&NodeTracker);
-
-							// Clean tracker if it's not used aby any node
-							if (!MaterialsAssignedToNodes[NodeTracker.MaterialTracker].Num())
-							{
-								MaterialsTracker.UnregisterMaterialTracker(*NodeTracker.MaterialTracker);
-								MaterialTrackers.Remove(NodeTracker.MaterialTracker->Material);
-							}
-						}
-
-						if (!MaterialTrackers.Contains(Material))
-						{
-							// Track material if not yet
-							FMaterialTrackerHandle& MaterialTrackerHandle = MaterialTrackers.Emplace(Material, Material);
-							InvalidatedMaterialTrackers.Add(MaterialTrackerHandle.GetMaterialTracker());
-						}
-
-						// Store new 
-						NodeTracker.MaterialTracker = MaterialTrackers.Find(Material)->GetMaterialTracker();
-						MaterialsAssignedToNodes.FindOrAdd(NodeTracker.MaterialTracker).Add(&NodeTracker);
-					}
-				}
-				else
-				{
-					// Release old material
-					if (NodeTracker.MaterialTracker)
-					{
-						MaterialsAssignedToNodes[NodeTracker.MaterialTracker].Remove(&NodeTracker);
-						if (!MaterialsAssignedToNodes[NodeTracker.MaterialTracker].Num())
-						{
-							MaterialsTracker.UnregisterMaterialTracker(*NodeTracker.MaterialTracker);
-							MaterialTrackers.Remove(Material);
-						}
-					}
-					NodeTracker.MaterialTracker = nullptr;
-				}
+				ExportedScene.GetDatasmithScene()->AddMesh(Instances.DatasmithMeshElement);
 
 				bResult = true; // Set to true, don't care what ExportToUObject does here - we need to move it to a thread anyway
 
 				// todo: parallelize this
-				if (DatasmithMeshExporter.ExportToUObject(DatasmithMeshElement, ExportedScene.GetSceneExporter().GetAssetsOutputPath(), DatasmithMesh, nullptr, FDatasmithExportOptions::LightmapUV))
+				if (DatasmithMeshExporter.ExportToUObject(Instances.DatasmithMeshElement, ExportedScene.GetSceneExporter().GetAssetsOutputPath(), DatasmithMesh, nullptr, FDatasmithExportOptions::LightmapUV))
 				{
 					// todo: handle error exporting mesh?
 				}
 			}
 
+
 			CachedMesh.FreeAll();
 		}
 		if (bNeedsDelete)
 		{
 			RenderMesh->DeleteThis();
 		}
+		return bResult;
+	}
+
+	// todo: change this function to create mesh element and schedule ExportToUObject threading
+	bool ConvertGeomObjToDatasmithMesh(FNodeTracker& NodeTracker, Object* Obj)
+	{
+		bool bResult = false;
+
+		// AnimHandle is unique and never reused for new objects 
+		// todo: reset instances and nodes when one node of an instance changes??? Check how it should be done actually, dependencies - nodes, object, invalidation place(Update, Event) etc...
+		AnimHandle Handle = Animatable::GetHandleByAnim(Obj);
+
+		NodeTracker.InstanceHandle = Handle;
+
+		TUniquePtr<FInstances>& Instances = InstancesForAnimHandle.FindOrAdd(Handle);
+
+		if (!Instances)
+		{
+			Instances = MakeUnique<FInstances>();
+			Instances->EvaluatedObj = Obj;
+		}
+		// need to invalidate mesh assignment to node that wasn't the first to add to instances(so if instances weren't invalidated - this node needs mesh anyway)
+		Instances->NodeTrackers.Add(&NodeTracker);
+		InvalidatedInstances.Add(Instances.Get());
 
 		return bResult;
 	}
@@ -1466,8 +1593,29 @@ public:
 				}
 			}
 
-			//todo: remove ActorMesh
-			//todo: remove Mesh(if not used)
+			// Clear from mesh instances
+			if (NodeTracker->IsInstance())
+			{
+				if (TUniquePtr<FInstances>* InstancesPtr = InstancesForAnimHandle.Find(NodeTracker->InstanceHandle))
+				{
+					FInstances& Instances = **InstancesPtr;
+					Instances.NodeTrackers.Remove(NodeTracker);
+					if (Instances.NodeTrackers.Num())
+					{
+						// Invalidate all instances - this will rebuild mesh(in case removed node affected this - like simplify geometry if multimat was used but no more)
+						InvalidateInstances(Instances);
+					}
+					else
+					{
+						ExportedScene.DatasmithSceneRef->RemoveMesh(Instances.DatasmithMeshElement);
+
+						InstancesForAnimHandle.Remove(NodeTracker->InstanceHandle);
+						InvalidatedInstances.Remove(&Instances);
+					}
+				}
+				// todo: mesh is removed from the scene but not deallocated for reuse
+				// OR it will stay taking up memory if not reused(e.g. node has no valid geometry now)
+			}
 		}
 
 	}
@@ -1549,6 +1697,10 @@ public:
 	TSet<FMaterialTracker*> InvalidatedMaterialTrackers;
 
 	TMap<FMaterialTracker*, TSet<FNodeTracker*>> MaterialsAssignedToNodes;
+
+	TMap<AnimHandle, TUniquePtr<FInstances>> InstancesForAnimHandle; // set of instanced nodes for each AnimHandle
+
+	TSet<FInstances*> InvalidatedInstances;
 };
 
 
@@ -1725,17 +1877,12 @@ public:
 	virtual void CallbackBegin() override
 	{
 
-		// todo: heavy logging - called a lot
-		return;
-		LogDebug(L"NodeEventCallback: CallbackBegin\n");
+		LOG_DEBUG_HEAVY(L"NodeEventCallback: CallbackBegin\n");
 	}
 
 	virtual void CallbackEnd() override
 	{
-
-		// todo: heavy logging - called a lot
-		return;
-		LogDebug(L"NodeEventCallback: CallbackEnd\n");
+		LOG_DEBUG_HEAVY(L"NodeEventCallback: CallbackEnd\n");
 	}
 };
 
@@ -1833,9 +1980,7 @@ public:
 		case NOTIFY_PLUGINS_PRE_SHUTDOWN:
 			break;
 		default:
-			// todo: heavy logging - called a lot
-			// LogDebug(FString(TEXT("Notify: ")) + StrValue);
-			;
+			LOG_DEBUG_HEAVY(FString(TEXT("Notify: ")) + StrValue);
 		};
 
 
@@ -2115,6 +2260,18 @@ Value* Crash_cf(Value** arg_list, int count)
 
 Primitive Crash_pf(_M("Datasmith_Crash"), Crash_cf);
 
+Value* LogInfo_cf(Value** arg_list, int count)
+{
+	check_arg_count(CreateScene, 1, count);
+	Value* Message = arg_list[0];
+
+	LogInfo(Message->to_string());
+
+	return bool_result(true);
+}
+Primitive LogInfo_pf(_M("Datasmith_LogInfo"), LogInfo_cf);
+
+
 
 #include "Windows/HideWindowsPlatformTypes.h"
 
