Datasmith Revit: File paths always revert to default when exporting datasmith files from revit

diff --git a/Engine/Source/Programs/Enterprise/Datasmith/DatasmithRevitExporter/Private/DatasmithRevitCommand.cs b/Engine/Source/Programs/Enterprise/Datasmith/DatasmithRevitExporter/Private/DatasmithRevitCommand.cs
--- a/Engine/Source/Programs/Enterprise/Datasmith/DatasmithRevitExporter/Private/DatasmithRevitCommand.cs
+++ b/Engine/Source/Programs/Enterprise/Datasmith/DatasmithRevitExporter/Private/DatasmithRevitCommand.cs
@@ -137,6 +137,15 @@ namespace DatasmithRevitExporter
 	[Transaction(TransactionMode.Manual)]
 	public class DatasmithExportRevitCommand : DatasmithRevitCommand
 	{
+		class DocumentExportPathCache
+		{
+			public string LastExportPath = null;
+			// Per-view export path
+			public Dictionary<ElementId, string> ViewPaths = new Dictionary<ElementId, string>();
+		};
+
+		private static Dictionary<Document, DocumentExportPathCache> ExportPaths = new Dictionary<Document, DocumentExportPathCache>();
+	
 		// Implement the interface to execute the command.
 		public override Result OnExecute(ExternalCommandData InCommandData, ref string OutCommandMessage, ElementSet OutElements)
 		{
@@ -187,13 +196,41 @@ namespace DatasmithRevitExporter
 					return Result.Cancelled;
 				}
 
-				string ViewFamilyName = ActiveView.get_Parameter(BuiltInParameter.ELEM_FAMILY_PARAM).AsValueString().Replace(" ", "");
-				string FileName       = Regex.Replace($"{Path.GetFileNameWithoutExtension(DocumentPath)}-{ViewFamilyName}-{ActiveView.Name}.udatasmith", @"\s+", "_");
+				if (!ExportPaths.ContainsKey(Doc))
+				{
+					ExportPaths[Doc] = new DocumentExportPathCache();
+				}
+
+				string InitialDir = null;
+				string FileName = null;
+				string LastFilePath = null;
+
+				if (ExportPaths[Doc].ViewPaths.TryGetValue(ActiveView.Id, out LastFilePath))
+				{
+					InitialDir = Path.GetDirectoryName(LastFilePath);
+					FileName = Path.GetFileName(LastFilePath);
+				}
+				else
+				{
+					LastFilePath = ExportPaths[Doc].LastExportPath;
+
+					if (LastFilePath != null)
+					{
+						InitialDir = LastFilePath;
+					}
+					else
+					{
+						InitialDir = Path.GetDirectoryName(DocumentPath);
+					}
+
+					string ViewFamilyName = ActiveView.get_Parameter(BuiltInParameter.ELEM_FAMILY_PARAM).AsValueString().Replace(" ", "");
+					FileName = Regex.Replace($"{Path.GetFileNameWithoutExtension(DocumentPath)}-{ViewFamilyName}-{ActiveView.Name}.udatasmith", @"\s+", "_");
+				}
 
 				SaveFileDialog Dialog = new SaveFileDialog();
 
 				Dialog.Title            = DIALOG_CAPTION;
-				Dialog.InitialDirectory = Path.GetDirectoryName(DocumentPath);
+				Dialog.InitialDirectory = InitialDir;
 				Dialog.FileName         = FileName;
 				Dialog.DefaultExt       = "udatasmith";
 				Dialog.Filter           = "Unreal Datasmith|*.udatasmith";
@@ -207,16 +244,17 @@ namespace DatasmithRevitExporter
 					return Result.Cancelled;
 				}
 
-				string FilePath = Dialog.FileName;
+				ExportPaths[Doc].LastExportPath = Path.GetDirectoryName(Dialog.FileName);
+				ExportPaths[Doc].ViewPaths[ActiveView.Id] = Dialog.FileName;
 
-				if (string.IsNullOrWhiteSpace(FilePath))
+				if (string.IsNullOrWhiteSpace(Dialog.FileName))
 				{
 					string message = "The given Unreal Datasmith file name is blank.";
 					MessageBox.Show(message, DIALOG_CAPTION, MessageBoxButtons.OK, MessageBoxIcon.Warning);
 					return Result.Cancelled;
 				}
 
-				FilePaths.Add(ActiveView.Id, FilePath);
+				FilePaths.Add(ActiveView.Id, Dialog.FileName);
 				ViewsToExport.Add(ActiveView);
 			}
 			else
