Added UBT modules to support ArchiCAD build on Windows

diff --git a/Engine/Source/Programs/Enterprise/Datasmith/DatasmithARCHICADExporter/UBTModules/DatasmithArchiCAD.Build.cs b/Engine/Source/Programs/Enterprise/Datasmith/DatasmithARCHICADExporter/UBTModules/DatasmithArchiCAD.Build.cs
new file mode 100644
--- /dev/null
+++ b/Engine/Source/Programs/Enterprise/Datasmith/DatasmithARCHICADExporter/UBTModules/DatasmithArchiCAD.Build.cs
@@ -0,0 +1,21 @@
+// Copyright Epic Games, Inc. All Rights Reserved.
+
+using System.IO;
+
+namespace UnrealBuildTool.Rules
+{
+	[SupportedPlatforms("Win64", "Mac")]
+	public class DatasmithArchiCAD : ModuleRules
+	{
+		public DatasmithArchiCAD(ReadOnlyTargetRules Target)
+			: base(Target)
+		{
+			PublicIncludePaths.Add("Runtime/Launch/Public");
+
+			PrivateIncludePaths.Add("Runtime/Launch/Private");
+
+			PrivateDependencyModuleNames.Add("Core");
+			//PrivateDependencyModuleNames.Add("Projects");
+		}
+	}
+}
diff --git a/Engine/Source/Programs/Enterprise/Datasmith/DatasmithARCHICADExporter/UBTModules/DatasmithArchiCAD.Target.cs b/Engine/Source/Programs/Enterprise/Datasmith/DatasmithARCHICADExporter/UBTModules/DatasmithArchiCAD.Target.cs
new file mode 100644
--- /dev/null
+++ b/Engine/Source/Programs/Enterprise/Datasmith/DatasmithARCHICADExporter/UBTModules/DatasmithArchiCAD.Target.cs
@@ -0,0 +1,46 @@
+// Copyright Epic Games, Inc. All Rights Reserved.
+
+using UnrealBuildTool;
+using System.IO;
+
+[SupportedPlatforms("Win64", "Mac")]
+public class DatasmithArchiCADTarget : TargetRules
+{
+	public DatasmithArchiCADTarget(TargetInfo Target)
+		: base(Target)
+	{
+		Type = TargetType.Program;
+		SolutionDirectory = "Programs/Datasmith";
+		bBuildInSolutionByDefault = false;
+
+		ExeBinariesSubFolder = Path.Combine("DatasmithArchiCADExporter", "ArchiCAD");
+		LaunchModuleName = "DatasmithArchiCAD";
+
+		bShouldCompileAsDLL = true;
+		LinkType = TargetLinkType.Monolithic;
+
+		bBuildDeveloperTools = false;
+		bUseMallocProfiler = false;
+		bBuildWithEditorOnlyData = true;
+		bCompileAgainstEngine = false;
+		bCompileAgainstCoreUObject = true;
+		bCompileICU = false;
+		bUsesSlate = false;
+
+		bHasExports = true;
+		bForceEnableExceptions = true;
+
+		// Define post-build step
+		if (Target.Platform == UnrealTargetPlatform.Win64)
+		{
+			string ArchiCADExporterPath = @"$(EngineDir)\Source\Programs\Enterprise\Datasmith\DatasmithARCHICADExporter";
+			string SolutionPath = Path.Combine(ArchiCADExporterPath, @"Build\DatasmithARCHICADExporter.sln");
+			string BuildCommand = string.Format(@"$(EngineDir)\Build\BatchFiles\MSBuild.bat /t:Build /p:Configuration=ReleaseUE;Platform=x64 {0}", SolutionPath);
+			PreBuildSteps.Add(BuildCommand);
+		}
+		else if (Target.Platform == UnrealTargetPlatform.Mac)
+		{
+		}
+		
+	}
+}
diff --git a/Engine/Source/Programs/Enterprise/Datasmith/DatasmithARCHICADExporter/UBTModules/DatasmithArchiCAD23.Build.cs b/Engine/Source/Programs/Enterprise/Datasmith/DatasmithARCHICADExporter/UBTModules/DatasmithArchiCAD23.Build.cs
deleted file mode 100644
diff --git a/Engine/Source/Programs/Enterprise/Datasmith/DatasmithARCHICADExporter/UBTModules/DatasmithArchiCAD23.Target.cs b/Engine/Source/Programs/Enterprise/Datasmith/DatasmithARCHICADExporter/UBTModules/DatasmithArchiCAD23.Target.cs
deleted file mode 100644
