SketchUp:

- meshes exported asynchronously

- handled replacing instances with different component
UE-115435 Sketchup - replaced instance are not pushed through Direct Link

- plugin description different when DL is not supported(2019)
UE-115371: Sketchup - plugin description is misleading for 2019 user

- fixed crash removing complex components. Test case: when last component instance of a component is removed DirectLink update would(eventually) crash as child component instances were not removed from internal structures(no api event). Crash is unpredictable due to SU holding unused objects for its reasons(e.g. undo/redo). But Sketchup.active_model.definitions.purge_unused helps reproduce this - it deallocates all unused component instances making their references in the code invalid.

- Pre-populated finename when exporting to file UE-114167

diff --git a/Engine/Source/Programs/Enterprise/Datasmith/DatasmithSketchUpRubyExporter/Private/DatasmithSketchUpComponent.cpp b/Engine/Source/Programs/Enterprise/Datasmith/DatasmithSketchUpRubyExporter/Private/DatasmithSketchUpComponent.cpp
--- a/Engine/Source/Programs/Enterprise/Datasmith/DatasmithSketchUpRubyExporter/Private/DatasmithSketchUpComponent.cpp
+++ b/Engine/Source/Programs/Enterprise/Datasmith/DatasmithSketchUpRubyExporter/Private/DatasmithSketchUpComponent.cpp
@@ -513,6 +513,28 @@ void FComponentDefinition::UnlinkComponentInstance(FComponentInstance* Component
 	Instances.Remove(ComponentInstance);
 }
 
+void FComponentDefinition::RemoveComponentDefinition(FExportContext& Context)
+{
+	// Remove ComponentDefinition that doesn't have tracked instances 
+	ensure(!Instances.Num());
+	
+	// todo: might better keep in the Definition's Entities all ComponentInstanceIDs of the tracked entities
+	// this way we don't need to check whether we are tracking them (inside RemoveComponentInstance) 
+	for (SUComponentInstanceRef ComponentInstanceRef : GetEntities().GetComponentInstances())
+	{
+		Context.ComponentInstances.RemoveComponentInstance(
+			DatasmithSketchUpUtils::GetComponentID(ComponentDefinitionRef), 
+			DatasmithSketchUpUtils::GetComponentInstanceID(ComponentInstanceRef));
+	}
+
+	for (SUGroupRef GroupRef : GetEntities().GetGroups())
+	{
+		Context.ComponentInstances.RemoveComponentInstance(
+			DatasmithSketchUpUtils::GetComponentID(ComponentDefinitionRef),
+			DatasmithSketchUpUtils::GetGroupID(GroupRef));
+	}
+}
+
 void FComponentDefinition::AddInstance(FExportContext& Context, TSharedPtr<FComponentInstance> Instance)
 {
 	for (FComponentInstance* ParentInstance : Instances)
@@ -751,7 +773,8 @@ SUComponentInstanceRef FComponentInstance::GetComponentInstanceRef()
 bool FComponentInstance::RetrieveHidden()
 {
 	bool bResult = false;
-	SUDrawingElementGetHidden(SUComponentInstanceToDrawingElement(GetComponentInstanceRef()), &bResult);
+	SUResult bSuResult  = SUDrawingElementGetHidden(SUComponentInstanceToDrawingElement(GetComponentInstanceRef()), &bResult);
+	ensure(bSuResult == SU_ERROR_NONE);
 	return bResult;
 }
 
@@ -762,6 +785,10 @@ void FComponentInstance::FillOccurrenceActorMetadata(FNodeOccurence& Node)
 		return;
 	}
 
+	//SUTransformation SComponentInstanceWorldTransform = DatasmithSketchUpUtils::GetComponentInstanceTransform(GetComponentInstanceRef(), Node.ParentNode->WorldTransform);
+	//double Volume;
+	//SUComponentInstanceComputeVolume(GetComponentInstanceRef(), &SComponentInstanceWorldTransform, &Volume);
+
 	// Add original instance/component names to metadata
 	TSharedPtr<IDatasmithKeyValueProperty> EntityName = FDatasmithSceneFactory::CreateKeyValueProperty(TEXT("Instance"));
 	EntityName->SetValue(*GetName());
@@ -831,6 +858,16 @@ void FComponentInstance::RemoveComponentInstance(FExportContext& Context)
 	Definition.EntityVisible(this, false);
 	Definition.UnlinkComponentInstance(this);
 	RemoveOccurrences(Context);
+
+	// If there's no Instances of this removed ComponentInstance we need to stop tracking Definition's Entities
+	// Details:
+	// SketchUp api doesn't fire event for those child Entities although they are effectively removed from Model 
+	// Sketchup.active_model.definitions.purge_unused will deallocate those dangling Entities leaving references invalid
+	// Although SU API tries to notify about this but fails e.g. DefinitionObserver.onComponentInstanceRemoved/onEraseEntity passes already deleted Entity making this callback useless
+	if (!Definition.Instances.Num())
+	{
+		Definition.RemoveComponentDefinition(Context);
+	}
 }
 
 void FComponentInstance::SetParentDefinition(FExportContext& Context, FDefinition* InParent)
diff --git a/Engine/Source/Programs/Enterprise/Datasmith/DatasmithSketchUpRubyExporter/Private/DatasmithSketchUpComponent.h b/Engine/Source/Programs/Enterprise/Datasmith/DatasmithSketchUpRubyExporter/Private/DatasmithSketchUpComponent.h
--- a/Engine/Source/Programs/Enterprise/Datasmith/DatasmithSketchUpRubyExporter/Private/DatasmithSketchUpComponent.h
+++ b/Engine/Source/Programs/Enterprise/Datasmith/DatasmithSketchUpRubyExporter/Private/DatasmithSketchUpComponent.h
@@ -167,7 +167,7 @@ namespace DatasmithSketchUp
 			SUComponentDefinitionRef InComponentDefinitionRef // source SketchUp component definition
 		);
 
-		// Being FDefinition
+		// Begin FDefinition
 		void Parse(FExportContext& Context);
 		void CreateActor(FExportContext& Context, FNodeOccurence& Node) override;
 		void BuildNodeNames(FNodeOccurence& Node) override;
@@ -187,10 +187,11 @@ namespace DatasmithSketchUp
 		// Register/unregister instanced of this definition
 		void LinkComponentInstance(FComponentInstance* ComponentInstance);
 		void UnlinkComponentInstance(FComponentInstance* ComponentInstance);
+		void RemoveComponentDefinition(FExportContext& Context);
 
 		// Source SketchUp component ID.
 		FComponentDefinitionIDType SketchupSourceID;
-		TSet<FComponentInstance*> Instances;
+		TSet<FComponentInstance*> Instances; // Tracked instances of this ComponentDefinition
 
 	private:
 		SUComponentDefinitionRef ComponentDefinitionRef;
diff --git a/Engine/Source/Programs/Enterprise/Datasmith/DatasmithSketchUpRubyExporter/Private/DatasmithSketchUpExportContext.cpp b/Engine/Source/Programs/Enterprise/Datasmith/DatasmithSketchUpRubyExporter/Private/DatasmithSketchUpExportContext.cpp
--- a/Engine/Source/Programs/Enterprise/Datasmith/DatasmithSketchUpRubyExporter/Private/DatasmithSketchUpExportContext.cpp
+++ b/Engine/Source/Programs/Enterprise/Datasmith/DatasmithSketchUpRubyExporter/Private/DatasmithSketchUpExportContext.cpp
@@ -126,6 +126,13 @@ void FExportContext::Update()
 	RootNode->Update(*this);
 
 	Textures.Update();
+
+	// Wait for mesh export to complete
+	for(TFuture<bool>& Task: MeshExportTasks)
+	{
+		Task.Get();
+	}
+	MeshExportTasks.Reset();
 }
 
 FDefinition* FExportContext::GetDefinition(SUEntityRef Entity)
@@ -340,17 +347,23 @@ bool FComponentInstanceCollection::RemoveComponentInstance(FComponentInstanceIDT
 	//
 	// Details:
 	// ComponentInstance which removal is notified could have been relocated to another Definition 
-	// This happens when Make Group is done - first new Group is added, containing existing ComponentInstace
-	// And only after that event about removal previous owning Definition is received
+	// This happens when Make Group is done - first new Group is added, containing existing ComponentInstance
+	// And only after that event about removal from previous owning Definition is received
 	if (ComponentInstance->IsParentDefinition(ParentDefinition))
 	{
-		ComponentInstance->RemoveComponentInstance(Context);
-		ComponentInstanceMap.Remove(ComponentInstanceId);
+		RemoveComponentInstance(ComponentInstance);
 	}
 
 	return true;
 }
 
+void FComponentInstanceCollection::RemoveComponentInstance(TSharedPtr<FComponentInstance> ComponentInstance)
+{
+	ComponentInstance->RemoveComponentInstance(Context);
+	ComponentInstanceMap.Remove(ComponentInstance->GetComponentInstanceId());
+}
+
+
 void FComponentInstanceCollection::InvalidateComponentInstanceGeometry(FComponentInstanceIDType ComponentInstanceID)
 {
 	if (TSharedPtr<FComponentInstance>* Ptr = FindComponentInstance(ComponentInstanceID))
@@ -367,11 +380,24 @@ void FComponentInstanceCollection::InvalidateComponentInstanceMetadata(FComponen
 	}
 }
 
-bool FComponentInstanceCollection::InvalidateComponentInstanceProperties(FComponentInstanceIDType ComponentInstanceID)
+bool FComponentInstanceCollection::InvalidateComponentInstanceProperties(FComponentInstanceIDType ComponentInstanceId)
 {
-	if (TSharedPtr<FComponentInstance>* Ptr = FindComponentInstance(ComponentInstanceID))
+	if (TSharedPtr<FComponentInstance>* Ptr = FindComponentInstance(ComponentInstanceId))
 	{
-		(*Ptr)->InvalidateEntityProperties();
+		TSharedPtr<FComponentInstance> ComponentInstance = *Ptr;
+		
+		// Replacing definition on a component instance fires the same event as changing properties
+		SUComponentInstanceRef ComponentInstanceRef = ComponentInstance->GetComponentInstanceRef();
+		TSharedPtr<FComponentDefinition> Definition = Context.ComponentDefinitions.GetComponentDefinition(ComponentInstanceRef);
+		if (ComponentInstance->GetDefinition() != Definition.Get())
+		{
+			// Recreate and re-add instance
+			FDefinition* ParentDefinition = ComponentInstance->Parent;
+			RemoveComponentInstance(ComponentInstance);
+			ParentDefinition->AddInstance(Context, AddComponentInstance(*ParentDefinition, ComponentInstanceRef));
+		}
+
+		ComponentInstance->InvalidateEntityProperties();
 		return true;
 	}
 	return false;
diff --git a/Engine/Source/Programs/Enterprise/Datasmith/DatasmithSketchUpRubyExporter/Private/DatasmithSketchUpExportContext.h b/Engine/Source/Programs/Enterprise/Datasmith/DatasmithSketchUpRubyExporter/Private/DatasmithSketchUpExportContext.h
--- a/Engine/Source/Programs/Enterprise/Datasmith/DatasmithSketchUpRubyExporter/Private/DatasmithSketchUpExportContext.h
+++ b/Engine/Source/Programs/Enterprise/Datasmith/DatasmithSketchUpRubyExporter/Private/DatasmithSketchUpExportContext.h
@@ -11,6 +11,8 @@
 #include "Containers/UnrealString.h"
 #include "Templates/SharedPointer.h"
 
+#include "Async/Future.h"
+
 class FDatasmithSceneExporter;
 
 class IDatasmithActorElement;
@@ -48,6 +50,7 @@ namespace DatasmithSketchUp
 
 		TSharedPtr<FComponentInstance> AddComponentInstance(FDefinition& ParentDefinition, SUComponentInstanceRef InComponentInstanceRef); // Register ComponentInstance as a child of ParentDefinition
 		bool RemoveComponentInstance(FComponentInstanceIDType ParentEntityId, FComponentInstanceIDType ComponentInstanceId); // Take note that ComponentInstance removed from ParentDefinition children
+		void RemoveComponentInstance(TSharedPtr<FComponentInstance> ComponentInstance);
 
 		bool InvalidateComponentInstanceProperties(FComponentInstanceIDType ComponentInstanceID);
 		void InvalidateComponentInstanceGeometry(FComponentInstanceIDType ComponentInstanceID);
@@ -211,5 +214,7 @@ namespace DatasmithSketchUp
 		FMaterialCollection Materials;
 		FSceneCollection Scenes;
 		FTextureCollection Textures;
+
+		TArray<TFuture<bool>> MeshExportTasks;
 	};
 }
diff --git a/Engine/Source/Programs/Enterprise/Datasmith/DatasmithSketchUpRubyExporter/Private/DatasmithSketchUpExporter.cpp b/Engine/Source/Programs/Enterprise/Datasmith/DatasmithSketchUpRubyExporter/Private/DatasmithSketchUpExporter.cpp
--- a/Engine/Source/Programs/Enterprise/Datasmith/DatasmithSketchUpRubyExporter/Private/DatasmithSketchUpExporter.cpp
+++ b/Engine/Source/Programs/Enterprise/Datasmith/DatasmithSketchUpRubyExporter/Private/DatasmithSketchUpExporter.cpp
@@ -243,121 +243,6 @@ public:
 		}
 	}
 
-	// XXX: REMOVE before shipping
-	// Used for simple testing on the plugin side what is being sent to DL
-	void ExportCurrentDatasmithSceneWithoutCleanup()
-	{
-		{
-			TSharedRef<IDatasmithScene> DatasmithScene = ExportedScene.GetDatasmithSceneRef();
-			FString FilePath = FPaths::Combine(ExportedScene.GetSceneExporterRef()->GetOutputPath(), ExportedScene.GetSceneExporterRef()->GetName()) + TEXT(".") + FDatasmithUtils::GetFileExtension();
-
-			TUniquePtr<FArchive> Archive(IFileManager::Get().CreateFileWriter(*FilePath));
-
-			if (!Archive.IsValid())
-			{
-				// XXX - removed 
-				//if (Impl->Logger.IsValid())
-				//{
-				//	Impl->Logger->AddGeneralError(*(TEXT("Unable to create file ") + FilePath + TEXT(", Aborting the export process")));
-				//}
-				return;
-			}
-
-			IPlatformFile& PlatformFile = FPlatformFileManager::Get().GetPlatformFile();
-			PlatformFile.CreateDirectoryTree(ExportedScene.GetSceneExporterRef()->GetAssetsOutputPath());
-
-			// Add Bump maps from Material objects to scene as TextureElement
-			// XXX - removed Impl->CheckBumpMaps(DatasmithScene);
-
-			// XXX - removed FDatasmithSceneUtils::CleanUpScene(DatasmithScene, bCleanupUnusedElements);
-
-			// Update TextureElements
-			// XXX - removed Impl->UpdateTextureElements(DatasmithScene);
-
-			// XXX - note todo
-			// todo: keep relative myself
-			// Convert paths to relative
-			FString AbsoluteDir = FString(ExportedScene.GetSceneExporterRef()->GetOutputPath()) + TEXT("/");
-
-			for (int32 MeshIndex = 0; MeshIndex < DatasmithScene->GetMeshesCount(); ++MeshIndex)
-			{
-				TSharedPtr< IDatasmithMeshElement > Mesh = DatasmithScene->GetMesh(MeshIndex);
-
-				FString RelativePath = Mesh->GetFile();
-				FPaths::MakePathRelativeTo(RelativePath, *AbsoluteDir);
-
-				Mesh->SetFile(*RelativePath);
-			}
-
-			for (int32 TextureIndex = 0; TextureIndex < DatasmithScene->GetTexturesCount(); ++TextureIndex)
-			{
-				TSharedPtr< IDatasmithTextureElement > Texture = DatasmithScene->GetTexture(TextureIndex);
-
-				FString TextureFile = Texture->GetFile();
-				FPaths::MakePathRelativeTo(TextureFile, *AbsoluteDir);
-				Texture->SetFile(*TextureFile);
-			}
-
-			// XXX - removed 
-			//FDatasmithAnimationSerializer AnimSerializer;
-			//int32 NumSequences = DatasmithScene->GetLevelSequencesCount();
-			//for (int32 SequenceIndex = 0; SequenceIndex < NumSequences; ++SequenceIndex)
-			//{
-			//	const TSharedPtr<IDatasmithLevelSequenceElement>& LevelSequence = DatasmithScene->GetLevelSequence(SequenceIndex);
-			//	if (LevelSequence.IsValid())
-			//	{
-			//		FString AnimFilePath = FPaths::Combine(Impl->AssetsOutputPath, LevelSequence->GetName()) + DATASMITH_ANIMATION_EXTENSION;
-
-			//		if (AnimSerializer.Serialize(LevelSequence.ToSharedRef(), *AnimFilePath))
-			//		{
-			//			TUniquePtr<FArchive> AnimArchive(IFileManager::Get().CreateFileReader(*AnimFilePath));
-			//			if (AnimArchive)
-			//			{
-			//				LevelSequence->SetFileHash(FMD5Hash::HashFileFromArchive(AnimArchive.Get()));
-			//			}
-
-			//			FPaths::MakePathRelativeTo(AnimFilePath, *AbsoluteDir);
-			//			LevelSequence->SetFile(*AnimFilePath);
-			//		}
-			//	}
-			//}
-
-			// XXX - removed 
-			// Log time spent to export scene in seconds
-			//int ElapsedTime = (int)FPlatformTime::ToSeconds64(FPlatformTime::Cycles64() - Impl->ExportStartCycles);
-			//DatasmithScene->SetExportDuration(ElapsedTime);
-
-			FDatasmithSceneXmlWriter DatasmithSceneXmlWriter;
-			DatasmithSceneXmlWriter.Serialize(DatasmithScene, *Archive);
-
-			Archive->Close();
-
-			// Run the garbage collector at this point so that we're in a good state for the next export
-			FDatasmithExporterManager::RunGarbageCollection();
-		}
-
-		TSharedRef<IDatasmithScene> DatasmithScene = ExportedScene.GetDatasmithSceneRef();
-
-		// Convert paths back to absolute(they were changes by Export)
-		FString AbsoluteDir = FString(ExportedScene.GetSceneExporterRef()->GetOutputPath()) + TEXT("/");
-		for (int32 MeshIndex = 0; MeshIndex < DatasmithScene->GetMeshesCount(); ++MeshIndex)
-		{
-			TSharedPtr< IDatasmithMeshElement > Mesh = DatasmithScene->GetMesh(MeshIndex);
-
-			FString RelativePath = Mesh->GetFile();
-			Mesh->SetFile(*FPaths::ConvertRelativePathToFull(AbsoluteDir, *RelativePath));
-		}
-
-		for (int32 TextureIndex = 0; TextureIndex < DatasmithScene->GetTexturesCount(); ++TextureIndex)
-		{
-			TSharedPtr< IDatasmithTextureElement > Texture = DatasmithScene->GetTexture(TextureIndex);
-
-			FString RelativePath = Texture->GetFile();
-
-			Texture->SetFile(*FPaths::ConvertRelativePathToFull(AbsoluteDir, *RelativePath));
-		}
-	}
-
 	void ExportCurrentDatasmithScene()
 	{
 		ExportedScene.GetSceneExporterRef()->Export(ExportedScene.GetDatasmithSceneRef());
@@ -674,17 +559,6 @@ VALUE DatasmithSketchUpDirectLinkExporter_export_current_datasmith_scene(VALUE s
 	return Qtrue;
 }
 
-VALUE DatasmithSketchUpDirectLinkExporter_export_current_datasmith_scene_no_cleanup(VALUE self)
-{
-	// Converting args
-	FDatasmithSketchUpDirectLinkExporter* ptr;
-	Data_Get_Struct(self, FDatasmithSketchUpDirectLinkExporter, ptr);
-	// Done converting args
-
-	ptr->ExportCurrentDatasmithSceneWithoutCleanup();
-	return Qtrue;
-}
-
 #ifndef SKP_SDK_2019
 VALUE DatasmithSketchUpDirectLinkExporter_on_component_instance_changed(VALUE self, VALUE ruby_entity)
 {
@@ -911,8 +785,6 @@ extern "C" DLLEXPORT void Init_DatasmithSketchUpRuby()
 	rb_define_method(DatasmithSketchUpDirectLinkExporterCRubyClass, "update", ToRuby(DatasmithSketchUpDirectLinkExporter_update), 0);
 	rb_define_method(DatasmithSketchUpDirectLinkExporterCRubyClass, "send_update", ToRuby(DatasmithSketchUpDirectLinkExporter_send_update), 0);
 	rb_define_method(DatasmithSketchUpDirectLinkExporterCRubyClass, "export_current_datasmith_scene", ToRuby(DatasmithSketchUpDirectLinkExporter_export_current_datasmith_scene), 0);
-	rb_define_method(DatasmithSketchUpDirectLinkExporterCRubyClass, "export_current_datasmith_scene_no_cleanup", ToRuby(DatasmithSketchUpDirectLinkExporter_export_current_datasmith_scene_no_cleanup), 0);
-	
 }
 
 /* todo:
diff --git a/Engine/Source/Programs/Enterprise/Datasmith/DatasmithSketchUpRubyExporter/Private/DatasmithSketchUpMesh.cpp b/Engine/Source/Programs/Enterprise/Datasmith/DatasmithSketchUpRubyExporter/Private/DatasmithSketchUpMesh.cpp
--- a/Engine/Source/Programs/Enterprise/Datasmith/DatasmithSketchUpRubyExporter/Private/DatasmithSketchUpMesh.cpp
+++ b/Engine/Source/Programs/Enterprise/Datasmith/DatasmithSketchUpRubyExporter/Private/DatasmithSketchUpMesh.cpp
@@ -32,6 +32,9 @@
 #include "DatasmithUtils.h"
 #include "Misc/SecureHash.h"
 
+#include "Async/Async.h"
+#include "UObject/GarbageCollection.h"
+
 class FDatasmithMesh;
 
 using namespace DatasmithSketchUp;
@@ -340,7 +343,7 @@ const TCHAR* FEntitiesGeometry::GetMeshElementName(int32 MeshIndex)
 	return Meshes[MeshIndex]->DatasmithMesh->GetName();
 }
 
-void ScanSketchUpEntitiesFaces(SUEntitiesRef EntitiesRef, TSet<int32>& ScannedFaceIDSet, TFunctionRef<void(FDatasmithSketchUpMesh& ExtractedMesh)> OnNewExtractedMesh);
+void ScanSketchUpEntitiesFaces(SUEntitiesRef EntitiesRef, TSet<int32>& ScannedFaceIDSet, TFunctionRef<void(TSharedPtr<FDatasmithSketchUpMesh> ExtractedMesh)> OnNewExtractedMesh);
 
 void FEntities::UpdateGeometry(FExportContext& Context)
 {
@@ -359,16 +362,12 @@ void FEntities::UpdateGeometry(FExportContext& Context)
 		EntitiesGeometry = MakeShared<FEntitiesGeometry>();
 	}
 
-	FDatasmithMeshExporter DatasmithMeshExporter;
 	int32 MeshCount = 0;
 
-	auto ProcessExtractedMesh = [&Context, &DatasmithMeshExporter, this, &MeshCount](FDatasmithSketchUpMesh& ExtractedMesh)
+	TFunction<void(TSharedPtr<FDatasmithSketchUpMesh> ExtractedMesh)> ProcessExtractedMesh = [&Context, this, &MeshCount](TSharedPtr<FDatasmithSketchUpMesh> ExtractedMeshPtr)
 	{
-		if (ExtractedMesh.ContainsGeometry())
+		if (ExtractedMeshPtr->ContainsGeometry())
 		{
-			FDatasmithMesh DatasmithMesh;
-			ExtractedMesh.ConvertMeshToDatasmith(DatasmithMesh);
-
 			FString MeshElementName = FString::Printf(TEXT("M%ls_%d"), *FMD5::HashAnsiString(*Definition.GetSketchupSourceGUID()), MeshCount + 1); // Count meshes from 1
 			FString MeshLabel = FDatasmithUtils::SanitizeObjectName(Definition.GetSketchupSourceName());
 
@@ -388,22 +387,31 @@ void FEntities::UpdateGeometry(FExportContext& Context)
 
 			Mesh->DatasmithMesh->SetName(*MeshElementName);
 			Mesh->DatasmithMesh->SetLabel(*MeshLabel);
-			Mesh->bIsUsingInheritedMaterial = ExtractedMesh.MeshTriangleMaterialIDSet.Contains(FMaterial::INHERITED_MATERIAL_ID);
+			Mesh->bIsUsingInheritedMaterial = ExtractedMeshPtr->MeshTriangleMaterialIDSet.Contains(FMaterial::INHERITED_MATERIAL_ID);
 
 			// Add the non-inherited materials used by the combined mesh triangles.
-			for (FMaterialIDType MeshMaterialID : ExtractedMesh.MeshTriangleMaterialIDSet)
+			for (FMaterialIDType MeshMaterialID : ExtractedMeshPtr->MeshTriangleMaterialIDSet)
 			{
-				if (MeshMaterialID != FMaterial::INHERITED_MATERIAL_ID)
+				// Default or (somehow)missing materials are also assigned to mesh(as a default material)
+				if (FMaterialOccurrence* Material = Context.Materials.RegisterGeometry(MeshMaterialID, EntitiesGeometry.Get()))
 				{
-					if (FMaterialOccurrence* Material = Context.Materials.RegisterGeometry(MeshMaterialID, EntitiesGeometry.Get()))
-					{
-						int32 SlotId = MeshMaterialID.EntityID; // Triangles assigned Material Ids as SlotIds
-						Mesh->DatasmithMesh->SetMaterial(Material->GetName(), SlotId);
-					}
+					int32 SlotId = MeshMaterialID.EntityID; // Triangles assigned Material Ids as SlotIds
+					Mesh->DatasmithMesh->SetMaterial(Material->GetName(), SlotId);
 				}
 			}
 
-			DatasmithMeshExporter.ExportToUObject(Mesh->DatasmithMesh, Context.GetAssetsOutputPath(), DatasmithMesh, nullptr, FDatasmithExportOptions::LightmapUV);
+			Context.MeshExportTasks.Emplace(Async(
+				EAsyncExecution::ThreadPool,
+				[Mesh, ExtractedMeshPtr, &Context]()
+				{
+					FDatasmithMeshExporter DatasmithMeshExporter;
+					FDatasmithMesh DatasmithMesh;
+					ExtractedMeshPtr->ConvertMeshToDatasmith(DatasmithMesh);
+
+					FGCScopeGuard GCGuard; // Prevent GC from running while UDatasmithMesh is created in ExportToUObject. 
+					return DatasmithMeshExporter.ExportToUObject(Mesh->DatasmithMesh, Context.GetAssetsOutputPath(), DatasmithMesh, nullptr, FDatasmithExportOptions::LightmapUV);
+				}
+			));
 		}
 	};
 
@@ -475,7 +483,7 @@ TArray<SUComponentInstanceRef> FEntities::GetComponentInstances()
 }
 
 
-void ScanSketchUpEntitiesFaces(SUEntitiesRef EntitiesRef, TSet<int32>& ScannedFaceIDSet, TFunctionRef<void(FDatasmithSketchUpMesh& ExtractedMesh)> OnNewExtractedMesh)
+void ScanSketchUpEntitiesFaces(SUEntitiesRef EntitiesRef, TSet<int32>& ScannedFaceIDSet, TFunctionRef<void(TSharedPtr<FDatasmithSketchUpMesh> ExtractedMesh)> OnNewExtractedMesh)
 {
 	// Get the number of faces in the source SketchUp entities.
 	size_t SFaceCount = 0;
@@ -509,7 +517,8 @@ void ScanSketchUpEntitiesFaces(SUEntitiesRef EntitiesRef, TSet<int32>& ScannedFa
 		}
 
 		// Create a mesh combining the geometry of the SketchUp connected faces.
-		FDatasmithSketchUpMesh ExtractedMesh;
+		TSharedPtr<FDatasmithSketchUpMesh> ExtractedMeshPtr = MakeShared<FDatasmithSketchUpMesh>();
+		FDatasmithSketchUpMesh& ExtractedMesh = *ExtractedMeshPtr;
 
 		// The source SketchUp face needs to be scanned once.
 		TArray<SUFaceRef> FacesToScan;
@@ -580,6 +589,6 @@ void ScanSketchUpEntitiesFaces(SUEntitiesRef EntitiesRef, TSet<int32>& ScannedFa
 			}
 		}
 
-		OnNewExtractedMesh(ExtractedMesh);
+		OnNewExtractedMesh(ExtractedMeshPtr);
 	}
 }
\ No newline at end of file
