TypeName in english only (i.e. non localized),
Add parameters to Element 2 string,
Fix bug caused by difference between TMap (mapped object can move) and std map (mapped object don't move),
Add lib path info,
Add Light IES,
Dont create menus,
Adjust material Opacity,
Begin Translation (Add French),
Handle correctly GetDirectLinkCacheDirectory,
Add Meta data for enscapeAssetId.

diff --git a/Engine/Source/Programs/Enterprise/Datasmith/DatasmithARCHICADExporter/Build/Config/Notes.txt b/Engine/Source/Programs/Enterprise/Datasmith/DatasmithARCHICADExporter/Build/Config/Notes.txt
new file mode 100755
--- /dev/null
+++ b/Engine/Source/Programs/Enterprise/Datasmith/DatasmithARCHICADExporter/Build/Config/Notes.txt
@@ -0,0 +1 @@
+Project information to come
diff --git a/Engine/Source/Programs/Enterprise/Datasmith/DatasmithARCHICADExporter/Private/AddOn.cpp b/Engine/Source/Programs/Enterprise/Datasmith/DatasmithARCHICADExporter/Private/AddOn.cpp
--- a/Engine/Source/Programs/Enterprise/Datasmith/DatasmithARCHICADExporter/Private/AddOn.cpp
+++ b/Engine/Source/Programs/Enterprise/Datasmith/DatasmithARCHICADExporter/Private/AddOn.cpp
@@ -43,7 +43,7 @@ using namespace UE_AC;
 
 API_AddonType __ACENV_CALL CheckEnvironment(API_EnvirParams* envir)
 {
-	UE_AC_TraceF("--- UE_AC CheckEnvironment\n");
+	UE_AC_TraceF("-> UE_AC CheckEnvironment\n");
 
 	short		  IdDescription = LocalizeResId(kStrListSyncPlugInDescription);
 	GS::UniString versStr("\n\t");
@@ -55,6 +55,8 @@ API_AddonType __ACENV_CALL CheckEnvironment(API_EnvirParams* envir)
 	RSGetIndString(&envir->addOnInfo.description, IdDescription, 2, ACAPI_GetOwnResModule());
 	envir->addOnInfo.description += versStr;
 
+	UE_AC_TraceF("<- UE_AC CheckEnvironment\n");
+
 	return APIAddon_Preload;
 }
 
@@ -64,7 +66,7 @@ API_AddonType __ACENV_CALL CheckEnvironment(API_EnvirParams* envir)
 
 GSErrCode __ACENV_CALL RegisterInterface(void)
 {
-	UE_AC_TraceF("--- UE_AC RegisterInterface\n");
+	UE_AC_TraceF("-> UE_AC RegisterInterface\n");
 
 	GSErrCode GSErr = FExport::Register();
 	if (GSErr == NoError)
@@ -82,6 +84,8 @@ GSErrCode __ACENV_CALL RegisterInterface(void)
 
 	ACAPI_KeepInMemory(true);
 
+	UE_AC_TraceF("<- UE_AC RegisterInterface\n");
+
 	return GSErr;
 }
 
@@ -92,7 +96,7 @@ GSErrCode __ACENV_CALL RegisterInterface(void)
 
 GSErrCode __ACENV_CALL Initialize(void)
 {
-	UE_AC_TraceF("--- UE_AC Initialize\n");
+	UE_AC_TraceF("-> UE_AC Initialize\n");
 
 	LoadDatasmithDlls();
 	FTraceListener::Get();
@@ -122,6 +126,9 @@ GSErrCode __ACENV_CALL Initialize(void)
 	FPalette::Register();
 
 	ACAPI_KeepInMemory(true);
+
+	UE_AC_TraceF("<- UE_AC Initialize\n");
+
 	return GSErr;
 }
 
@@ -132,7 +139,7 @@ GSErrCode __ACENV_CALL Initialize(void)
 
 GSErrCode __ACENV_CALL FreeData(void)
 {
-	UE_AC_TraceF("--- UE_AC FreeData\n");
+	UE_AC_TraceF("-> UE_AC FreeData\n");
 
 	FPalette::Delete();
 	FPalette::Unregister();
@@ -141,5 +148,7 @@ GSErrCode __ACENV_CALL FreeData(void)
 	FReportWindow::Delete();
 	FTraceListener::Delete();
 
+	UE_AC_TraceF("<- UE_AC FreeData\n");
+
 	return NoError;
 }
diff --git a/Engine/Source/Programs/Enterprise/Datasmith/DatasmithARCHICADExporter/Private/ElementEvent.cpp b/Engine/Source/Programs/Enterprise/Datasmith/DatasmithARCHICADExporter/Private/ElementEvent.cpp
--- a/Engine/Source/Programs/Enterprise/Datasmith/DatasmithARCHICADExporter/Private/ElementEvent.cpp
+++ b/Engine/Source/Programs/Enterprise/Datasmith/DatasmithARCHICADExporter/Private/ElementEvent.cpp
@@ -64,7 +64,7 @@ GSErrCode FElementEvent::Initialize()
 
 GSErrCode FElementEvent::Event(const API_NotifyElementType& ElemType)
 {
-	UE_AC_TraceF("--- FElementEvent::Event(%s)\n", TAssEnumName< API_ElementDBEventID >::GetName(ElemType.notifID));
+	UE_AC_TraceF("-> FElementEvent::Event(%s)\n", TAssEnumName< API_ElementDBEventID >::GetName(ElemType.notifID));
 	GS::UniString ElemTypeName;
 	if (ACAPI_Goodies(APIAny_GetElemTypeNameID, (void*)(size_t)ElemType.elemHead.typeID, &ElemTypeName) != NoError)
 	{
@@ -131,6 +131,8 @@ GSErrCode FElementEvent::Event(const API_NotifyElementType& ElemType)
 			break;
 	}
 	FSynchronizer::PostDoSnapshot("Element modified");
+
+	UE_AC_TraceF("<- FElementEvent::Event\n");
 	return NoError;
 }
 
diff --git a/Engine/Source/Programs/Enterprise/Datasmith/DatasmithARCHICADExporter/Private/ElementID.cpp b/Engine/Source/Programs/Enterprise/Datasmith/DatasmithARCHICADExporter/Private/ElementID.cpp
--- a/Engine/Source/Programs/Enterprise/Datasmith/DatasmithARCHICADExporter/Private/ElementID.cpp
+++ b/Engine/Source/Programs/Enterprise/Datasmith/DatasmithARCHICADExporter/Private/ElementID.cpp
@@ -2,6 +2,8 @@
 
 #include "ElementID.h"
 
+#include <stdexcept>
+
 BEGIN_NAMESPACE_UE_AC
 
 template <>
@@ -64,10 +66,111 @@ FAssValueName::SAssValueName TAssEnumName< ModelerAPI::Element::Type >::AssEnumN
 	EnumName(ModelerAPI::Element, OtherElement),
 	EnumEnd(-1)};
 
+// Contructor
+FElementID::FElementID(const FSyncContext& InSyncContext)
+	: SyncContext(InSyncContext)
+	, Index3D(0)
+	, SyncData(nullptr)
+	, LibPartInfo(nullptr)
+	, bFullElementFetched(false)
+	, bLibPartInfoFetched(false)
+{
+}
+
+// Initialize with 3D element
+void FElementID::InitElement(GS::Int32 InIndex3d)
+{
+	Index3D = InIndex3d;
+	SyncContext.GetModel().GetElement(Index3D, &Element3D);
+	APIElement.header.guid = APINULLGuid;
+	bFullElementFetched = false;
+	LibPartInfo = nullptr;
+	bLibPartInfoFetched = false;
+}
+
+// Initialize with sync data
+void FElementID::InitElement(FSyncData* IOSyncdata)
+{
+	UE_AC_TestPtr(IOSyncdata);
+	SyncData = IOSyncdata;
+	Index3D = IOSyncdata->GetIndex3D();
+	if (Index3D > 0)
+	{
+		SyncContext.GetModel().GetElement(Index3D, &Element3D);
+	}
+	bFullElementFetched = false;
+	LibPartInfo = nullptr;
+	bLibPartInfoFetched = false;
+}
+
+// Initialize element header from 3D element
+bool FElementID::InitHeader()
+{
+	if (IsInvalid())
+	{
+		throw std::runtime_error(
+			Utf8StringFormat("FElementID::InitHeader - Invalid element for index=%d\n", Index3D).c_str());
+	}
+	bFullElementFetched = false;
+	Zap(&APIElement.header);
+	APIElement.header.guid = GSGuid2APIGuid(Element3D.GetElemGuid());
+	GSErrCode GSErr = ACAPI_Element_GetHeader(&APIElement.header, 0);
+	if (GSErr != NoError)
+	{
+		utf8_string ErrorName(GetErrorName(GSErr));
+		utf8_string TypeName(GetTypeName());
+		UE_AC_DebugF("Error \"%s\" with element %d {%s} Type=%s\n", ErrorName.c_str(), Index3D,
+					 Element3D.GetElemGuid().ToUniString().ToUtf8(), TypeName.c_str());
+		if (GSErr != APIERR_BADID)
+		{
+			UE_AC::ThrowGSError(GSErr, __FILE__, __LINE__);
+		}
+		return false;
+	}
+	return true;
+}
+
+const API_Element& FElementID::GetAPIElement()
+{
+	if (bFullElementFetched == false)
+	{
+		API_Guid guid = APIElement.header.guid;
+		Zap(&APIElement);
+		APIElement.header.guid = guid;
+		UE_AC_TestGSError(ACAPI_Element_Get(&APIElement, 0));
+		bFullElementFetched = true;
+	}
+
+	return APIElement;
+}
+
+// If this element is related to a lib part ?
+const FLibPartInfo* FElementID::GetLibPartInfo()
+{
+	if (bLibPartInfoFetched == false)
+	{
+		// Get the lib part from it's UnId
+		FGSUnID::Buffer lpfUnID = {0};
+		GSErrCode		GSErr = ACAPI_Goodies(APIAny_GetElemLibPartUnIdID, &APIElement.header, lpfUnID);
+		if (GSErr == NoError)
+		{
+			LibPartInfo = SyncContext.GetSyncDatabase().GetLibPartInfo(lpfUnID);
+		}
+		else if (GSErr != APIERR_BADID)
+		{
+			UE_AC_DebugF("FElementID::InitLibPartInfo - APIAny_GetElemLibPartUnIdID return error %s\n",
+						 GetErrorName(GSErr));
+		}
+		bLibPartInfoFetched = true;
+	}
+
+	return LibPartInfo;
+}
+
 void FElementID::CollectDependantElementsType(API_ElemTypeID TypeID) const
 {
 	GS::Array< API_Guid > ConnectedElements;
-	UE_AC_TestGSError(ACAPI_Element_GetConnectedElements(ElementHeader.guid, TypeID, &ConnectedElements));
+	UE_AC_TestGSError(ACAPI_Element_GetConnectedElements(APIElement.header.guid, TypeID, &ConnectedElements));
 	for (USize i = 0; i < ConnectedElements.GetSize(); ++i)
 	{
 		FSyncData*& ChildSyncData = SyncContext.GetSyncDatabase().GetSyncData(APIGuid2GSGuid(ConnectedElements[i]));
@@ -83,24 +186,24 @@ void FElementID::CollectDependantElementsType(API_ElemTypeID TypeID) const
 
 void FElementID::HandleDepedencies() const
 {
-	if (ElementHeader.typeID == API_WallID)
+	if (APIElement.header.typeID == API_WallID)
 	{
 		CollectDependantElementsType(API_WindowID);
 		CollectDependantElementsType(API_DoorID);
 	}
-	else if (ElementHeader.typeID == API_RoofID || ElementHeader.typeID == API_ShellID)
+	else if (APIElement.header.typeID == API_RoofID || APIElement.header.typeID == API_ShellID)
 	{
 		CollectDependantElementsType(API_SkylightID);
 	}
-	else if (ElementHeader.typeID == API_WindowID || ElementHeader.typeID == API_DoorID ||
-			 ElementHeader.typeID == API_SkylightID)
+	else if (APIElement.header.typeID == API_WindowID || APIElement.header.typeID == API_DoorID ||
+			 APIElement.header.typeID == API_SkylightID)
 	{
 		// Do nothing
 	}
 	else
 	{
-		GS::Guid				  OwnerElemGuid = APIGuid2GSGuid(ElementHeader.guid);
-		API_Guid				  OwnerElemApiGuid = ElementHeader.guid;
+		GS::Guid				  OwnerElemGuid = APIGuid2GSGuid(APIElement.header.guid);
+		API_Guid				  OwnerElemApiGuid = APIElement.header.guid;
 		API_HierarchicalElemType  HierarchicalElemType = API_SingleElem;
 		API_HierarchicalOwnerType HierarchicalOwnerType = API_RootHierarchicalOwner;
 		GSErrCode GSErr = ACAPI_Goodies(APIAny_GetHierarchicalElementOwnerID, &OwnerElemGuid, &HierarchicalOwnerType,
diff --git a/Engine/Source/Programs/Enterprise/Datasmith/DatasmithARCHICADExporter/Private/ElementID.h b/Engine/Source/Programs/Enterprise/Datasmith/DatasmithARCHICADExporter/Private/ElementID.h
--- a/Engine/Source/Programs/Enterprise/Datasmith/DatasmithARCHICADExporter/Private/ElementID.h
+++ b/Engine/Source/Programs/Enterprise/Datasmith/DatasmithARCHICADExporter/Private/ElementID.h
@@ -10,8 +10,6 @@
 
 #include "ModelElement.hpp"
 
-#include <stdexcept>
-
 BEGIN_NAMESPACE_UE_AC
 
 // Class that contain element id and related
@@ -19,89 +17,88 @@ class FElementID
 {
   public:
 	// Contructor
-	FElementID(const FSyncContext& InSyncContext)
-		: SyncContext(InSyncContext)
-		, Index3D(0)
-		, SyncData(nullptr)
-	{
-	}
+	FElementID(const FSyncContext& InSyncContext);
 
 	// Initialize with 3D element
-	void InitElement(GS::Int32 InIndex3d)
-	{
-		Index3D = InIndex3d;
-		SyncContext.GetModel().GetElement(Index3D, &Element3D);
-		ElementHeader.guid = APINULLGuid;
-	}
+	void InitElement(GS::Int32 InIndex3D);
 
-	// Initialize with 3D element
-	void InitElement(FSyncData* IOSyncdata)
-	{
-		UE_AC_TestPtr(IOSyncdata);
-		SyncData = IOSyncdata;
-		Index3D = IOSyncdata->GetIndex3D();
-		if (Index3D > 0)
-		{
-			SyncContext.GetModel().GetElement(Index3D, &Element3D);
-		}
-	}
+	// Initialize with sync data
+	void InitElement(FSyncData* IOSyncdata);
 
 	// Return true if object is valid (i.e. not recently deleted)
 	bool IsInvalid() const { return Element3D.IsInvalid(); }
 
+	// Return the element index (in 3D list)
+	GS::Int32 GetIndex3D() const { return Index3D; }
+
+	// Return the 3D element
+	const ModelerAPI::Element& GetElement3D() const { return Element3D; }
+
+	// Return the 3D type name
 	static const utf8_t* GetTypeName(ModelerAPI::Element::Type InType)
 	{
 		return TAssEnumName< ModelerAPI::Element::Type >::GetName(InType);
 	}
 
+	// Return element 3D type name
 	const utf8_t* GetTypeName() const { return GetTypeName(Element3D.GetType()); }
 
 	// Initialize element header from 3D element
-	bool InitHeader()
-	{
-		if (IsInvalid())
-		{
-			throw std::runtime_error(
-				Utf8StringFormat("FElementID::InitHeader - Invalid element for index=%d\n", Index3D).c_str());
-		}
-		Zap(&ElementHeader);
-		ElementHeader.guid = GSGuid2APIGuid(Element3D.GetElemGuid());
-		GSErrCode GSErr = ACAPI_Element_GetHeader(&ElementHeader, 0);
-		if (GSErr != NoError)
-		{
-			utf8_string ErrorName(GetErrorName(GSErr));
-			utf8_string TypeName(GetTypeName());
-			UE_AC_DebugF("Error \"%s\" with element %d {%s} Type=%s\n", ErrorName.c_str(), Index3D,
-						 Element3D.GetElemGuid().ToUniString().ToUtf8(), TypeName.c_str());
-			if (GSErr != APIERR_BADID)
-			{
-				UE_AC::ThrowGSError(GSErr, __FILE__, __LINE__);
-			}
-			return false;
-		}
-		return true;
-	}
+	bool InitHeader();
 
 	// Initialize element header with element guid
 	void InitHeader(const API_Guid& InGuid)
 	{
-		Zap(&ElementHeader);
-		ElementHeader.guid = InGuid;
-		UE_AC_TestGSError(ACAPI_Element_GetHeader(&ElementHeader, 0));
+		bFullElementFetched = false;
+		Zap(&APIElement.header);
+		APIElement.header.guid = InGuid;
+		UE_AC_TestGSError(ACAPI_Element_GetHeader(&APIElement.header, 0));
 	}
 
+	// Set the Sync Data associated to the current element
+	void SetSyncData(FSyncData* InSyncData) { SyncData = InSyncData; }
+
+	// Return the Sync Data associated to the current element
+	FSyncData* GetSyncData() const { return SyncData; }
+
+	// Connect to parent or childs
 	void HandleDepedencies() const;
 
 	// Return true if element is a morph type body (will need double side)
 	bool IsSurface() const;
 
-	const FSyncContext& SyncContext; // Current synchronisation context
-	FSyncData*			SyncData;
-	GS::Int32			Index3D; // 3D element index
-	ModelerAPI::Element Element3D; // 3D element
-	API_Elem_Head		ElementHeader; // AC element header
+	// Return the element's header
+	const API_Elem_Head& GetHeader() const { return APIElement.header; }
+
+	// Return the complete element
+	const API_Element& GetAPIElement();
+
+	// Return the lib part info if this element come from it
+	const FLibPartInfo* GetLibPartInfo();
+
+	// Current synchronisation context
+	const FSyncContext& SyncContext;
+
   private:
 	void CollectDependantElementsType(API_ElemTypeID TypeID) const;
+
+	// 3D element index
+	GS::Int32 Index3D;
+	// 3D element
+	ModelerAPI::Element Element3D;
+
+	// APIElement contain all values (by opposition to header only)
+	bool bFullElementFetched;
+	// AC API Element
+	API_Element APIElement;
+
+	// Sync data associated to the current element
+	FSyncData* SyncData;
+
+	// Lib part info have been fetched
+	bool bLibPartInfoFetched;
+	// LibPartInfo != nullptr if element come from it
+	const FLibPartInfo* LibPartInfo;
 };
 
 class FSyncData::FProcessInfo
diff --git a/Engine/Source/Programs/Enterprise/Datasmith/DatasmithARCHICADExporter/Private/Export.cpp b/Engine/Source/Programs/Enterprise/Datasmith/DatasmithARCHICADExporter/Private/Export.cpp
--- a/Engine/Source/Programs/Enterprise/Datasmith/DatasmithARCHICADExporter/Private/Export.cpp
+++ b/Engine/Source/Programs/Enterprise/Datasmith/DatasmithARCHICADExporter/Private/Export.cpp
@@ -35,7 +35,7 @@ static GSErrCode __ACENV_CALL SaveToDatasmithFile(const API_IOParams* IOParams,
 
 GSErrCode FExport::Register()
 {
-	return ACAPI_Register_FileType(kDatasmithFileRefCon, kStrFileType, kStrFileCreator, StrFileExtension, kIconDS,
+	return ACAPI_Register_FileType(kDatasmithFileRefCon, kStrFileType, kStrFileCreator, StrFileExtension, kIconDSFile,
 								   LocalizeResId(kStrListFileTypes), 0, SaveAs3DSupported);
 }
 
diff --git a/Engine/Source/Programs/Enterprise/Datasmith/DatasmithARCHICADExporter/Private/Exporter.cpp b/Engine/Source/Programs/Enterprise/Datasmith/DatasmithARCHICADExporter/Private/Exporter.cpp
--- a/Engine/Source/Programs/Enterprise/Datasmith/DatasmithARCHICADExporter/Private/Exporter.cpp
+++ b/Engine/Source/Programs/Enterprise/Datasmith/DatasmithARCHICADExporter/Private/Exporter.cpp
@@ -107,7 +107,8 @@ void FExporter::DoExport(const ModelerAPI::Model& InModel, const IO::Location& I
 	FProgression Progression(kStrListProgression, kExportTitle, kNbPhases, FProgression::kThrowOnCancel,
 							 &OutUserCancelled);
 
-	FSyncDatabase SyncDatabase(*FilePath, *LabelString, SceneExporter.GetAssetsOutputPath());
+	FSyncDatabase SyncDatabase(*FilePath, *LabelString, SceneExporter.GetAssetsOutputPath(),
+							   FSyncDatabase::GetCachePath());
 
 	FSyncContext SyncContext(false, InModel, SyncDatabase, &Progression);
 
@@ -139,9 +140,10 @@ GSErrCode FExporter::DoChooseDestination(IO::Location* OutDestFile)
 	// IO::fileSystem.GetSpecialLocation(IO::FileSystem::CurrentFolder, OutDestFile);
 
 	FTM::FileTypeManager templateFileFTM("TemplateFileFTM");
-	FTM::TypeID datasmithTypeID = templateFileFTM.AddType(FTM::FileType("Datasmith file", "udatasmith", 0, 0, kIconDS));
+	FTM::TypeID			 datasmithTypeID = templateFileFTM.AddType(
+		 FTM::FileType(GetStdName(kName_DatasmithFileTypeName), "udatasmith", 0, 0, kIconDSFile));
 
-	fileDialog.SetTitle("Export Datasmith File");
+	fileDialog.SetTitle(GetGSName(kName_ExportToDatasmithFile));
 	fileDialog.AddFilter(datasmithTypeID);
 	fileDialog.AddFilter(FTM::RootGroup);
 	fileDialog.SelectFilter(0);
diff --git a/Engine/Source/Programs/Enterprise/Datasmith/DatasmithARCHICADExporter/Private/LoadDatasmithDlls.cpp b/Engine/Source/Programs/Enterprise/Datasmith/DatasmithARCHICADExporter/Private/LoadDatasmithDlls.cpp
--- a/Engine/Source/Programs/Enterprise/Datasmith/DatasmithARCHICADExporter/Private/LoadDatasmithDlls.cpp
+++ b/Engine/Source/Programs/Enterprise/Datasmith/DatasmithARCHICADExporter/Private/LoadDatasmithDlls.cpp
@@ -31,10 +31,12 @@ DISABLE_SDK_WARNINGS_END
 #endif
 
 #include "Utils/APIEnvir.h"
+
+DISABLE_SDK_WARNINGS_START
 #include "ACAPinc.h"
-//#include "ModelColor.hpp"
 #include "Md5.hpp"
 #include "Lock.hpp"
+DISABLE_SDK_WARNINGS_END
 
 #include "Utils/LocalizeTools.h"
 #include "Utils/AddonTools.h"
diff --git a/Engine/Source/Programs/Enterprise/Datasmith/DatasmithARCHICADExporter/Private/MaterialsDatabase.cpp b/Engine/Source/Programs/Enterprise/Datasmith/DatasmithARCHICADExporter/Private/MaterialsDatabase.cpp
--- a/Engine/Source/Programs/Enterprise/Datasmith/DatasmithARCHICADExporter/Private/MaterialsDatabase.cpp
+++ b/Engine/Source/Programs/Enterprise/Datasmith/DatasmithARCHICADExporter/Private/MaterialsDatabase.cpp
@@ -222,29 +222,33 @@ void FMaterialInfo::UpdateMaterial(const TSharedPtr< IDatasmithUEPbrMaterialElem
 			TextureCoordinateExpression->SetCoordinateIndex(0);
 			TextureCoordinateExpression->ConnectExpression(*UVEditExpression->GetInput(0));
 		}
-		if (Texture->bHasAlpha && Texture->bAlphaIsTransparence && bIsTransparent)
+		if (Texture->bHasAlpha && Texture->bAlphaIsTransparence)
 		{
-			IDatasmithMaterialExpressionGeneric* MultiplyExpression =
-				DSMaterial->AddMaterialExpression< IDatasmithMaterialExpressionGeneric >();
-			MultiplyExpression->SetExpressionName(TEXT("Multiply"));
-			MultiplyExpression->SetName(TEXT("Multiply Expression"));
+			if (bIsTransparent)
+			{
+				IDatasmithMaterialExpressionGeneric* MultiplyExpression =
+					DSMaterial->AddMaterialExpression< IDatasmithMaterialExpressionGeneric >();
+				MultiplyExpression->SetExpressionName(TEXT("Multiply"));
+				MultiplyExpression->SetName(TEXT("Multiply Expression"));
 
-			IDatasmithMaterialExpressionScalar* OpacityExpression =
-				DSMaterial->AddMaterialExpression< IDatasmithMaterialExpressionScalar >();
-			OpacityExpression->GetScalar() = Opacity;
-			OpacityExpression->SetName(TEXT("Opacity"));
+				IDatasmithMaterialExpressionScalar* OpacityExpression =
+					DSMaterial->AddMaterialExpression< IDatasmithMaterialExpressionScalar >();
+				OpacityExpression->GetScalar() = Opacity;
+				OpacityExpression->SetName(TEXT("Opacity"));
 
-			IDatasmithExpressionInput* MultiplyInputA = MultiplyExpression->GetInput(0);
-			IDatasmithExpressionInput* MultiplyInputB = MultiplyExpression->GetInput(1);
+				IDatasmithExpressionInput* MultiplyInputA = MultiplyExpression->GetInput(0);
+				IDatasmithExpressionInput* MultiplyInputB = MultiplyExpression->GetInput(1);
 
-			MultiplyExpression->ConnectExpression(DSMaterial->GetOpacity());
+				MultiplyExpression->ConnectExpression(DSMaterial->GetOpacity());
 
-			BaseTextureExpression->ConnectExpression(*MultiplyInputA, 4);
-			OpacityExpression->ConnectExpression(*MultiplyInputB);
-		}
-		else if (Texture->bHasAlpha && Texture->bAlphaIsTransparence)
-		{
-			BaseTextureExpression->ConnectExpression(DSMaterial->GetOpacity(), 4);
+				BaseTextureExpression->ConnectExpression(*MultiplyInputA, 4);
+				OpacityExpression->ConnectExpression(*MultiplyInputB);
+			}
+			else
+			{
+				BaseTextureExpression->ConnectExpression(DSMaterial->GetOpacity(), 4);
+			}
+			bHasAphaMask = true;
 		}
 	}
 	else
@@ -322,9 +326,10 @@ bool FMaterialsDatabase::CheckModify()
 {
 	FAutoChangeDatabase changeDB(APIWind_3DModelID);
 
-	for (TPair< FMaterialKey, FMaterialSyncData >& Iter : MapMaterials)
+	for (TPair< FMaterialKey, TUniquePtr< FMaterialSyncData > >& Iter : MapMaterials)
 	{
-		if (Iter.Value.CheckModify(Iter.Key))
+		UE_AC_Assert(Iter.Value.IsValid());
+		if (Iter.Value->CheckModify(Iter.Key))
 		{
 			return true;
 		}
@@ -335,14 +340,15 @@ bool FMaterialsDatabase::CheckModify()
 // Scan all material and update modified ones
 void FMaterialsDatabase::UpdateModified(const FSyncContext& SyncContext)
 {
-	for (TPair< FMaterialKey, FMaterialSyncData >& Iter : MapMaterials)
+	for (TPair< FMaterialKey, TUniquePtr< FMaterialSyncData > >& Iter : MapMaterials)
 	{
-		Iter.Value.Update(SyncContext, Iter.Key);
+		UE_AC_Assert(Iter.Value.IsValid());
+		Iter.Value->Update(SyncContext, Iter.Key);
 	}
 }
 
 // Return true if the material have been modified
-bool FMaterialsDatabase::FMaterialSyncData::CheckModify(const FMaterialKey& MaterialKey)
+bool FMaterialsDatabase::FMaterialSyncData::CheckModify(const FMaterialKey& /* MaterialKey */)
 {
 	if (!bIsDuplicate && !bIdIsSynthetized && MaterialIndex != kInvalidMaterialIndex)
 	{
@@ -390,13 +396,15 @@ const FMaterialsDatabase::FMaterialSyncData& FMaterialsDatabase::GetMaterial(con
 	}
 	FMaterialKey MaterialKey(inACMaterialIndex, inACTextureIndex, InSided);
 
-	FMaterialsDatabase::FMaterialSyncData& material = MapMaterials.FindOrAdd(MaterialKey);
-	if (!material.bIsInitialized)
+	TUniquePtr< FMaterialSyncData >& material = MapMaterials.FindOrAdd(MaterialKey);
+	if (!material.IsValid())
 	{
-		material.Init(SyncContext, MaterialKey);
+		FMaterialSyncData* MatSyncData = new FMaterialSyncData();
+		material.Reset(MatSyncData);
+		MatSyncData->Init(SyncContext, MaterialKey);
 	}
 
-	return material;
+	return *material;
 }
 
 void FMaterialsDatabase::FMaterialSyncData::Init(const FSyncContext& SyncContext, const FMaterialKey& MaterialKey)
diff --git a/Engine/Source/Programs/Enterprise/Datasmith/DatasmithARCHICADExporter/Private/MaterialsDatabase.h b/Engine/Source/Programs/Enterprise/Datasmith/DatasmithARCHICADExporter/Private/MaterialsDatabase.h
--- a/Engine/Source/Programs/Enterprise/Datasmith/DatasmithARCHICADExporter/Private/MaterialsDatabase.h
+++ b/Engine/Source/Programs/Enterprise/Datasmith/DatasmithARCHICADExporter/Private/MaterialsDatabase.h
@@ -105,7 +105,7 @@ class FMaterialsDatabase
 										 GS::Int32 inACTextureIndex, ESided InSided);
 
   private:
-	typedef TMap< FMaterialKey, FMaterialSyncData > MapSyncData; // Map Material key to Material sync data
+	typedef TMap< FMaterialKey, TUniquePtr< FMaterialSyncData > > MapSyncData; // Map Material key to Material sync data
 
 	typedef TSet< FString > SetMaterialsNames;
 
diff --git a/Engine/Source/Programs/Enterprise/Datasmith/DatasmithARCHICADExporter/Private/Menus.cpp b/Engine/Source/Programs/Enterprise/Datasmith/DatasmithARCHICADExporter/Private/Menus.cpp
--- a/Engine/Source/Programs/Enterprise/Datasmith/DatasmithARCHICADExporter/Private/Menus.cpp
+++ b/Engine/Source/Programs/Enterprise/Datasmith/DatasmithARCHICADExporter/Private/Menus.cpp
@@ -28,15 +28,17 @@ GSErrCode FMenus::Register()
 
 	RegisterMenuAndHelp(&GSErr, kStrListMenuDatasmith, MenuCode_Palettes, MenuFlag_Default);
 
+#if !UE_AC_NO_MENU
 	RegisterMenuAndHelp(&GSErr, kStrListMenuItemSnapshot, MenuCode_UserDef, MenuFlag_Default);
-#if AUTO_SYNC
+	#if AUTO_SYNC
 	RegisterMenuAndHelp(&GSErr, kStrListMenuItemAutoSync, MenuCode_UserDef, MenuFlag_Default);
-#endif
+	#endif
 	RegisterMenuAndHelp(&GSErr, kStrListMenuItemConnections, MenuCode_UserDef, MenuFlag_Default);
 	RegisterMenuAndHelp(&GSErr, kStrListMenuItemExport, MenuCode_UserDef, MenuFlag_Default);
 	RegisterMenuAndHelp(&GSErr, kStrListMenuItemMessages, MenuCode_UserDef, MenuFlag_Default);
 	RegisterMenuAndHelp(&GSErr, kStrListMenuItemPalette, MenuCode_UserDef, MenuFlag_SeparatorBefore);
 	RegisterMenuAndHelp(&GSErr, kStrListMenuItemAbout, MenuCode_UserDef, MenuFlag_Default);
+#endif
 
 	return GSErr;
 }
@@ -45,6 +47,8 @@ GSErrCode FMenus::Register()
 GSErrCode FMenus::Initialize()
 {
 	GSErrCode GSErr = NoError;
+
+#if !UE_AC_NO_MENU
 	for (short IndexMenu = kStrListMenuItemSnapshot; IndexMenu <= kStrListMenuItemAbout && GSErr == NoError;
 		 IndexMenu++)
 	{
@@ -54,17 +58,21 @@ GSErrCode FMenus::Initialize()
 	{
 		UE_AC_DebugF("FMenus::Initialize - ACAPI_Install_MenuHandler error=%s\n", GetErrorName(GSErr));
 	}
+#endif
+
 	GSErr = ACAPI_Install_MenuHandler(LocalizeResId(kStrListMenuDatasmith), MenuCommandHandler);
 	if (GSErr != NoError)
 	{
 		UE_AC_DebugF("FMenus::Initialize - ACAPI_Install_MenuHandler error=%s\n", GetErrorName(GSErr));
 	}
+
 	return GSErr;
 }
 
 // Ename or disable menu item
 void FMenus::SetMenuItemStatus(short InMenu, short InItem, bool InSet, GSFlags InFlag)
 {
+#if !UE_AC_NO_MENU
 	API_MenuItemRef ItemRef;
 	Zap(&ItemRef);
 	ItemRef.menuResID = LocalizeResId(InMenu);
@@ -84,18 +92,25 @@ void FMenus::SetMenuItemStatus(short InMenu, short InItem, bool InSet, GSFlags I
 		GSErr = ACAPI_Interface(APIIo_SetMenuItemFlagsID, &ItemRef, &ItemFlags);
 		if (GSErr != NoError)
 		{
-			UE_AC_TraceF("FMenus::SetMenuItemStatus - APIIo_SetMenuItemFlagsID error=%s\n", GetErrorName(GSErr));
+			UE_AC_DebugF("FMenus::SetMenuItemStatus - APIIo_SetMenuItemFlagsID error=%s\n", GetErrorName(GSErr));
 		}
 	}
 	else
 	{
-		UE_AC_TraceF("FMenus::SetMenuItemStatus - APIIo_GetMenuItemFlagsID error=%s\n", GetErrorName(GSErr));
+		UE_AC_DebugF("FMenus::SetMenuItemStatus - APIIo_GetMenuItemFlagsID error=%s\n", GetErrorName(GSErr));
 	}
+#else
+	(void)InMenu; // No unused warnings
+	(void)InItem;
+	(void)InSet;
+	(void)InFlag;
+#endif
 }
 
 // Change the text of an item
 void FMenus::SetMenuItemText(short InMenu, short InItem, const GS::UniString& ItemStr)
 {
+#if !UE_AC_NO_MENU
 	API_MenuItemRef ItemRef;
 	Zap(&ItemRef);
 	ItemRef.menuResID = LocalizeResId(InMenu);
@@ -104,15 +119,20 @@ void FMenus::SetMenuItemText(short InMenu, short InItem, const GS::UniString& It
 		ACAPI_Interface(APIIo_SetMenuItemTextID, &ItemRef, nullptr, const_cast< GS::UniString* >(&ItemStr));
 	if (GSErr != NoError)
 	{
-		UE_AC_TraceF("FMenus::SetMenuItemText - APIIo_SetMenuItemTextID error=%s\n", GetErrorName(GSErr));
+		UE_AC_DebugF("FMenus::SetMenuItemText - APIIo_SetMenuItemTextID error=%s\n", GetErrorName(GSErr));
 	}
+#else
+	(void)InMenu; // No unused warnings
+	(void)InItem;
+	(void)ItemStr;
+#endif
 }
 
 // AutoSync status changed
 void FMenus::AutoSyncChanged()
 {
-	GS::UniString StartAutoSync("Start Auto Sync");
-	GS::UniString PauseAutoSync("Pause Auto Sync");
+	static const GS::UniString StartAutoSync(GetGSName(kName_StartAutoSync));
+	static const GS::UniString PauseAutoSync(GetGSName(kName_PauseAutoSync));
 #if AUTO_SYNC
 	SetMenuItemText(kStrListMenuItemAutoSync, 1, FCommander::IsAutoSyncEnabled() ? PauseAutoSync : StartAutoSync);
 	SetMenuItemStatus(kStrListMenuItemAutoSync, 1, FCommander::IsAutoSyncEnabled(), API_MenuItemChecked);
diff --git a/Engine/Source/Programs/Enterprise/Datasmith/DatasmithARCHICADExporter/Private/Menus.h b/Engine/Source/Programs/Enterprise/Datasmith/DatasmithARCHICADExporter/Private/Menus.h
--- a/Engine/Source/Programs/Enterprise/Datasmith/DatasmithARCHICADExporter/Private/Menus.h
+++ b/Engine/Source/Programs/Enterprise/Datasmith/DatasmithARCHICADExporter/Private/Menus.h
@@ -6,6 +6,8 @@
 
 BEGIN_NAMESPACE_UE_AC
 
+#define UE_AC_NO_MENU 1 // 1 = Dont create menu
+
 class FMenus
 {
   public:
diff --git a/Engine/Source/Programs/Enterprise/Datasmith/DatasmithARCHICADExporter/Private/MetaData.cpp b/Engine/Source/Programs/Enterprise/Datasmith/DatasmithARCHICADExporter/Private/MetaData.cpp
--- a/Engine/Source/Programs/Enterprise/Datasmith/DatasmithARCHICADExporter/Private/MetaData.cpp
+++ b/Engine/Source/Programs/Enterprise/Datasmith/DatasmithARCHICADExporter/Private/MetaData.cpp
@@ -106,13 +106,14 @@ void FMetaData::AddProperty(const TCHAR* InPropKey, EDatasmithKeyValuePropertyTy
 	MetaData->AddProperty(MetaDataProperty);
 }
 
-void FMetaData::ExportMetaData(const GS::Guid& InElementId)
+void FMetaData::ExportMetaData(const GS::Guid& InElementID)
 {
-	UE_AC_Assert(FString::Printf(TEXT("MetaData_%s"), GSStringToUE(InElementId.ToUniString())) == MetaData->GetName());
-	const API_Guid& ElementId(GSGuid2APIGuid(InElementId));
+	UE_AC_Assert(FString::Printf(TEXT("MetaData_%s"), GSStringToUE(InElementID.ToUniString())) == MetaData->GetName());
+	const API_Guid& ElementId(GSGuid2APIGuid(InElementID));
 	ExportElementIDProperty(ElementId);
 	ExportClassifications(ElementId);
 	ExportCategories(ElementId);
+	ExportIFCType(ElementId);
 	ExportIFCProperties(ElementId);
 	ExportIFCAttributes(ElementId);
 }
@@ -185,8 +186,8 @@ void FMetaData::ExportClassifications(const API_Guid& ElementId)
 	}
 	else
 	{
-		UE_AC_DebugF("FMetaData::ExportClassifications - FElementTools::GetElementClassifications returned error %d\n",
-					 GSErr);
+		UE_AC_DebugF("FMetaData::ExportClassifications - FElementTools::GetElementClassifications returned error %s\n",
+					 GetErrorName(GSErr));
 	}
 }
 
@@ -212,15 +213,40 @@ void FMetaData::ExportCategories(const API_Guid& ElementId)
 			{
 				if (GSErr != APIERR_BADPARS)
 				{
-					UE_AC_DebugF("FMetaData::ExportCategories - ACAPI_Element_GetCategoryValue returned error %d\n",
-								 GSErr);
+					UE_AC_DebugF("FMetaData::ExportCategories - ACAPI_Element_GetCategoryValue returned error %s\n",
+								 GetErrorName(GSErr));
 				}
 			}
 		}
 	}
 	else
 	{
-		UE_AC_DebugF("FMetaData::ExportCategories - APIDb_GetElementCategoriesID returned error %d\n", GSErr);
+		UE_AC_DebugF("FMetaData::ExportCategories - APIDb_GetElementCategoriesID returned error %s\n",
+					 GetErrorName(GSErr));
+	}
+}
+
+void FMetaData::ExportIFCType(const API_Guid& ElementId)
+{
+	GS::UniString IfcType;
+	GS::UniString TypeObjectIfcType;
+	GSErrCode	  GSErr = ACAPI_Element_GetIFCType(ElementId, &IfcType, &TypeObjectIfcType);
+	if (GSErr == NoError)
+	{
+		if (!IfcType.IsEmpty())
+		{
+			static const GS::UniString KeyIfcType("IFC_Type");
+			AddMetaDataProperty(API_PropertyStringValueType, KeyIfcType, IfcType);
+		}
+		if (!TypeObjectIfcType.IsEmpty())
+		{
+			static const GS::UniString KeyTypeObjectIfcType("IFC_ObjectType");
+			AddMetaDataProperty(API_PropertyStringValueType, KeyTypeObjectIfcType, TypeObjectIfcType);
+		}
+	}
+	else
+	{
+		UE_AC_DebugF("FMetaData::ExportIFCType - ACAPI_Element_GetIFCType returned error %s", GetErrorName(GSErr));
 	}
 }
 
@@ -292,7 +318,8 @@ void FMetaData::ExportIFCProperties(const API_Guid& ElementId)
 	{
 		if (GSErr != APIERR_BADPARS)
 		{
-			UE_AC_DebugF("FMetaData::ExportIFCProperties - ACAPI_Element_GetIFCProperties returned error %d\n", GSErr);
+			UE_AC_DebugF("FMetaData::ExportIFCProperties - ACAPI_Element_GetIFCProperties returned error %s\n",
+						 GetErrorName(GSErr));
 		}
 	}
 }
@@ -315,7 +342,8 @@ void FMetaData::ExportIFCAttributes(const API_Guid& ElementId)
 	{
 		if (GSErr != APIERR_BADPARS)
 		{
-			UE_AC_DebugF("FMetaData::ExportIFCAttributes - ACAPI_Element_GetIFCAttributes returned error %d\n", GSErr);
+			UE_AC_DebugF("FMetaData::ExportIFCAttributes - ACAPI_Element_GetIFCAttributes returned error %s\n",
+						 GetErrorName(GSErr));
 		}
 	}
 }
@@ -380,7 +408,8 @@ void FMetaData::ExportProperties(const API_Guid& ElementId)
 	}
 	else
 	{
-		UE_AC_DebugF("FMetaData::ExportProperties - FElementTools::GetElementProperties returned error %d\n", GSErr);
+		UE_AC_DebugF("FMetaData::ExportProperties - FElementTools::GetElementProperties returned error %s\n",
+					 GetErrorName(GSErr));
 	}
 }
 
diff --git a/Engine/Source/Programs/Enterprise/Datasmith/DatasmithARCHICADExporter/Private/MetaData.h b/Engine/Source/Programs/Enterprise/Datasmith/DatasmithARCHICADExporter/Private/MetaData.h
--- a/Engine/Source/Programs/Enterprise/Datasmith/DatasmithARCHICADExporter/Private/MetaData.h
+++ b/Engine/Source/Programs/Enterprise/Datasmith/DatasmithARCHICADExporter/Private/MetaData.h
@@ -9,18 +9,18 @@ BEGIN_NAMESPACE_UE_AC
 class FMetaData
 {
   public:
-	//	FMetaData(const GS::Guid& InElementId);
+	//	FMetaData(const GS::Guid& InElementID);
 
 	FMetaData(const TSharedPtr< IDatasmithElement >& InElement);
 
-	void SetAssociatedElement(const GS::Guid& /* InElementId */, const TSharedPtr< IDatasmithElement >& InElement)
+	void SetAssociatedElement(const GS::Guid& /* InElementID */, const TSharedPtr< IDatasmithElement >& InElement)
 	{
 		MetaData->SetAssociatedElement(InElement);
 	}
 
 	void SetOrUpdate(TSharedPtr< IDatasmithMetaDataElement >* IOPtr, IDatasmithScene* IOScene);
 
-	void ExportMetaData(const GS::Guid& InElementId);
+	void ExportMetaData(const GS::Guid& InElementID);
 
 	const TSharedRef< IDatasmithMetaDataElement >& GetMetaData() const { return MetaData; }
 
@@ -52,6 +52,8 @@ class FMetaData
 
 	void ExportCategories(const API_Guid& InlementId);
 
+	void ExportIFCType(const API_Guid& InlementId);
+
 	void ExportIFCProperties(const API_Guid& InlementId);
 
 	void ExportIFCAttributes(const API_Guid& InlementId);
diff --git a/Engine/Source/Programs/Enterprise/Datasmith/DatasmithARCHICADExporter/Private/Options.cpp b/Engine/Source/Programs/Enterprise/Datasmith/DatasmithARCHICADExporter/Private/Options.cpp
--- a/Engine/Source/Programs/Enterprise/Datasmith/DatasmithARCHICADExporter/Private/Options.cpp
+++ b/Engine/Source/Programs/Enterprise/Datasmith/DatasmithARCHICADExporter/Private/Options.cpp
@@ -45,7 +45,10 @@ size_t FOptions::ReadFrom(CReader* IOReader)
 
 const GS::UniString& FOptions::TypeOptionsName(ETypeOptions InType)
 {
-	return GetGSName(InType == kSync ? kName_SyncOptions : kName_ExportOptions);
+	static const GS::UniString SyncOptions("SyncOptions");
+	static const GS::UniString ExportOptions("ExportOptions");
+
+	return InType == kSync ? SyncOptions : ExportOptions;
 }
 
 // Get saved options of the last sync
diff --git a/Engine/Source/Programs/Enterprise/Datasmith/DatasmithARCHICADExporter/Private/Progression.cpp b/Engine/Source/Programs/Enterprise/Datasmith/DatasmithARCHICADExporter/Private/Progression.cpp
deleted file mode 100644
diff --git a/Engine/Source/Programs/Enterprise/Datasmith/DatasmithARCHICADExporter/Private/Progression.h b/Engine/Source/Programs/Enterprise/Datasmith/DatasmithARCHICADExporter/Private/Progression.h
deleted file mode 100644
diff --git a/Engine/Source/Programs/Enterprise/Datasmith/DatasmithARCHICADExporter/Private/ProjectEvent.cpp b/Engine/Source/Programs/Enterprise/Datasmith/DatasmithARCHICADExporter/Private/ProjectEvent.cpp
--- a/Engine/Source/Programs/Enterprise/Datasmith/DatasmithARCHICADExporter/Private/ProjectEvent.cpp
+++ b/Engine/Source/Programs/Enterprise/Datasmith/DatasmithARCHICADExporter/Private/ProjectEvent.cpp
@@ -75,7 +75,7 @@ GSErrCode FProjectEvent::Initialize()
 GSErrCode FProjectEvent::Event(API_NotifyEventID NotifID, Int32 Param)
 {
 	(void)Param;
-	UE_AC_VerboseF("--- FProjectEvent::Event(%s, %d)\n", TAssEnumName< API_NotifyEventID >::GetName(NotifID), Param);
+	UE_AC_VerboseF("-> FProjectEvent::Event(%s, %d)\n", TAssEnumName< API_NotifyEventID >::GetName(NotifID), Param);
 	FSynchronizer* Synchronizer = FSynchronizer::GetCurrent();
 	switch (NotifID)
 	{
@@ -113,6 +113,7 @@ GSErrCode FProjectEvent::Event(API_NotifyEventID NotifID, Int32 Param)
 			break;
 	}
 
+	UE_AC_VerboseF("<- FProjectEvent::Event\n");
 	return NoError;
 }
 
diff --git a/Engine/Source/Programs/Enterprise/Datasmith/DatasmithARCHICADExporter/Private/ReportWindow.cpp b/Engine/Source/Programs/Enterprise/Datasmith/DatasmithARCHICADExporter/Private/ReportWindow.cpp
--- a/Engine/Source/Programs/Enterprise/Datasmith/DatasmithARCHICADExporter/Private/ReportWindow.cpp
+++ b/Engine/Source/Programs/Enterprise/Datasmith/DatasmithARCHICADExporter/Private/ReportWindow.cpp
@@ -39,7 +39,7 @@ class FReportDialog : public DG::Palette,
 	FReportDialog();
 	~FReportDialog();
 
-	virtual void PanelClosed(const DG::PanelCloseEvent& ev) override {}
+	virtual void PanelClosed(const DG::PanelCloseEvent& /* ev */) override {}
 
 	virtual void PanelCloseRequested(const DG::PanelCloseRequestEvent& /* ev */, bool* /* accepted */) override
 	{
diff --git a/Engine/Source/Programs/Enterprise/Datasmith/DatasmithARCHICADExporter/Private/ResourcesIDs.h b/Engine/Source/Programs/Enterprise/Datasmith/DatasmithARCHICADExporter/Private/ResourcesIDs.h
--- a/Engine/Source/Programs/Enterprise/Datasmith/DatasmithARCHICADExporter/Private/ResourcesIDs.h
+++ b/Engine/Source/Programs/Enterprise/Datasmith/DatasmithARCHICADExporter/Private/ResourcesIDs.h
@@ -18,6 +18,7 @@
 	#define kIconAutoSync 32514
 #endif
 #define kIconDS 32600
+#define kIconDSFile 32601
 
 /* Localizable resources */
 #define kStrListFileTypes 302
diff --git a/Engine/Source/Programs/Enterprise/Datasmith/DatasmithARCHICADExporter/Private/SyncContext.h b/Engine/Source/Programs/Enterprise/Datasmith/DatasmithARCHICADExporter/Private/SyncContext.h
--- a/Engine/Source/Programs/Enterprise/Datasmith/DatasmithARCHICADExporter/Private/SyncContext.h
+++ b/Engine/Source/Programs/Enterprise/Datasmith/DatasmithARCHICADExporter/Private/SyncContext.h
@@ -6,7 +6,7 @@
 
 #include "SyncDatabase.h"
 #include "SyncData.h"
-#include "Progression.h"
+#include "Utils/Progression.h"
 #include "StatsCounter.h"
 
 class FDatasmithDirectLink;
diff --git a/Engine/Source/Programs/Enterprise/Datasmith/DatasmithARCHICADExporter/Private/SyncData.cpp b/Engine/Source/Programs/Enterprise/Datasmith/DatasmithARCHICADExporter/Private/SyncData.cpp
--- a/Engine/Source/Programs/Enterprise/Datasmith/DatasmithARCHICADExporter/Private/SyncData.cpp
+++ b/Engine/Source/Programs/Enterprise/Datasmith/DatasmithARCHICADExporter/Private/SyncData.cpp
@@ -11,12 +11,15 @@
 #include "GeometryUtil.h"
 #include "Utils/AutoChangeDatabase.h"
 #include "Utils/TimeStat.h"
+#include "TexturesCache.h"
 
 DISABLE_SDK_WARNINGS_START
 #include "Transformation.hpp"
 #include "Line3D.hpp"
 DISABLE_SDK_WARNINGS_END
 
+#include <stdexcept>
+
 BEGIN_NAMESPACE_UE_AC
 
 #define TRACE_ATTACH_OBSERVERS 0
@@ -45,38 +48,37 @@ FSyncData::~FSyncData()
 }
 
 // Update data from a 3d element
-void FSyncData::Update(const FElementID& InElementId)
+void FSyncData::Update(const FElementID& InElementID)
 {
-	UE_AC_Assert(ElementId == APIGuid2GSGuid(InElementId.ElementHeader.guid));
-	UE_AC_Assert(Index3D == 0 && InElementId.Index3D != 0);
+	const API_Elem_Head& Header = InElementID.GetHeader();
+	UE_AC_Assert(ElementId == APIGuid2GSGuid(Header.guid));
+	UE_AC_Assert(Index3D == 0 && InElementID.GetIndex3D() != 0);
 
-	Index3D = InElementId.Index3D;
-	if (GenId != InElementId.Element3D.GetGenId())
+	Index3D = InElementID.GetIndex3D();
+	if (GenId != InElementID.GetElement3D().GetGenId())
 	{
-		GenId = InElementId.Element3D.GetGenId();
+		GenId = InElementID.GetElement3D().GetGenId();
 		bIsModified = true;
 	}
 
 	// If AC element has been modified, recheck connections
-	if (ModificationStamp != InElementId.ElementHeader.modiStamp)
+	if (ModificationStamp != Header.modiStamp)
 	{
-		if (ModificationStamp > InElementId.ElementHeader.modiStamp)
+		if (ModificationStamp > Header.modiStamp)
 		{
 			UE_AC_DebugF("FSyncData::Update {%s} New stamp younger: %lld, %lld\n",
-						 APIGuidToString(InElementId.ElementHeader.guid).ToUtf8(), ModificationStamp,
-						 InElementId.ElementHeader.modiStamp);
+						 APIGuidToString(Header.guid).ToUtf8(), ModificationStamp, Header.modiStamp);
 		}
-		ModificationStamp = InElementId.ElementHeader.modiStamp;
-		InElementId.HandleDepedencies();
+		ModificationStamp = Header.modiStamp;
+		InElementID.HandleDepedencies();
 		bIsModified = true;
 	}
 	if (ModificationStamp == 0)
 	{
-		UE_AC_DebugF("FSyncData::Update {%s} ModificationStamp == 0\n",
-					 APIGuidToString(InElementId.ElementHeader.guid).ToUtf8());
+		UE_AC_DebugF("FSyncData::Update {%s} ModificationStamp == 0\n", APIGuidToString(Header.guid).ToUtf8());
 	}
 
-	SetDefaultParent(InElementId);
+	SetDefaultParent(InElementID);
 }
 
 // Recursively clean. Delete element that hasn't 3d geometry related to it
@@ -121,19 +123,21 @@ void FSyncData::SetDefaultParent(const FElementID& InElementID)
 {
 	if (!HasParent())
 	{
-		if (InElementID.ElementHeader.hotlinkGuid == APINULLGuid)
+		const API_Elem_Head& Header = InElementID.GetHeader();
+
+		if (Header.hotlinkGuid == APINULLGuid)
 		{
 			// Parent is a layer
-			SetParent(&InElementID.SyncContext.GetSyncDatabase().GetLayerSyncData(InElementID.ElementHeader.layer));
+			SetParent(&InElementID.SyncContext.GetSyncDatabase().GetLayerSyncData(Header.layer));
 		}
 		else
 		{
 			// Parent is a hot link instance
-			FSyncData*& ParentFound = InElementID.SyncContext.GetSyncDatabase().GetSyncData(
-				APIGuid2GSGuid(InElementID.ElementHeader.hotlinkGuid));
+			FSyncData*& ParentFound =
+				InElementID.SyncContext.GetSyncDatabase().GetSyncData(APIGuid2GSGuid(Header.hotlinkGuid));
 			if (ParentFound == nullptr)
 			{
-				ParentFound = new FSyncData::FHotLinkInstance(APIGuid2GSGuid(InElementID.ElementHeader.hotlinkGuid),
+				ParentFound = new FSyncData::FHotLinkInstance(APIGuid2GSGuid(Header.hotlinkGuid),
 															  &InElementID.SyncContext.GetSyncDatabase());
 			}
 			SetParent(ParentFound);
@@ -387,11 +391,11 @@ void FSyncData::FActor::SetActorElement(const TSharedPtr< IDatasmithActorElement
 // Add tags data
 void FSyncData::FActor::UpdateTags(const FTagsArray& InTags)
 {
-	int32 Count = (int32)InTags.Num();
+	int32 Count = (int32)InTags.GetSize();
 	int32 Index = 0;
 	if (ActorElement->GetTagsCount() == Count)
 	{
-		while (Index < Count && InTags[Index] == ActorElement->GetTag(Index))
+		while (Index < Count && FCString::Strcmp(GSStringToUE(InTags[Index]), ActorElement->GetTag(Index)) == 0)
 		{
 			++Index;
 		}
@@ -404,47 +408,119 @@ void FSyncData::FActor::UpdateTags(const FTagsArray& InTags)
 	ActorElement->ResetTags();
 	for (Index = 0; Index < Count; ++Index)
 	{
-		ActorElement->AddTag(*InTags[Index]);
+		ActorElement->AddTag(GSStringToUE(InTags[Index]));
 	}
 }
 
 // Add tags data
-void FSyncData::FActor::AddTags(const FElementID& InElementID)
+void FSyncData::FActor::AddTags(FElementID& InElementID)
 {
 	UE_AC_Assert(ActorElement.IsValid());
+	const API_Elem_Head& Header = InElementID.GetHeader();
+
+	FTagsArray Tags(10);
+
+	static const GS::UniString PrefixTagUniqueID("Archicad.Element.UniqueID.");
+	Tags.Push(PrefixTagUniqueID + ElementId.ToUniString());
+
+	static const GS::UniString PrefixTagType("Archicad.Element.Type.");
+	GS::UniString			   ElementTypeName(FElementTools::TypeName(Header.typeID), CC_UTF8);
+	Tags.Push(PrefixTagType + ElementTypeName);
+
+	const FLibPartInfo* LibPartInfo = InElementID.GetLibPartInfo();
+	if (LibPartInfo != nullptr)
+	{
+		static const GS::UniString PrefixLibPartMain("Archicad.Element.LibPart.Main.");
+		Tags.Push(PrefixLibPartMain + LibPartInfo->Guid.Main.ToUniString());
 
-	FTagsArray Tags;
+		static const GS::UniString PrefixLibPartRev("Archicad.Element.LibPart.Rev.");
+		Tags.Push(PrefixLibPartRev + LibPartInfo->Guid.Rev.ToUniString());
 
-	static GS::UniString PrefixTagUniqueID("Archicad.Element.UniqueID.");
-	GS::UniString		 TagUniqueID = PrefixTagUniqueID + ElementId.ToUniString();
-	Tags.Add(GSStringToUE(TagUniqueID));
+		static const GS::UniString PrefixLibPartName("Archicad.Element.LibPart.Name.");
+		Tags.Push(PrefixLibPartName + LibPartInfo->Name);
 
-	static GS::UniString PrefixTagType("Archicad.Element.Type.");
-	GS::UniString		 TagElementType = PrefixTagType + FElementTools::TypeName(InElementID.ElementHeader.typeID);
-	Tags.Add(GSStringToUE(TagElementType));
+		if (Header.typeID == API_ObjectID || Header.typeID == API_LampID || Header.typeID == API_WindowID ||
+			Header.typeID == API_DoorID)
+		{
+			const API_Element& APIElement = InElementID.GetAPIElement();
+
+			if (APIElement.header.typeID == API_ObjectID || APIElement.header.typeID == API_LampID)
+			{
+				if (APIElement.object.useObjMaterials)
+				{
+					static const GS::UniString TagUseObjectMaterial("Archicad.Element.UseObjectMaterial");
+					Tags.Push(TagUseObjectMaterial);
+				}
+				if (APIElement.object.reflected)
+				{
+					static const GS::UniString TagObjectReflected("Archicad.Element.Reflected");
+					Tags.Push(TagObjectReflected);
+				}
+			}
+			else if (APIElement.header.typeID == API_WindowID || APIElement.header.typeID == API_DoorID)
+			{
+				if (APIElement.window.openingBase.reflected)
+				{
+					static const GS::UniString TagObjectReflected("Archicad.Element.Mirror.Y");
+					Tags.Push(TagObjectReflected);
+				}
+				if (APIElement.window.openingBase.oSide)
+				{
+					static const GS::UniString TagObjectReflected("Archicad.Element.Mirror.X");
+					Tags.Push(TagObjectReflected);
+				}
+				if (APIElement.window.openingBase.refSide)
+				{
+					static const GS::UniString TagObjectReflected("Archicad.Element.Mirror.X.Same");
+					Tags.Push(TagObjectReflected);
+				}
+			}
+		}
+	}
+
+#if 0 // IfcType and Classification only as metadata
+	GS::UniString IfcType;
+	GS::UniString TypeObjectIfcType;
+	GSErrCode GSErr = ACAPI_Element_GetIFCType(Header.guid, &IfcType, &TypeObjectIfcType);
+	if (GSErr == NoError)
+	{
+		if (!IfcType.IsEmpty())
+		{
+			static const GS::UniString PrefixIfcType("Archicad.IfcType.");
+			Tags.Push(PrefixIfcType + IfcType);
+		}
+		if (!TypeObjectIfcType.IsEmpty())
+		{
+			static const GS::UniString PrefixTypeObjectIfcType("Archicad.TypeObjectIfcType.");
+			Tags.Push(PrefixTypeObjectIfcType + TypeObjectIfcType);
+		}
+	}
+	else
+	{
+		UE_AC_DebugF("FSyncData::AddTags - ACAPI_Element_GetIFCType returned error %s", GetErrorName(GSErr));
+	}
 
 	GS::Array< GS::Pair< API_ClassificationSystem, API_ClassificationItem > > ApiClassifications;
-	GSErrCode GSErr = FElementTools::GetElementClassifications(ApiClassifications, InElementID.ElementHeader.guid);
+	GSErr = FElementTools::GetElementClassifications(ApiClassifications, Header.guid);
 
 	if (GSErr == NoError)
 	{
-		TSet< FString > ClassificationIds;
-		static FString	PrefixTagClassificationID(TEXT("Archicad.Classification.ID."));
+		static const GS::UniString PrefixTagClassificationID("Archicad.Classification.");
+		static const GS::UniString StringID(".ID.");
 		for (const GS::Pair< API_ClassificationSystem, API_ClassificationItem >& Classification : ApiClassifications)
 		{
-			FString TagClassificationID = PrefixTagClassificationID + GSStringToUE(Classification.second.id);
-			bool	bTagClassificationIsAlreadyInSet = false;
-			ClassificationIds.Add(TagClassificationID, &bTagClassificationIsAlreadyInSet);
-			if (!bTagClassificationIsAlreadyInSet)
+			GS::UniString TagClassificationID(PrefixTagClassificationID + Classification.first.name + StringID + Classification.second.id);
+			if (!Tags.Contains(TagClassificationID))
 			{
-				Tags.Add(TagClassificationID);
+				Tags.Push(TagClassificationID);
 			}
 		}
 	}
 	else
 	{
-		UE_AC_DebugF("FSyncData::AddTags - FElementTools::GetElementClassifications returned error %d", GSErr);
+		UE_AC_DebugF("FSyncData::AddTags - FElementTools::GetElementClassifications returned error %s", GetErrorName(GSErr));
 	}
+#endif
 
 	UpdateTags(Tags);
 }
@@ -523,11 +599,13 @@ void FSyncData::FLayer::Process(FProcessInfo* /* IOProcessInfo */)
 						 LayerIndex);
 			if (error == APIERR_DELETED)
 			{
-				LayerName = GetGSName(kName_LayerDeleted);
+				static const GS::UniString LayerDeleted("Layer deleted");
+				LayerName = LayerDeleted;
 			}
 			else
 			{
-				LayerName = GS::UniString::Printf(GetGSName(kName_LayerError), GetErrorName(error));
+				utf8_string LayerError(Utf8StringFormat("Layer error=%s", GetErrorName(error)));
+				LayerName = GS::UniString(LayerError.c_str(), CC_UTF8);
 			}
 		}
 		else if (LayerName == "\x14") // Special ARCHICAD layer
@@ -641,10 +719,10 @@ void FConvertGeometry2MeshElement::AddElementGeometry(FElementID*						IOElement
 	Geometry::Transformation3D World2Local = Local2World.GetInverse().Get(Geometry::Transformation3D());
 #endif
 
-	Element2StaticMesh.AddElementGeometry(IOElementID->Element3D, World2Local);
+	Element2StaticMesh.AddElementGeometry(IOElementID->GetElement3D(), World2Local);
 }
 
-FSyncData::FElement::FElement(const GS::Guid& InGuid, const FSyncContext& InSyncContext)
+FSyncData::FElement::FElement(const GS::Guid& InGuid, const FSyncContext& /* InSyncContext */)
 	: FSyncData::FActor(InGuid)
 {
 }
@@ -699,14 +777,14 @@ void FSyncData::FElement::Process(FProcessInfo* IOProcessInfo)
 		{
 			IOProcessInfo->ElementID.InitElement(this);
 			IOProcessInfo->ElementID.InitHeader(GSGuid2APIGuid(ElementId));
-			CheckModificationStamp(IOProcessInfo->ElementID.ElementHeader.modiStamp);
+			CheckModificationStamp(IOProcessInfo->ElementID.GetHeader().modiStamp);
 
 			UE_AC_STAT(IOProcessInfo->SyncContext.Stats.TotalOwnerCreated++);
 			TSharedRef< IDatasmithActorElement > NewActor =
 				FDatasmithSceneFactory::CreateActor(GSStringToUE(ElementId.ToUniString()));
 
 			GS::UniString ElemenInfo;
-			if (FElementTools::GetInfoString(IOProcessInfo->ElementID.ElementHeader.guid, &ElemenInfo))
+			if (FElementTools::GetInfoString(IOProcessInfo->ElementID.GetHeader().guid, &ElemenInfo))
 			{
 				NewActor->SetLabel(GSStringToUE(ElemenInfo));
 			}
@@ -718,7 +796,7 @@ void FSyncData::FElement::Process(FProcessInfo* IOProcessInfo)
 
 			SetActorElement(NewActor);
 			AddTags(IOProcessInfo->ElementID);
-			UpdateMetaData(IOProcessInfo->SyncContext.GetScene());
+			UpdateMetaData(IOProcessInfo);
 		}
 	}
 	else
@@ -732,10 +810,10 @@ void FSyncData::FElement::Process(FProcessInfo* IOProcessInfo)
 			IOProcessInfo->ElementID.InitHeader();
 
 			ModelerAPI::Transformation LocalToWorld =
-				IOProcessInfo->ElementID.Element3D.GetElemLocalToWorldTransformation();
+				IOProcessInfo->ElementID.GetElement3D().GetElemLocalToWorldTransformation();
 			if ((LocalToWorld.status & TR_IDENT) != 0)
 			{
-				Box3D Bounds = IOProcessInfo->ElementID.Element3D.GetBounds();
+				Box3D Bounds = IOProcessInfo->ElementID.GetElement3D().GetBounds();
 				LocalToWorld.matrix[0][3] = (Bounds.xMin + Bounds.xMax) * 0.5;
 				LocalToWorld.matrix[1][3] = (Bounds.yMin + Bounds.yMax) * 0.5;
 				LocalToWorld.matrix[2][3] = Bounds.zMin;
@@ -797,7 +875,7 @@ void FSyncData::FElement::Process(FProcessInfo* IOProcessInfo)
 
 			// Set actor label
 			GS::UniString ElemenInfo;
-			if (FElementTools::GetInfoString(IOProcessInfo->ElementID.ElementHeader.guid, &ElemenInfo))
+			if (FElementTools::GetInfoString(IOProcessInfo->ElementID.GetHeader().guid, &ElemenInfo))
 			{
 				ActorElement->SetLabel(GSStringToUE(ElemenInfo));
 			}
@@ -810,12 +888,12 @@ void FSyncData::FElement::Process(FProcessInfo* IOProcessInfo)
 			ActorElement->SetRotation(FGeometryUtil::GetRotationQuat(LocalToWorld.matrix));
 
 			// Set actor layer
-			ActorElement->SetLayer(*IOProcessInfo->SyncContext.GetSyncDatabase().GetLayerName(
-				IOProcessInfo->ElementID.ElementHeader.layer));
+			ActorElement->SetLayer(
+				*IOProcessInfo->SyncContext.GetSyncDatabase().GetLayerName(IOProcessInfo->ElementID.GetHeader().layer));
 
 			AddTags(IOProcessInfo->ElementID);
 
-			UpdateMetaData(IOProcessInfo->SyncContext.GetScene());
+			UpdateMetaData(IOProcessInfo);
 
 			ConvertGeometry2MeshElement->CreateDatasmithMesh();
 		}
@@ -878,11 +956,28 @@ void FSyncData::FElement::DeleteMe(FSyncDatabase* IOSyncDatabase)
 }
 
 // Rebuild the meta data of this element
-void FSyncData::FElement::UpdateMetaData(IDatasmithScene& IOScene)
+void FSyncData::FElement::UpdateMetaData(FProcessInfo* IOProcessInfo)
 {
 	FMetaData MetaDataExporter(ActorElement);
 	MetaDataExporter.ExportMetaData(ElementId);
-	MetaDataExporter.SetOrUpdate(&MetaData, &IOScene);
+	API_ElemTypeID TypeID = IOProcessInfo->ElementID.GetHeader().typeID;
+	if (TypeID == API_ObjectID || TypeID == API_LampID)
+	{
+		FAutoMemo AutoMemo(GSGuid2APIGuid(ElementId), APIMemoMask_AddPars);
+		if (AutoMemo.GSErr == NoError)
+		{
+			if (AutoMemo.Memo.params) // Can be null
+			{
+				GS::UniString EscapeAssetId;
+				if (GetParameter(AutoMemo.Memo.params, "enscapeAssetId", &EscapeAssetId))
+				{
+					MetaDataExporter.AddProperty(TEXT("EnscapeAssetId"), EDatasmithKeyValuePropertyType::String,
+												 EscapeAssetId);
+				}
+			}
+		}
+	}
+	MetaDataExporter.SetOrUpdate(&MetaData, &IOProcessInfo->SyncContext.GetScene());
 }
 
 #pragma mark -
@@ -1010,7 +1105,7 @@ void FSyncData::FCamera::InitWithCameraElement()
 
 #pragma mark -
 
-void FSyncData::FLight::Process(FProcessInfo* /* IOProcessInfo */)
+void FSyncData::FLight::Process(FProcessInfo* IOProcessInfo)
 {
 	if (!ActorElement.IsValid())
 	{
@@ -1057,12 +1152,31 @@ void FSyncData::FLight::Process(FProcessInfo* /* IOProcessInfo */)
 
 		LightElement.SetTranslation(Position);
 		LightElement.SetRotation(Rotation);
-		LightElement.SetIntensity(5000.0);
+		LightElement.SetIntensity(Intensity * 5000.0);
 		LightElement.SetColor(Color);
 		if (Color == FLinearColor(0, 0, 0))
 		{
 			LightElement.SetEnabled(false);
 		}
+
+		LightElement.SetUseIes(bUseIES);
+		if (bUseIES)
+		{
+			const FTexturesCache::FIESTexturesCacheElem& Texture =
+				IOProcessInfo->SyncContext.GetTexturesCache().GetIESTexture(IOProcessInfo->SyncContext,
+																			GSStringToUE(IESFileName));
+			LightElement.SetIesTexturePathName(*Texture.TexturePath);
+
+			/*
+			 LightElement.SetIesBrightnessScale(double IesBrightnessScale);
+			 LightElement.SetUseIesBrightness(bool bUseIesBrightness);
+			 LightElement.SetIesRotation(const FQuat& IesRotation);
+			 */
+		}
+		else
+		{
+			LightElement.SetIesTexturePathName(TEXT(""));
+		}
 	}
 }
 
diff --git a/Engine/Source/Programs/Enterprise/Datasmith/DatasmithARCHICADExporter/Private/SyncData.h b/Engine/Source/Programs/Enterprise/Datasmith/DatasmithARCHICADExporter/Private/SyncData.h
--- a/Engine/Source/Programs/Enterprise/Datasmith/DatasmithARCHICADExporter/Private/SyncData.h
+++ b/Engine/Source/Programs/Enterprise/Datasmith/DatasmithARCHICADExporter/Private/SyncData.h
@@ -222,13 +222,13 @@ class FSyncData::FActor : public FSyncData
 	// Return Element as an actor
 	virtual const TSharedPtr< IDatasmithActorElement >& GetActorElement() const override { return ActorElement; }
 
-	typedef TArray< FString > FTagsArray;
+	typedef GS::Array< GS::UniString > FTagsArray;
 
 	// Update tags data
 	void UpdateTags(const FTagsArray& InTags);
 
 	// Add tags data
-	void AddTags(const FElementID& InElementID);
+	void AddTags(FElementID& InElementID);
 
 	void ReplaceMetaData(IDatasmithScene& IOScene, const TSharedPtr< IDatasmithMetaDataElement >& InNewMetaData);
 
@@ -280,7 +280,7 @@ class FSyncData::FElement : public FSyncData::FActor
 	virtual bool AttachObserver(FAttachObservers* IOAttachObservers) override;
 
 	// Rebuild the meta data of this element
-	void UpdateMetaData(IDatasmithScene& IOScene);
+	void UpdateMetaData(FProcessInfo* IOProcessInfo);
 
 	// Return true if this element and all it's childs have been cut out
 	bool CheckAllCutOut() override;
@@ -354,6 +354,17 @@ class FSyncData::FLight : public FSyncData::FActor
 		}
 	}
 
+	void SetValuesFromParameters(double InIntensity, bool bInUseIES, const GS::UniString& InIESFileName)
+	{
+		if (Intensity != InIntensity || bUseIES != bInUseIES || IESFileName != InIESFileName)
+		{
+			Intensity = InIntensity;
+			bUseIES = bInUseIES;
+			IESFileName = InIESFileName;
+			bIsModified = true;
+		}
+	}
+
 	void Placement(const FVector& InPosition, const FQuat& InRotation)
 	{
 		if (Position != InPosition || Rotation != InRotation)
@@ -374,6 +385,9 @@ class FSyncData::FLight : public FSyncData::FActor
 	FLinearColor			Color;
 	FVector					Position;
 	FQuat					Rotation;
+	double					Intensity;
+	bool					bUseIES;
+	GS::UniString			IESFileName;
 };
 
 class FSyncData::FHotLinksRoot : public FSyncData::FActor
diff --git a/Engine/Source/Programs/Enterprise/Datasmith/DatasmithARCHICADExporter/Private/SyncDatabase.cpp b/Engine/Source/Programs/Enterprise/Datasmith/DatasmithARCHICADExporter/Private/SyncDatabase.cpp
--- a/Engine/Source/Programs/Enterprise/Datasmith/DatasmithARCHICADExporter/Private/SyncDatabase.cpp
+++ b/Engine/Source/Programs/Enterprise/Datasmith/DatasmithARCHICADExporter/Private/SyncDatabase.cpp
@@ -11,6 +11,8 @@
 #include "Utils/Element2String.h"
 
 #include "DatasmithUtils.h"
+#include "IDirectLinkUI.h"
+#include "IDatasmithExporterUIModule.h"
 
 #include "ModelMeshBody.hpp"
 #include "Light.hpp"
@@ -25,11 +27,12 @@ BEGIN_NAMESPACE_UE_AC
 #endif
 
 // Constructor
-FSyncDatabase::FSyncDatabase(const TCHAR* InSceneName, const TCHAR* InSceneLabel, const TCHAR* InAssetsPath)
+FSyncDatabase::FSyncDatabase(const TCHAR* InSceneName, const TCHAR* InSceneLabel, const TCHAR* InAssetsPath,
+							 const GS::UniString& InAssetsCache)
 	: Scene(FDatasmithSceneFactory::CreateScene(*FDatasmithUtils::SanitizeObjectName(InSceneName)))
 	, AssetsFolderPath(InAssetsPath)
 	, MaterialsDatabase(new FMaterialsDatabase())
-	, TexturesCache(new FTexturesCache())
+	, TexturesCache(new FTexturesCache(InAssetsCache))
 {
 	Scene->SetLabel(InSceneLabel);
 }
@@ -227,6 +230,75 @@ bool FSyncDatabase::SetMesh(TSharedPtr< IDatasmithMeshElement >*	   Handle,
 	return true;
 }
 
+// Return the libpart from it's index
+FLibPartInfo* FSyncDatabase::GetLibPartInfo(GS::Int32 InIndex)
+{
+	TUniquePtr< FLibPartInfo >* LibPartInfo = IndexToLibPart.Find(InIndex);
+	if (LibPartInfo == nullptr)
+	{
+		LibPartInfo = &IndexToLibPart.Add(InIndex, MakeUnique< FLibPartInfo >());
+		LibPartInfo->Get()->Initialize(InIndex);
+	}
+	return LibPartInfo->Get();
+}
+
+// Return the libpart from it's unique id
+FLibPartInfo* FSyncDatabase::GetLibPartInfo(const char* InUnID)
+{
+	FGSUnID	  UnID;
+	GSErrCode GSErr = UnID.InitWithString(InUnID);
+	if (GSErr != NoError)
+	{
+		UE_AC_DebugF("FSyncDatabase::GetLibPartInfo - InitWithString(\"%s\") return error\n", InUnID);
+		return nullptr;
+	}
+	if (UnID.Main == GS::NULLGuid && UnID.Rev == GS::NULLGuid)
+	{
+		return nullptr;
+	}
+
+	FLibPartInfo** LibPartInfoPtr = UnIdToLibPart.Find(UnID);
+	if (LibPartInfoPtr == nullptr)
+	{
+		API_LibPart LibPart;
+		Zap(&LibPart);
+		strncpy(LibPart.ownUnID, InUnID, sizeof(LibPart.ownUnID));
+		GSErrCode err = ACAPI_LibPart_Search(&LibPart, false);
+		if (err != NoError)
+		{
+			UE_AC_DebugF("FSyncDatabase::GetLibPartInfo - Can't find libpart \"%s\"\n", InUnID);
+			return nullptr;
+		}
+		FLibPartInfo* LibPartInfo = GetLibPartInfo(LibPart.index);
+		UnIdToLibPart.Add(UnID, LibPartInfo);
+		return LibPartInfo;
+	}
+	return *LibPartInfoPtr;
+}
+
+// Return the cache path
+GS::UniString FSyncDatabase::GetCachePath()
+{
+	const TCHAR*				CacheDirectory = nullptr;
+	IDatasmithExporterUIModule* DsExporterUIModule = IDatasmithExporterUIModule::Get();
+	if (DsExporterUIModule != nullptr)
+	{
+		IDirectLinkUI* DLUI = DsExporterUIModule->GetDirectLinkExporterUI();
+		if (DLUI != nullptr)
+		{
+			CacheDirectory = DLUI->GetDirectLinkCacheDirectory();
+		}
+	}
+	if (CacheDirectory != nullptr && *CacheDirectory != 0)
+	{
+		return UEToGSString(CacheDirectory);
+	}
+	else
+	{
+		return GetAddonDataDirectory();
+	}
+}
+
 // SetSceneInfo
 void FSyncDatabase::SetSceneInfo()
 {
@@ -264,7 +336,7 @@ UInt32 FSyncDatabase::ScanElements(const FSyncContext& InSyncContext)
 			continue;
 		}
 
-		API_Guid ElementGuid = GSGuid2APIGuid(ElementID.Element3D.GetElemGuid());
+		API_Guid ElementGuid = GSGuid2APIGuid(ElementID.GetElement3D().GetElemGuid());
 		if (ElementGuid == APINULLGuid)
 		{
 #if UE_AC_DO_TRACE && 1
@@ -285,7 +357,7 @@ UInt32 FSyncDatabase::ScanElements(const FSyncContext& InSyncContext)
 #endif
 
 		// Check 3D geometry bounding box
-		Box3D box = ElementID.Element3D.GetBounds(ModelerAPI::CoordinateSystem::ElemLocal);
+		Box3D box = ElementID.GetElement3D().GetBounds(ModelerAPI::CoordinateSystem::ElemLocal);
 
 		// Bonding box is empty, must not happen, but it happen
 		if (box.xMin > box.xMax || box.yMin > box.yMax || box.zMin > box.zMax)
@@ -310,12 +382,12 @@ UInt32 FSyncDatabase::ScanElements(const FSyncContext& InSyncContext)
 		UE_AC_STAT(InSyncContext.Stats.TotalElementsWithGeometry++);
 
 		// Get sync data for this element (Create or reuse already existing)
-		FSyncData*& SyncData = GetSyncData(APIGuid2GSGuid(ElementID.ElementHeader.guid));
+		FSyncData*& SyncData = GetSyncData(APIGuid2GSGuid(ElementID.GetHeader().guid));
 		if (SyncData == nullptr)
 		{
-			SyncData = new FSyncData::FElement(APIGuid2GSGuid(ElementID.ElementHeader.guid), InSyncContext);
+			SyncData = new FSyncData::FElement(APIGuid2GSGuid(ElementID.GetHeader().guid), InSyncContext);
 		}
-		ElementID.SyncData = SyncData;
+		ElementID.SetSyncData(SyncData);
 		SyncData->Update(ElementID);
 		if (SyncData->IsModified())
 		{
@@ -323,7 +395,7 @@ UInt32 FSyncDatabase::ScanElements(const FSyncContext& InSyncContext)
 		}
 
 		// Add lights
-		if (ElementID.Element3D.GetLightCount() > 0)
+		if (ElementID.GetElement3D().GetLightCount() > 0)
 		{
 			ScanLights(ElementID);
 		}
@@ -345,29 +417,55 @@ UInt32 FSyncDatabase::ScanElements(const FSyncContext& InSyncContext)
 // Scan all lights of this element
 void FSyncDatabase::ScanLights(const FElementID& InElementID)
 {
-	GS::Int32 LightsCount = InElementID.Element3D.GetLightCount();
+	GS::Int32 LightsCount = InElementID.GetElement3D().GetLightCount();
 	if (LightsCount > 0)
 	{
+		const API_Elem_Head& Header = InElementID.GetHeader();
 #if UE_AC_DO_TRACE_LIGHTS
-		UE_AC_TraceF("%s", FElement2String::GetElementAsShortString(InElementID.ElementHeader.guid).c_str());
-		UE_AC_TraceF("%s", FElement2String::GetParametersAsString(InElementID.ElementHeader.guid).c_str());
+		UE_AC_TraceF("%s", FElement2String::GetElementAsShortString(Header.guid).c_str());
+		UE_AC_TraceF("%s", FElement2String::GetParametersAsString(Header.guid).c_str());
 		UE_AC_TraceF("%s", F3DElement2String::ElementLight2String(InElementID.Element3D).c_str());
 #endif
+		double		  Intensity = 1.0;
+		bool		  bUseIES = false;
+		GS::UniString IESFileName;
+		FAutoMemo	  AutoMemo(Header.guid, APIMemoMask_AddPars);
+		if (AutoMemo.GSErr == NoError)
+		{
+			if (AutoMemo.Memo.params) // Can be null
+			{
+				double value;
+				if (GetParameter(AutoMemo.Memo.params, "gs_light_intensity", &value))
+				{
+					Intensity = value / 100.0;
+				}
+				if (GetParameter(AutoMemo.Memo.params, "c4dPhoPhotometric", &value))
+				{
+					bUseIES = value != 0;
+				}
+				GetParameter(AutoMemo.Memo.params, "c4dPhoIESFile", &IESFileName);
+			}
+		}
+		else
+		{
+			UE_AC_DebugF("FSyncDatabase::ScanLights - Error=%d when getting element memo\n", AutoMemo.GSErr);
+		}
+
 		ModelerAPI::Light Light;
 
 		for (GS::Int32 LightIndex = 1; LightIndex <= LightsCount; ++LightIndex)
 		{
-			InElementID.Element3D.GetLight(LightIndex, &Light);
+			InElementID.GetElement3D().GetLight(LightIndex, &Light);
 			ModelerAPI::Light::Type LightType = Light.GetType();
 			if (LightType == ModelerAPI::Light::DirectionLight || LightType == ModelerAPI::Light::SpotLight ||
 				LightType == ModelerAPI::Light::PointLight)
 			{
-				API_Guid	LightId = CombineGuid(InElementID.ElementHeader.guid, GuidFromMD5(LightIndex));
+				API_Guid	LightId = CombineGuid(Header.guid, GuidFromMD5(LightIndex));
 				FSyncData*& SyncData = FSyncDatabase::GetSyncData(APIGuid2GSGuid(LightId));
 				if (SyncData == nullptr)
 				{
 					SyncData = new FSyncData::FLight(APIGuid2GSGuid(LightId), LightIndex);
-					SyncData->SetParent(InElementID.SyncData);
+					SyncData->SetParent(InElementID.GetSyncData());
 					SyncData->MarkAsModified();
 				}
 				FSyncData::FLight& LightSyncData = static_cast< FSyncData::FLight& >(*SyncData);
@@ -376,7 +474,9 @@ void FSyncDatabase::ScanLights(const FElementID& InElementID)
 				const float	 InnerConeAngle = float(Light.GetFalloffAngle1() * 180.0f / PI);
 				const float	 OuterConeAngle = float(Light.GetFalloffAngle2() * 180.0f / PI);
 				FLinearColor LightColor = ACRGBColorToUELinearColor(Light.GetColor());
+
 				LightSyncData.SetValues(LightType, InnerConeAngle, OuterConeAngle, LightColor);
+				LightSyncData.SetValuesFromParameters(Intensity, bUseIES, IESFileName);
 				LightSyncData.Placement(FGeometryUtil::GetTranslationVector(Light.GetPosition()),
 										FGeometryUtil::GetRotationQuat(Light.GetDirection()));
 			}
diff --git a/Engine/Source/Programs/Enterprise/Datasmith/DatasmithARCHICADExporter/Private/SyncDatabase.h b/Engine/Source/Programs/Enterprise/Datasmith/DatasmithARCHICADExporter/Private/SyncDatabase.h
--- a/Engine/Source/Programs/Enterprise/Datasmith/DatasmithARCHICADExporter/Private/SyncDatabase.h
+++ b/Engine/Source/Programs/Enterprise/Datasmith/DatasmithARCHICADExporter/Private/SyncDatabase.h
@@ -3,6 +3,7 @@
 #pragma once
 
 #include "Utils/AddonTools.h"
+#include "Utils/LibPartInfo.h"
 
 #include "Lock.hpp"
 
@@ -15,13 +16,15 @@ class FSyncData;
 class FMaterialsDatabase;
 class FTexturesCache;
 class FElementID;
+class FLibPartInfo;
 
 // Class that maintain synchronization datas (SyncData, Material, Texture)
 class FSyncDatabase
 {
   public:
 	// Constructor
-	FSyncDatabase(const TCHAR* InSceneName, const TCHAR* InSceneLabel, const TCHAR* InAssetsPath);
+	FSyncDatabase(const TCHAR* InSceneName, const TCHAR* InSceneLabel, const TCHAR* InAssetsPath,
+				  const GS::UniString& InAssetsCache);
 
 	// Destructor
 	~FSyncDatabase();
@@ -69,9 +72,20 @@ class FSyncDatabase
 	// Set the mesh in the handle and take care of mesh life cycle.
 	bool SetMesh(TSharedPtr< IDatasmithMeshElement >* Handle, const TSharedPtr< IDatasmithMeshElement >& InMesh);
 
+	// Return the libpart from it's index
+	FLibPartInfo* GetLibPartInfo(GS::Int32 InIndex);
+
+	// Return the libpart from it's unique id
+	FLibPartInfo* GetLibPartInfo(const char* InUnID);
+
+	// Return the cache path
+	static GS::UniString GetCachePath();
+
   private:
-	typedef TMap< FGuid, FSyncData* > FMapGuid2SyncData;
-	typedef TMap< short, FString >	  FMapLayerIndex2Name;
+	typedef TMap< FGuid, FSyncData* >					  FMapGuid2SyncData;
+	typedef TMap< short, FString >						  FMapLayerIndex2Name;
+	typedef TMap< GS::Int32, TUniquePtr< FLibPartInfo > > MapIndex2LibPart;
+	typedef TMap< FGSUnID, FLibPartInfo* >				  MapUnId2LibPart;
 
 	// To take care of mesh life cycle.
 	class FMeshInfo
@@ -115,6 +129,12 @@ class FSyncDatabase
 	// Map mesh name (hash) to mesh info
 	FMapHashToMeshInfo HashToMeshInfo;
 	GS::Lock		   HashToMeshInfoAccesControl;
+
+	// Map lib part by index
+	MapIndex2LibPart IndexToLibPart;
+
+	// Map lib part by UnId
+	MapUnId2LibPart UnIdToLibPart;
 };
 
 END_NAMESPACE_UE_AC
diff --git a/Engine/Source/Programs/Enterprise/Datasmith/DatasmithARCHICADExporter/Private/Synchronizer.cpp b/Engine/Source/Programs/Enterprise/Datasmith/DatasmithARCHICADExporter/Private/Synchronizer.cpp
--- a/Engine/Source/Programs/Enterprise/Datasmith/DatasmithARCHICADExporter/Private/Synchronizer.cpp
+++ b/Engine/Source/Programs/Enterprise/Datasmith/DatasmithARCHICADExporter/Private/Synchronizer.cpp
@@ -10,8 +10,6 @@
 #include "DatasmithDirectLink.h"
 #include "DatasmithSceneExporter.h"
 #include "DatasmithSceneXmlWriter.h"
-#include "IDirectLinkUI.h"
-#include "IDatasmithExporterUIModule.h"
 
 #ifdef TicksPerSecond
 	#undef TicksPerSecond
@@ -247,7 +245,8 @@ void FSynchronizer::ProjectOpen()
 	GS::UniString ProjectPath;
 	GS::UniString ProjectName;
 	GetProjectPathAndName(&ProjectPath, &ProjectName);
-	SyncDatabase = new FSyncDatabase(GSStringToUE(ProjectPath), GSStringToUE(ProjectName), *GetExportPath());
+	SyncDatabase = new FSyncDatabase(GSStringToUE(ProjectPath), GSStringToUE(ProjectName),
+									 GSStringToUE(FSyncDatabase::GetCachePath()), FSyncDatabase::GetCachePath());
 
 	// Announce it to potential receivers
 #if ENGINE_MAJOR_VERSION == 4 && ENGINE_MINOR_VERSION == 26
@@ -290,29 +289,6 @@ void FSynchronizer::ProjectClosed()
 	Reset("Project Closed");
 }
 
-// Return the export path from the ExporterUIModule or a default one
-FString FSynchronizer::GetExportPath()
-{
-	const TCHAR*				CacheDirectory = nullptr;
-	IDatasmithExporterUIModule* DsExporterUIModule = IDatasmithExporterUIModule::Get();
-	if (DsExporterUIModule != nullptr)
-	{
-		IDirectLinkUI* DLUI = DsExporterUIModule->GetDirectLinkExporterUI();
-		if (DLUI != nullptr)
-		{
-			CacheDirectory = DLUI->GetDirectLinkCacheDirectory();
-		}
-	}
-	if (CacheDirectory != nullptr)
-	{
-		return CacheDirectory;
-	}
-	else
-	{
-		return GSStringToUE(GetAddonDataDirectory());
-	}
-}
-
 // Do a snapshot of the model 3D data
 void FSynchronizer::DoSnapshot(const ModelerAPI::Model& InModel)
 {
@@ -328,12 +304,12 @@ void FSynchronizer::DoSnapshot(const ModelerAPI::Model& InModel)
 
 	ViewState = FViewState();
 
-	FString ExportPath = GetExportPath();
+	GS::UniString ExportPath = FSyncDatabase::GetCachePath();
 
 	// If we have a sync database validate it use the ExportPath
 	if (SyncDatabase != nullptr)
 	{
-		if (ExportPath != SyncDatabase->GetAssetsFolderPath())
+		if (FCString::Strcmp(GSStringToUE(ExportPath), SyncDatabase->GetAssetsFolderPath()) != 0)
 		{
 			delete SyncDatabase;
 			SyncDatabase = nullptr;
@@ -347,7 +323,8 @@ void FSynchronizer::DoSnapshot(const ModelerAPI::Model& InModel)
 		GS::UniString ProjectName;
 		GetProjectPathAndName(&ProjectPath, &ProjectName);
 
-		SyncDatabase = new FSyncDatabase(GSStringToUE(ProjectPath), GSStringToUE(ProjectName), *ExportPath);
+		SyncDatabase = new FSyncDatabase(GSStringToUE(ProjectPath), GSStringToUE(ProjectName), GSStringToUE(ExportPath),
+										 ExportPath);
 
 #if ENGINE_MAJOR_VERSION == 4 && ENGINE_MINOR_VERSION == 26
 		TSharedRef< IDatasmithScene > ToBuildWith_4_26(SyncDatabase->GetScene());
diff --git a/Engine/Source/Programs/Enterprise/Datasmith/DatasmithARCHICADExporter/Private/Synchronizer.h b/Engine/Source/Programs/Enterprise/Datasmith/DatasmithARCHICADExporter/Private/Synchronizer.h
--- a/Engine/Source/Programs/Enterprise/Datasmith/DatasmithARCHICADExporter/Private/Synchronizer.h
+++ b/Engine/Source/Programs/Enterprise/Datasmith/DatasmithARCHICADExporter/Private/Synchronizer.h
@@ -65,9 +65,6 @@ class FSynchronizer
 	// Inform that the project has been closed
 	void ProjectClosed();
 
-	// Return the export path from the ExporterUIModule or a default one
-	static FString GetExportPath();
-
 	// Do a snapshot of the model 3D data
 	void DoSnapshot(const ModelerAPI::Model& InModel);
 
diff --git a/Engine/Source/Programs/Enterprise/Datasmith/DatasmithARCHICADExporter/Private/TexturesCache.cpp b/Engine/Source/Programs/Enterprise/Datasmith/DatasmithARCHICADExporter/Private/TexturesCache.cpp
--- a/Engine/Source/Programs/Enterprise/Datasmith/DatasmithARCHICADExporter/Private/TexturesCache.cpp
+++ b/Engine/Source/Programs/Enterprise/Datasmith/DatasmithARCHICADExporter/Private/TexturesCache.cpp
@@ -9,14 +9,16 @@
 #include "Graphics2D.h"
 #include "Folder.hpp"
 
+#include "Paths.h"
+#include "DatasmithUtils.h"
+
 BEGIN_NAMESPACE_UE_AC
 
-FTexturesCache::FTexturesCache()
-	: bUseRelative(false)
+FTexturesCache::FTexturesCache(const GS::UniString& InAssetsCache)
+	: AssetsCache(InAssetsCache)
 {
-	AbsolutePath = GetAddonDataDirectory() + UE_AC_DirSep + GetGSName(kName_Textures) + UE_AC_DirSep;
-	RelativePath = AbsolutePath;
-	bUseRelative = false;
+	AbsolutePath = AssetsCache + UE_AC_DirSep + GS::UniString("Textures") + UE_AC_DirSep;
+	IESTexturesPath = AssetsCache + UE_AC_DirSep + GS::UniString("IESTextures") + UE_AC_DirSep;
 }
 
 const FTexturesCache::FTexturesCacheElem& FTexturesCache::GetTexture(const FSyncContext& InSyncContext,
@@ -95,11 +97,10 @@ const FTexturesCache::FTexturesCacheElem& FTexturesCache::GetTexture(const FSync
 			Texture.Fingerprint = GSGuid2APIGuid(GS::Guid(fp));
 			UE_AC_VerboseF("Texture name=\"%s\": ACFingerprint=\"%s\"\n", Texture.TextureLabel.ToUtf8(), fp.ToUtf8());
 		}
-		Texture.TexturePath =
-			(bUseRelative ? RelativePath : AbsolutePath) + Texture.TextureLabel + GetGSName(kName_TextureExtension);
 
-		WriteTexture(AcTexture, AbsolutePath + Texture.TextureLabel + GetGSName(kName_TextureExtension),
-					 InSyncContext.bUseFingerPrint);
+		CreateCacheFolders();
+		Texture.TexturePath = AbsolutePath + Texture.TextureLabel + GetGSName(kName_TextureExtension);
+		WriteTexture(AcTexture, Texture.TexturePath, InSyncContext.bUseFingerPrint);
 	}
 	else
 	{
@@ -135,15 +136,192 @@ const FTexturesCache::FTexturesCacheElem& FTexturesCache::GetTexture(const FSync
 	return Texture;
 }
 
-void FTexturesCache::WriteTexture(const ModelerAPI::Texture& InACTexture, const GS::UniString& InPath,
-								  bool InIsFingerprint) const
+/* Tool class to search file in Folder and Sub folder of a Parent folder.
+ * If found, copy the file to destination folder */
+class FSearchAndCopyFile
+{
+  public:
+	// Constructor
+	FSearchAndCopyFile(IO::Folder& InDestination, const IO::Name& InFileName)
+		: Destination(InDestination)
+		, FileName(InFileName)
+		, bFileInDestination(false)
+	{
+		// Check if the file is already in the destination folder
+		GSErrCode GSErr = InDestination.Contains(InFileName, &bFileInDestination);
+		if (GSErr != NoError)
+		{
+			UE_AC_DebugF("FSearchAndCopyFile::FSearchAndCopyFile - IESTexturesFolder.Contains return error %s\n",
+						 GetErrorName(GSErr));
+		}
+	}
+
+	// Search in the folder tree starting at parent folder.
+	bool DoSearchIn(const IO::Folder& Parent) { return Parent.Enumerate(EnumCallBack, this); }
+
+	// Callback function called for each element of the parent specified
+	bool EnumCallBack(const IO::Folder& Parent, const IO::Name& EntryName, bool bIsFolder);
+
+	// Return true if the file is in the destination folder
+	bool IsFileInDestination() const { return bFileInDestination; }
+
+  private:
+	// Callback to folder enumerate function
+	static bool CCALL EnumCallBack(const IO::Folder& Parent, const IO::Name& EntryName, bool bIsFolder, void* UserData);
+
+	// Destination where we must copy the file
+	IO::Folder& Destination;
+
+	// File's name to be copied
+	const IO::Name& FileName;
+
+	// True if the file is in the destination folder
+	bool bFileInDestination;
+};
+
+// Callback to folder enumerate function
+bool CCALL FSearchAndCopyFile::EnumCallBack(const IO::Folder& Parent, const IO::Name& EntryName, bool bIsFolder,
+											void* UserData)
 {
+	return reinterpret_cast< FSearchAndCopyFile* >(UserData)->EnumCallBack(Parent, EntryName, bIsFolder);
+}
+
+// Callback function called for each element of the parent specified
+bool FSearchAndCopyFile::EnumCallBack(const IO::Folder& Parent, const IO::Name& EntryName, bool bIsFolder)
+{
+	if (bIsFolder)
+	{
+		// Search in sub folder
+		DoSearchIn(IO::Folder(Parent, EntryName));
+	}
+	else
+	{
+		if (EntryName == FileName)
+		{
+			// Try to copy the file
+			GSErrCode GSErr = Parent.Copy(EntryName, Destination, FileName);
+			if (GSErr == NoError)
+			{
+				bFileInDestination = true;
+			}
+			else
+			{
+				UE_AC_DebugF("FSearchAndCopyFile::EnumCallBack - IO::Folder::Copy returned error %s\n",
+							 GetErrorName(GSErr));
+			}
+		}
+	}
+
+	return !bFileInDestination; // Stop if copy is done
+}
+
+// Insure we have a copy of the IES file in the cache folder
+GS::UniString FTexturesCache::CopyIESFile(const GS::UniString& InIESFileName)
+{
+	IO::Location IESTexturesLocation(IESTexturesPath);
+	IO::Name	 IESFileName(InIESFileName);
+
+	// Create the cached texture path
+	IO::Location  IESTextureLocation(IESTexturesLocation, IESFileName);
+	GS::UniString IESTexturePath;
+	IESTextureLocation.ToPath(&IESTexturePath);
+
 	// Create the texture folder if it's not present
-	IO::Location FolderLocation(AbsolutePath);
+	IO::Folder IESTexturesFolder(IESTexturesLocation, IO::Folder::Create);
+
+	// Search the IES file and copy it in the cache
+	FSearchAndCopyFile SearchAndCopyIESFile(IESTexturesFolder, IESFileName);
+	if (!SearchAndCopyIESFile.IsFileInDestination())
+	{
+		GS::Array< API_LibraryInfo > LibInfoArray;
+		GSErrCode					 GSErr = ACAPI_Environment(APIEnv_GetLibrariesID, &LibInfoArray);
+		if (GSErr == NoError)
+		{
+			for (UInt32 IndexLibrary = 0;
+				 IndexLibrary < LibInfoArray.GetSize() && !SearchAndCopyIESFile.IsFileInDestination(); IndexLibrary++)
+			{
+				const API_LibraryInfo& LibInfo = LibInfoArray[IndexLibrary];
+
+				if (LibInfo.libraryType == API_LocalLibrary || LibInfo.libraryType == API_EmbeddedLibrary ||
+					LibInfo.libraryType == API_ServerLibrary)
+				{
+					IO::Folder LibraryFolder(LibInfo.location, IO::Folder::Ignore);
+					SearchAndCopyIESFile.DoSearchIn(LibraryFolder);
+				}
+			}
+		}
+	}
+
+	if (!SearchAndCopyIESFile.IsFileInDestination())
+	{
+		UE_AC_ReportF("FTexturesCache::CopyIESFile - Cannot find IES File \"%s\"\n", InIESFileName.ToUtf8());
+	}
+
+	return IESTexturePath;
+}
 
-	IO::Folder TextureFolder(FolderLocation, IO::Folder::Create);
-	UE_AC_TestGSError(TextureFolder.GetStatus());
+// Create IES texture
+const FTexturesCache::FIESTexturesCacheElem& FTexturesCache::GetIESTexture(const FSyncContext& InSyncContext,
+																		   const FString&	   InIESFileName)
+{
+	CreateCacheFolders();
+	FTexturesCache::FIESTexturesCacheElem* found = IESTextures.Find(InIESFileName);
+	if (found == nullptr)
+	{
+		// Find and copy the texture in the cache.
+		FString IESFilePath(GSStringToUE(CopyIESFile(UEToGSString(*InIESFileName))));
+
+		// Create IES texture
+		const FString BaseFilename = FPaths::GetBaseFilename(IESFilePath);
+		FString		  TextureName = FDatasmithUtils::SanitizeObjectName(BaseFilename + TEXT("_IES"));
+
+		TSharedPtr< IDatasmithTextureElement > Texture = FDatasmithSceneFactory::CreateTexture(*TextureName);
+		Texture->SetTextureMode(EDatasmithTextureMode::Ies);
+		Texture->SetLabel(*BaseFilename);
+		Texture->SetFile(*IESFilePath);
+
+		InSyncContext.GetScene().AddTexture(Texture);
 
+		found = &IESTextures.Add(InIESFileName, TextureName);
+	}
+
+	return *found;
+}
+
+void FTexturesCache::CreateCacheFolders()
+{
+	if (bCacheFoldersCreated == false)
+	{
+		{
+			// Create the asset folder if it's not present
+			IO::Location AssetsFolderLocation(AssetsCache);
+			IO::Folder	 AssetsFolder(AssetsFolderLocation, IO::Folder::Create);
+			if (AssetsFolder.GetStatus() != NoError)
+			{
+				UE_AC_ReportF("Unable to create/access assets cache folder: \"%s\"", AssetsCache.ToUtf8());
+				UE_AC::ThrowGSError(AssetsFolder.GetStatus(), __FILE__, __LINE__);
+			}
+		}
+
+		{
+			// Create the textures folder if it's not present
+			IO::Location TexturesLocation(AbsolutePath);
+			IO::Folder	 TexturesFolder(TexturesLocation, IO::Folder::Create);
+			if (TexturesFolder.GetStatus() != NoError)
+			{
+				UE_AC_ReportF("Unable to create/access cache textures cache folder: \"%s\"", AbsolutePath.ToUtf8());
+				UE_AC::ThrowGSError(TexturesFolder.GetStatus(), __FILE__, __LINE__);
+			}
+		}
+
+		bCacheFoldersCreated = true;
+	}
+}
+
+// Write the texture to the cache
+void FTexturesCache::WriteTexture(const ModelerAPI::Texture& InACTexture, const GS::UniString& InPath,
+								  bool InIsFingerprint) const
+{
 	IO::Location TextureLoc(InPath);
 
 	// If texture already exist, we do nothing
@@ -152,7 +330,9 @@ void FTexturesCache::WriteTexture(const ModelerAPI::Texture& InACTexture, const
 		IO::File TextureFile(TextureLoc);
 		if ((TextureFile.GetStatus() == NoError) &&
 			(TextureFile.IsOpen() || (TextureFile.Open(IO::File::ReadMode) == NoError)))
+		{
 			return;
+		}
 	}
 
 	// Create a pixmap of the same size of the texture
diff --git a/Engine/Source/Programs/Enterprise/Datasmith/DatasmithARCHICADExporter/Private/TexturesCache.h b/Engine/Source/Programs/Enterprise/Datasmith/DatasmithARCHICADExporter/Private/TexturesCache.h
--- a/Engine/Source/Programs/Enterprise/Datasmith/DatasmithARCHICADExporter/Private/TexturesCache.h
+++ b/Engine/Source/Programs/Enterprise/Datasmith/DatasmithARCHICADExporter/Private/TexturesCache.h
@@ -61,23 +61,46 @@ class FTexturesCache
 		bool bUsed = false; // True if this texture is used
 	};
 
-	FTexturesCache();
+	class FIESTexturesCacheElem
+	{
+	  public:
+		FIESTexturesCacheElem(FString& InTexturePath)
+			: TexturePath(InTexturePath)
+		{
+		}
+
+		FString TexturePath; // Path of the written texture
+	};
+
+	FTexturesCache(const GS::UniString& InAssetsCache);
 
 	// Return the texture specified by the index. (Will throw an exception if index is invalid)
 	const FTexturesCacheElem& GetTexture(const FSyncContext& InSyncContext, GS::Int32 InTextureIndex);
 
+	// Create IES texture
+	const FIESTexturesCacheElem& GetIESTexture(const FSyncContext& InSyncContext, const FString& InIESFilePath);
+
+	// Write the texture to the cache
 	void WriteTexture(const ModelerAPI::Texture& inACTexture, const GS::UniString& InPath, bool InIsFingerprint) const;
 
 	size_t GetCount() const { return Textures.Num(); }
 
+	// Insure we have a copy of the IES file in the cache folder
+	GS::UniString CopyIESFile(const GS::UniString& InIESFileName);
+
   private:
-	typedef TMap< GS::Int32, FTexturesCacheElem > MapTextureIndex2CacheElem;
+	void CreateCacheFolders();
+
+	typedef TMap< GS::Int32, FTexturesCacheElem >  MapTextureIndex2CacheElem;
+	typedef TMap< FString, FIESTexturesCacheElem > MapIESTexture2CacheElem;
 
 	typedef TSet< FString >	  SetTexturesIds;
 	MapTextureIndex2CacheElem Textures;
+	MapIESTexture2CacheElem	  IESTextures;
+	bool					  bCacheFoldersCreated = false;
+	GS::UniString			  AssetsCache;
 	GS::UniString			  AbsolutePath;
-	GS::UniString			  RelativePath;
-	bool					  bUseRelative;
+	GS::UniString			  IESTexturesPath;
 
 	typedef TSet< FUniStringPtr > SetStrings;
 
diff --git a/Engine/Source/Programs/Enterprise/Datasmith/DatasmithARCHICADExporter/Private/Utils/AddonTools.cpp b/Engine/Source/Programs/Enterprise/Datasmith/DatasmithARCHICADExporter/Private/Utils/AddonTools.cpp
--- a/Engine/Source/Programs/Enterprise/Datasmith/DatasmithARCHICADExporter/Private/Utils/AddonTools.cpp
+++ b/Engine/Source/Programs/Enterprise/Datasmith/DatasmithARCHICADExporter/Private/Utils/AddonTools.cpp
@@ -82,6 +82,17 @@ utf8_string Utf8StringFormat(const utf8_t* InFmt, ...)
 	return FormattedString;
 }
 
+// Compute Guid of the string
+API_Guid String2API_Guid(const GS::UniString& inString)
+{
+	utf8_string	   Utf8String(inString.ToUtf8());
+	MD5::Generator g;
+	g.Update(Utf8String.c_str(), (unsigned int)Utf8String.size());
+	MD5::FingerPrint fp;
+	g.Finish(fp);
+	return Fingerprint2API_Guid(fp);
+}
+
 // Combine 2 guid in one
 API_Guid CombineGuid(const API_Guid& InGuid1, const API_Guid& InGuid2)
 {
@@ -109,7 +120,7 @@ IO::Location GetCompanyDataDirectory()
 	// Get user application support directory
 	IO::Location location(GetApplicationSupportDirectory());
 	// Get or create company directory
-	location.AppendToLocal(IO::Name(GetGSName(kName_Company)));
+	location.AppendToLocal(IO::Name("Epic"));
 	IO::Folder CompanyFolder(location, IO::Folder::Create);
 	GSErrCode  GSErr = CompanyFolder.GetStatus();
 	if (GSErr != NoError)
@@ -181,11 +192,13 @@ GS::UniString GetLayerName(API_AttributeIndex InLayer)
 		// This case happened for the special ArchiCAD layer
 		if (error == APIERR_DELETED)
 		{
-			LayerName = GetGSName(kName_LayerDeleted);
+			static const GS::UniString LayerDeleted("Layer deleted");
+			LayerName = LayerDeleted;
 		}
 		else
 		{
-			LayerName = GS::UniString::Printf(GetGSName(kName_LayerError), GetErrorName(error));
+			utf8_string LayerError(Utf8StringFormat("Layer error=%s", GetErrorName(error)));
+			LayerName = GS::UniString(LayerError.c_str(), CC_UTF8);
 		}
 	}
 	return LayerName;
diff --git a/Engine/Source/Programs/Enterprise/Datasmith/DatasmithARCHICADExporter/Private/Utils/AddonTools.h b/Engine/Source/Programs/Enterprise/Datasmith/DatasmithARCHICADExporter/Private/Utils/AddonTools.h
--- a/Engine/Source/Programs/Enterprise/Datasmith/DatasmithARCHICADExporter/Private/Utils/AddonTools.h
+++ b/Engine/Source/Programs/Enterprise/Datasmith/DatasmithARCHICADExporter/Private/Utils/AddonTools.h
@@ -2,6 +2,8 @@
 
 #pragma once
 
+#include "WarningsDisabler.h"
+
 #if defined(WIN32)
 	#pragma warning(push)
 	#pragma warning(disable : 4800)
@@ -15,9 +17,11 @@
 	#pragma warning(pop)
 #endif
 
+DISABLE_SDK_WARNINGS_START
 #include "APIEnvir.h"
 #include "ACAPinc.h"
 #include "Md5.hpp"
+DISABLE_SDK_WARNINGS_END
 
 #include "LocalizeTools.h"
 #include "DebugTools.h"
@@ -39,12 +43,21 @@ utf8_string Utf8StringFormat(const utf8_t* InFmt, ...) __printflike(1, 2);
 // Short way to get TCHAR pointer of a GS::UniString
 #define GSStringToUE(InGSString) UTF16_TO_TCHAR(InGSString.ToUStr().Get())
 
+// Convert a UE String to GS::UniString
+inline GS::UniString UEToGSString(const TCHAR* InUEString)
+{
+	return GS::UniString(reinterpret_cast< const GS::UniChar::Layout* >(InUEString));
+}
+
 // Convert an Archicad fingerprint to an API_Guid
 inline const API_Guid& Fingerprint2API_Guid(const MD5::FingerPrint& inFP)
 {
 	return *(const API_Guid*)&inFP;
 }
 
+// Compute Guid of the string
+API_Guid String2API_Guid(const GS::UniString& inString);
+
 // Compute Guid from MD5 of the value
 template < class T > inline API_Guid GuidFromMD5(const T& inV)
 {
@@ -126,33 +139,14 @@ class FAutoHandle
 enum ENames
 {
 	kName_Invalid,
-	kName_Company,
-	kName_Floor,
-	kName_Layer,
-	kName_Group,
-	kName_LayerDeleted,
-	kName_LayerError,
-	kName_ElementType,
-	kName_InvalidGroupId,
-	kName_UndefinedValueType,
-	kName_InvalidVariant,
-	kName_UndefinedCollectionType,
-	kName_InvalidCollectionType,
-	kName_InvalidPrimitiveType,
-	kName_IFCAttributes,
-	kName_IFC_,
-	kName_Camera,
-	kName_Textures,
-	kName__Assets,
 	kName_TextureExtension,
 	kName_TextureMime,
-	kName_SyncOptions,
-	kName_ExportOptions,
-	kName_Unknown,
 	kName_ShowPalette,
 	kName_HidePalette,
-	kName_OkButtonLabel,
-	kName_CancelButtonLabel,
+	kName_DatasmithFileTypeName,
+	kName_ExportToDatasmithFile,
+	kName_StartAutoSync,
+	kName_PauseAutoSync,
 	kName_Undefined,
 	kName_NBNames
 };
diff --git a/Engine/Source/Programs/Enterprise/Datasmith/DatasmithARCHICADExporter/Private/Utils/CurrentOS.h b/Engine/Source/Programs/Enterprise/Datasmith/DatasmithARCHICADExporter/Private/Utils/CurrentOS.h
--- a/Engine/Source/Programs/Enterprise/Datasmith/DatasmithARCHICADExporter/Private/Utils/CurrentOS.h
+++ b/Engine/Source/Programs/Enterprise/Datasmith/DatasmithARCHICADExporter/Private/Utils/CurrentOS.h
@@ -9,6 +9,7 @@
 BEGIN_NAMESPACE_UE_AC
 
 #if defined(WIN32)
+
 // Throw a runtime_error for last windows error
 void ThrowWinError(DWORD winErr, const utf8_t* InFile, int InLineNo);
 
@@ -25,6 +26,9 @@ void ThrowWinError(DWORD winErr, const utf8_t* InFile, int InLineNo);
 // String conversions utf8_t vs wchar_t
 std::wstring Utf8ToUtf16(const utf8_t* s);
 utf8_string	 Utf16ToUtf8(const wchar_t* s);
+
+void SetThreadName(const char* InName);
+
 #endif
 
 typedef std::vector< utf8_string > VecStrings;
diff --git a/Engine/Source/Programs/Enterprise/Datasmith/DatasmithARCHICADExporter/Private/Utils/Element2String.cpp b/Engine/Source/Programs/Enterprise/Datasmith/DatasmithARCHICADExporter/Private/Utils/Element2String.cpp
--- a/Engine/Source/Programs/Enterprise/Datasmith/DatasmithARCHICADExporter/Private/Utils/Element2String.cpp
+++ b/Engine/Source/Programs/Enterprise/Datasmith/DatasmithARCHICADExporter/Private/Utils/Element2String.cpp
@@ -2,6 +2,7 @@
 
 #include "Element2String.h"
 #include "ElementTools.h"
+#include "TAssValueName.h"
 
 BEGIN_NAMESPACE_UE_AC
 
@@ -53,7 +54,7 @@ utf8_string FElement2String::GetElementAsShortString(const API_Element& InElemen
 	FElementTools::GetInfoString(InElement.header.guid, &ElementInfo);
 
 	ElementString += Utf8StringFormat("Element \"%s\", type=%s, Guid={%s}, Layer=\"%s\", Floor=%d\n",
-									  ElementInfo.ToUtf8(), FElementTools::TypeName(InElement.header.typeID).ToUtf8(),
+									  ElementInfo.ToUtf8(), FElementTools::TypeName(InElement.header.typeID),
 									  APIGuidToString(InElement.header.guid).ToUtf8(),
 									  GetLayerName(InElement.header.layer).ToUtf8(), InElement.header.floorInd);
 
@@ -122,7 +123,7 @@ utf8_string FElement2String::GetElementAsString(const API_Element& InElement)
 
 	ElementString += Utf8StringFormat(
 		"\tElement \"%s\", type=%s, Guid={%s}\n\t\t\tMemoMD5={%s} Layer=\"%s\", Floor=%d\n", ElementInfo.ToUtf8(),
-		FElementTools::TypeName(InElement.header.typeID).ToUtf8(), APIGuidToString(InElement.header.guid).ToUtf8(),
+		FElementTools::TypeName(InElement.header.typeID), APIGuidToString(InElement.header.guid).ToUtf8(),
 		APIGuidToString(MemoMD5).ToUtf8(), GetLayerName(InElement.header.layer).ToUtf8(), InElement.header.floorInd);
 	utf8_string HotLinkInfo = GetHotLinkInfo(InElement);
 	if (HotLinkInfo.size())
@@ -294,7 +295,7 @@ utf8_string FElement2String::GetParametersAsString(const API_Guid& InElementGuid
 	FAutoMemo AutoMemo(InElementGuid, APIMemoMask_AddPars);
 	if (AutoMemo.GSErr == NoError)
 	{
-		if (AutoMemo.Memo.params) // Can be null on AC21
+		if (AutoMemo.Memo.params) // Can be null
 		{
 			return GetParametersAsString(AutoMemo.Memo.params);
 		}
@@ -312,7 +313,8 @@ utf8_string FElement2String::GetIFCAttributesAsString(const API_Guid& InElementG
 	utf8_string IFCAttributesString;
 
 	GS::Array< API_IFCAttribute > Attributes;
-	GSErrCode					  GSErr = ACAPI_Element_GetIFCAttributes(InElementGuid, false, &Attributes);
+	bool						  bStoredOnly = false;
+	GSErrCode					  GSErr = ACAPI_Element_GetIFCAttributes(InElementGuid, bStoredOnly, &Attributes);
 	if (GSErr == NoError)
 	{
 		GS::USize NbAttr = Attributes.GetSize();
@@ -345,7 +347,8 @@ utf8_string FElement2String::GetIFCPropertiesAsString(const API_Guid& InElementG
 	utf8_string IFCPropertiesString;
 
 	GS::Array< API_IFCProperty > Properties;
-	GSErrCode					 GSErr = ACAPI_Element_GetIFCProperties(InElementGuid, false, &Properties);
+	bool						 bStoredOnly = false;
+	GSErrCode					 GSErr = ACAPI_Element_GetIFCProperties(InElementGuid, bStoredOnly, &Properties);
 	if (GSErr == NoError)
 	{
 		GS::USize NbProp = Properties.GetSize();
@@ -748,6 +751,33 @@ utf8_string FElement2String::GetPropertyObjectsAsString(const API_Elem_Head& InE
 	#define strlcat(d, s, sz) strncat(d, s, sz - strlen(d) - 1)
 #endif
 
+// clang-format off
+template <>
+FAssValueName::SAssValueName TAssEnumName< API_AddParID >::AssEnumName[] = {
+	{ API_ZombieParT, 				"Zombie" },
+	{ APIParT_Integer, 				"Integer" },
+	{ APIParT_Length, 				"Length" },
+	{ APIParT_Angle, 				"Angle" },
+	{ APIParT_RealNum, 				"RealNum" },
+	{ APIParT_LightSw, 				"LightSw" },
+	{ APIParT_ColRGB, 				"ColRGB" },
+	{ APIParT_Intens, 				"Intens" },
+	{ APIParT_LineTyp, 				"LineTyp" },
+	{ APIParT_Mater, 				"Mater" },
+	{ APIParT_FillPat, 				"FillPat" },
+	{ APIParT_PenCol, 				"PenCol" },
+	{ APIParT_CString, 				"CString" },
+	{ APIParT_Boolean, 				"Boolean" },
+	{ APIParT_Separator, 			"Separator" },
+	{ APIParT_Title, 				"Title" },
+	{ APIParT_BuildingMaterial, 	"BuildingMaterial" },
+	{ APIParT_Profile, 				"Profile" },
+	{ APIParT_Dictionary, 			"Dictionary" },
+
+	{ -1, nullptr }
+};
+// clang-format on
+
 // Tool:Return parameters as a string
 utf8_string FElement2String::GetParametersAsString(const API_AddParType* const* InParamsHandle)
 {
@@ -872,23 +902,27 @@ utf8_string FElement2String::GetParametersAsString(const API_AddParType* const*
 				case APIParT_Mater:
 				case APIParT_FillPat:
 				case APIParT_PenCol:
-					Params += Utf8StringFormat("\t\t\t%s%d\n", ParamName.c_str(), (Int32)Param.value.real);
+					Params +=
+						Utf8StringFormat("\t\t\t%s %s:%d\n", ParamName.c_str(),
+										 TAssEnumName< API_AddParID >::GetName(Param.typeID), (Int32)Param.value.real);
 					break;
 				case APIParT_Boolean:
-					Params +=
-						Utf8StringFormat("\t\t\t%s%s\n", ParamName.c_str(), Param.value.real != 0 ? "true" : "false");
+					Params += Utf8StringFormat("\t\t\t%s Boolean:%s\n", ParamName.c_str(),
+											   Param.value.real != 0 ? "true" : "false");
 					break;
 
 				case APIParT_Length:
 				case APIParT_Angle:
 				case APIParT_RealNum:
 				case APIParT_ColRGB:
-					Params += Utf8StringFormat("\t\t\t%s%lg\n", ParamName.c_str(), Param.value.real);
+					Params += Utf8StringFormat("\t\t\t%s %s:%lg\n", ParamName.c_str(),
+											   TAssEnumName< API_AddParID >::GetName(Param.typeID), Param.value.real);
 					break;
 
 				case APIParT_CString:
 				case APIParT_Title:
-					Params += Utf8StringFormat("\t\t\t%s\"%s\"\n", ParamName.c_str(),
+					Params += Utf8StringFormat("\t\t\t%s %s:\"%s\"\n", ParamName.c_str(),
+											   TAssEnumName< API_AddParID >::GetName(Param.typeID),
 											   GS::UniString(Param.value.uStr).ToUtf8());
 					break;
 
@@ -953,7 +987,8 @@ GS::UniString FElement2String::GetPropertyGroupName(const API_Guid& InGroupGuid)
 	Group.guid = InGroupGuid;
 	if (ACAPI_Property_GetPropertyGroup(Group) != NoError)
 	{
-		return GetGSName(kName_InvalidGroupId);
+		static const GS::UniString InvalidGroupId("Invalid group id");
+		return InvalidGroupId;
 	}
 
 	static const utf8_t* typeNames[] = {"Static  ", "Dynamic ", "Custom  "};
@@ -1081,7 +1116,7 @@ utf8_string FElement2String::GetIFCPropertyAnyValue(const API_IFCPropertyAnyValu
 			return Utf8StringFormat("\"%s\"", InValue.stringValue.ToUtf8());
 	}
 	UE_AC_DebugF("GetIFCPropertyAnyValue primitiveType=%d\n", InValue.primitiveType);
-	return Utf8StringFormat(GetStdName(kName_InvalidPrimitiveType), InValue.primitiveType);
+	return Utf8StringFormat("Invalid primitive type(%d)", InValue.primitiveType);
 }
 
 inline static bool IsSpace(char C)
@@ -1112,4 +1147,27 @@ void FElement2String::RemoveLeadingAndTrailing(utf8_string* IOString)
 	IOString->erase(0, IndexChar);
 }
 
+utf8_string FDump2String::ListLibraries()
+{
+	utf8_string DumpString;
+
+	GS::Array< API_LibraryInfo > LibInfoArray;
+
+	if (ACAPI_Environment(APIEnv_GetLibrariesID, &LibInfoArray) == NoError)
+	{
+		DumpString = Utf8StringFormat("ListLibraries - The number of loaded libraries is %u\n", LibInfoArray.GetSize());
+		for (UInt32 IndexLibrary = 0; IndexLibrary < LibInfoArray.GetSize(); IndexLibrary++)
+		{
+			const API_LibraryInfo& LibInfo = LibInfoArray[IndexLibrary];
+			DumpString += Utf8StringFormat("\tLibrary #%-2d name=\"%s\" type=%d\n", IndexLibrary, LibInfo.name.ToUtf8(),
+										   LibInfo.libraryType);
+			GS::UniString LibPath;
+			LibInfo.location.ToPath(&LibPath);
+			DumpString += Utf8StringFormat("\t\tPath=\"%s\"\n", LibPath.ToUtf8());
+		}
+	}
+
+	return DumpString;
+}
+
 END_NAMESPACE_UE_AC
diff --git a/Engine/Source/Programs/Enterprise/Datasmith/DatasmithARCHICADExporter/Private/Utils/Element2String.h b/Engine/Source/Programs/Enterprise/Datasmith/DatasmithARCHICADExporter/Private/Utils/Element2String.h
--- a/Engine/Source/Programs/Enterprise/Datasmith/DatasmithARCHICADExporter/Private/Utils/Element2String.h
+++ b/Engine/Source/Programs/Enterprise/Datasmith/DatasmithARCHICADExporter/Private/Utils/Element2String.h
@@ -69,4 +69,9 @@ class FElement2String
 	static void RemoveLeadingAndTrailing(utf8_string* IOString);
 };
 
+class FDump2String
+{
+	static utf8_string ListLibraries();
+};
+
 END_NAMESPACE_UE_AC
diff --git a/Engine/Source/Programs/Enterprise/Datasmith/DatasmithARCHICADExporter/Private/Utils/ElementTools.cpp b/Engine/Source/Programs/Enterprise/Datasmith/DatasmithARCHICADExporter/Private/Utils/ElementTools.cpp
--- a/Engine/Source/Programs/Enterprise/Datasmith/DatasmithARCHICADExporter/Private/Utils/ElementTools.cpp
+++ b/Engine/Source/Programs/Enterprise/Datasmith/DatasmithARCHICADExporter/Private/Utils/ElementTools.cpp
@@ -6,6 +6,104 @@
 BEGIN_NAMESPACE_UE_AC
 
 // clang-format off
+
+template <>
+FAssValueName::SAssValueName TAssEnumName< API_ElemTypeID >::AssEnumName[] = {
+	{ API_ZombieElemID, 				"Zombie" },
+
+	{ API_WallID, 						"Wall" },
+	{ API_ColumnID, 					"Column" },
+	{ API_BeamID, 						"Beam" },
+	{ API_WindowID, 					"Window" },
+	{ API_DoorID, 						"Door" },
+	{ API_ObjectID, 					"Object" },
+	{ API_LampID, 						"Lamp" },
+	{ API_SlabID, 						"Slab" },
+	{ API_RoofID, 						"Roof" },
+	{ API_MeshID, 						"Mesh" },
+
+	{ API_DimensionID, 					"Dimension" },
+	{ API_RadialDimensionID, 			"RadialDimension" },
+	{ API_LevelDimensionID, 			"LevelDimension" },
+	{ API_AngleDimensionID, 			"AngleDimension" },
+
+	{ API_TextID, 						"Text" },
+	{ API_LabelID, 						"Label" },
+	{ API_ZoneID, 						"Zone" },
+
+	{ API_HatchID, 						"Hatch" },
+	{ API_LineID, 						"Line" },
+	{ API_PolyLineID, 					"PolyLine" },
+	{ API_ArcID, 						"Arc" },
+	{ API_CircleID, 					"Circle" },
+	{ API_SplineID, 					"Spline" },
+	{ API_HotspotID, 					"Hotspot" },
+
+	{ API_CutPlaneID, 					"CutPlane" },
+	{ API_CameraID, 					"Camera" },
+	{ API_CamSetID, 					"CamSet" },
+
+	{ API_GroupID, 						"Group" },
+	{ API_SectElemID, 					"SectElem" },
+
+	{ API_DrawingID, 					"Drawing" },
+	{ API_PictureID, 					"Picture" },
+	{ API_DetailID, 					"Detail" },
+	{ API_ElevationID, 					"Elevation" },
+	{ API_InteriorElevationID, 			"InteriorElevation" },
+	{ API_WorksheetID, 					"Worksheet" },
+
+	{ API_HotlinkID, 					"Hotlink" },
+
+	{ API_CurtainWallID, 				"CurtainWall" },
+	{ API_CurtainWallSegmentID, 		"CurtainWallSegment" },
+	{ API_CurtainWallFrameID, 			"CurtainWallFrame" },
+	{ API_CurtainWallPanelID, 			"CurtainWallPanel" },
+	{ API_CurtainWallJunctionID, 		"CurtainWallJunction" },
+	{ API_CurtainWallAccessoryID, 		"CurtainWallAccessory" },
+	{ API_ShellID, 						"Shell" },
+	{ API_SkylightID, 					"Skylight" },
+	{ API_MorphID, 						"Morph" },
+	{ API_ChangeMarkerID, 				"ChangeMarker" },
+	{ API_StairID, 						"Stair" },
+	{ API_RiserID, 						"Riser" },
+	{ API_TreadID, 						"Tread" },
+	{ API_StairStructureID, 			"StairStructure" },
+	{ API_RailingID, 					"Railing" },
+	{ API_RailingToprailID, 			"RailingToprail" },
+	{ API_RailingHandrailID, 			"RailingHandrail" },
+	{ API_RailingRailID, 				"RailingRail" },
+	{ API_RailingPostID, 				"RailingPost" },
+	{ API_RailingInnerPostID, 			"RailingInnerPost" },
+	{ API_RailingBalusterID, 			"RailingBaluster" },
+	{ API_RailingPanelID, 				"RailingPanel" },
+	{ API_RailingSegmentID, 			"RailingSegment" },
+	{ API_RailingNodeID, 				"RailingNode" },
+	{ API_RailingBalusterSetID, 		"RailingBalusterSet" },
+	{ API_RailingPatternID, 			"RailingPattern" },
+	{ API_RailingToprailEndID, 			"RailingToprailEnd" },
+	{ API_RailingHandrailEndID, 		"RailingHandrailEnd" },
+	{ API_RailingRailEndID, 			"RailingRailEnd" },
+	{ API_RailingToprailConnectionID, 	"RailingToprailConnection" },
+	{ API_RailingHandrailConnectionID, 	"RailingHandrailConnection" },
+	{ API_RailingRailConnectionID, 		"RailingRailConnection" },
+	{ API_RailingEndFinishID, 			"RailingEndFinish" },
+
+	{ API_AnalyticalSupportID, 			"AnalyticalSupport" },
+	{ API_AnalyticalLinkID, 			"AnalyticalLink" },
+
+	{ API_BeamSegmentID, 				"BeamSegment" },
+	{ API_ColumnSegmentID, 				"ColumnSegment" },
+	{ API_OpeningID, 					"Opening" },
+#if AC_VERSION > 24
+	{ API_AnalyticalPointLoadID, 		"AnalyticalPointLoad" },
+	{ API_AnalyticalEdgeLoadID, 		"AnalyticalEdgeLoad" },
+	{ API_AnalyticalSurfaceLoadID, 		"AnalyticalSurfaceLoad" },
+#endif
+
+	{ -1, nullptr }
+};
+
 template <>
 FAssValueName::SAssValueName TAssEnumName< API_ElemVariationID >::AssEnumName[] = {
 	ValueName(APIVarId_Generic),
@@ -33,18 +131,18 @@ FAssValueName::SAssValueName TAssEnumName< API_ElemVariationID >::AssEnumName[]
 	ValueName(APIVarId_LabelCWJunction),
 #endif
 
-	ValueName(APIVarId_SymbStair),
-	ValueName(APIVarId_WallEnd),
-
-	ValueName(APIVarId_Door),
-	ValueName(APIVarId_Skylight),
-	ValueName(APIVarId_Object),
-	ValueName(APIVarId_GridElement),
-	ValueName(APIVarId_Light),
-	ValueName(APIVarId_CornerWindow),
+	{ APIVarId_SymbStair, 		"SymbStair" },
+	{ APIVarId_WallEnd, 		"WallEnd" },
+	{ APIVarId_SymbStair, 		"SymbStair" },
+	{ APIVarId_Door, 			"Door" },
+	{ APIVarId_Object, 			"Object" },
+	{ APIVarId_GridElement, 	"GridElement" },
+	{ APIVarId_Light, 			"Light" },
+	{ APIVarId_CornerWindow, 	"CornerWindow" },
 
 	EnumEnd(-1)
 };
+
 // clang-format on
 
 // Tool: return the info string (≈ name)
@@ -60,28 +158,22 @@ bool FElementTools::GetInfoString(const API_Guid& InGUID, GS::UniString* OutStri
 }
 
 // Tool: Return the localize name for element type id
-const GS::UniString& FElementTools::TypeName(API_ElemTypeID InElementType)
+const utf8_t* FElementTools::TypeName(API_ElemTypeID InElementType)
 {
 	UE_AC_Assert(API_FirstElemType <= InElementType && InElementType <= API_LastElemType);
 
-	static std::unique_ptr< GS::UniString > TypeNames[API_LastElemType + 1] = {};
+	static const utf8_t* TypeNames[API_LastElemType + 1] = {};
 
 	if (TypeNames[InElementType] == nullptr)
 	{
-		TypeNames[InElementType] = std::make_unique< GS::UniString >();
-		GSErrCode GSErr =
-			ACAPI_Goodies(APIAny_GetElemTypeNameID, (void*)(size_t)InElementType, TypeNames[InElementType].get());
-		if (GSErr != NoError)
-		{
-			UE_AC_DebugF("CElement::TypeName - Error %d for type%d\n", GSErr, InElementType);
-		}
+		TypeNames[InElementType] = TAssEnumName< API_ElemTypeID >::GetName(InElementType);
 	}
 
-	return *TypeNames[InElementType];
+	return TypeNames[InElementType];
 }
 
 // Tool: Return the localize name for element's type
-const GS::UniString& FElementTools::TypeName(const API_Guid& InElementGuid)
+const utf8_t* FElementTools::TypeName(const API_Guid& InElementGuid)
 {
 	API_Elem_Head ElementHead;
 	Zap(&ElementHead);
@@ -147,6 +239,100 @@ GS::Int32 FElementTools::GetLibPartIndex(const API_Element& InElement)
 	}
 }
 
+class GSUnID
+{
+  public:
+	typedef char Buffer[128];
+	GS::Guid	 Main;
+	GS::Guid	 Rev;
+
+	static const char* UnIDNullStr;
+
+	GS::GSErrCode InitWithString(const Buffer inStr);
+};
+const char* GSUnID::UnIDNullStr = "{00000000-0000-0000-0000-000000000000}-{00000000-0000-0000-0000-000000000000}";
+
+GS::GSErrCode GSUnID::InitWithString(const Buffer inStr)
+{
+	GS::Guid	main, rev;
+	char		guidStr[60];
+	const Int32 strLen = sizeof(guidStr) / sizeof(guidStr[0]);
+	Int32		i;
+
+	const char* uiStr = inStr;
+	while (*uiStr == ' ' || *uiStr == '\t')
+	{
+		uiStr++;
+	}
+
+	if (*uiStr++ != '{')
+	{
+		return Error;
+	}
+
+	for (i = 0; i < strLen - 1 && *uiStr != 0 && *uiStr != '}'; i++, uiStr++)
+	{
+		guidStr[i] = *uiStr;
+	}
+	if (*uiStr != '}')
+	{
+		return Error;
+	}
+	guidStr[i] = 0;
+
+	if (main.ConvertFromString(guidStr) != NoError)
+	{
+		return Error;
+	}
+
+	if (*(++uiStr)++ != '-' || *uiStr++ != '{')
+	{
+		return Error;
+	}
+
+	for (i = 0; i < strLen - 1 && *uiStr != 0 && *uiStr != '}'; i++, uiStr++)
+	{
+		guidStr[i] = *uiStr;
+	}
+	if (*uiStr != '}')
+	{
+		return Error;
+	}
+	guidStr[i] = 0;
+
+	if (rev.ConvertFromString(guidStr) != NoError)
+	{
+		return Error;
+	}
+
+	Main = main;
+	Rev = rev;
+
+	return NoError;
+}
+
+GS::Guid FElementTools::GetLibPartId(const API_Elem_Head& InElement)
+{
+	GSUnID::Buffer lpfUnID = {0};
+
+	GS::GSErrCode GSErr = ACAPI_Goodies(APIAny_GetElemLibPartUnIdID, const_cast< API_Elem_Head* >(&InElement), lpfUnID);
+	if (GSErr == NoError)
+	{
+		GSUnID LibPartUnIdID;
+		GSErr = LibPartUnIdID.InitWithString(lpfUnID);
+		if (GSErr == NoError)
+		{
+			return LibPartUnIdID.Main;
+		}
+		UE_AC_DebugF("CElementID::InitLibPartInfo - InitWithString (error=%s)\n", GetErrorName(GSErr));
+	}
+	else
+	{
+		UE_AC_DebugF("CElementID::InitLibPartInfo - Can't get element lib part id (error=%s)\n", GetErrorName(GSErr));
+	}
+	return GS::NULLGuid;
+}
+
 // Tool: return element's owner guid
 API_Guid FElementTools::GetOwner(const API_Element& InElement)
 {
diff --git a/Engine/Source/Programs/Enterprise/Datasmith/DatasmithARCHICADExporter/Private/Utils/ElementTools.h b/Engine/Source/Programs/Enterprise/Datasmith/DatasmithARCHICADExporter/Private/Utils/ElementTools.h
--- a/Engine/Source/Programs/Enterprise/Datasmith/DatasmithARCHICADExporter/Private/Utils/ElementTools.h
+++ b/Engine/Source/Programs/Enterprise/Datasmith/DatasmithARCHICADExporter/Private/Utils/ElementTools.h
@@ -13,10 +13,10 @@ class FElementTools
 	static bool GetInfoString(const API_Guid& InGUID, GS::UniString* OutString);
 
 	// Tool: Return the localize name for element type id
-	static const GS::UniString& TypeName(API_ElemTypeID InElementType);
+	static const utf8_t* TypeName(API_ElemTypeID InElementType);
 
 	// Tool: Return the localize name for element's type
-	static const GS::UniString& TypeName(const API_Guid& InElementGuid);
+	static const utf8_t* TypeName(const API_Guid& InElementGuid);
 
 	// Tool: Return the variation as string
 	static utf8_string GetVariationAsString(API_ElemVariationID InVariation);
@@ -24,6 +24,8 @@ class FElementTools
 	// Tool: return libpart index (or 0 if no libpart)
 	static GS::Int32 GetLibPartIndex(const API_Element& InElement);
 
+	static GS::Guid GetLibPartId(const API_Elem_Head& InElement);
+
 	// Tool: return element's owner guid
 	static API_Guid GetOwner(const API_Element& InElement);
 
diff --git a/Engine/Source/Programs/Enterprise/Datasmith/DatasmithARCHICADExporter/Private/Utils/Error.h b/Engine/Source/Programs/Enterprise/Datasmith/DatasmithARCHICADExporter/Private/Utils/Error.h
--- a/Engine/Source/Programs/Enterprise/Datasmith/DatasmithARCHICADExporter/Private/Utils/Error.h
+++ b/Engine/Source/Programs/Enterprise/Datasmith/DatasmithARCHICADExporter/Private/Utils/Error.h
@@ -61,7 +61,7 @@ template < typename Functor > GSErrCode TryFunctionCatchAndAlert(const utf8_t* I
 	}
 	catch (...)
 	{
-		ShowAlert(GetStdName(kName_Unknown), InFctName);
+		ShowAlert("Unknown", InFctName);
 	}
 
 	return GSErr;
diff --git a/Engine/Source/Programs/Enterprise/Datasmith/DatasmithARCHICADExporter/Private/Utils/LibPartInfo.cpp b/Engine/Source/Programs/Enterprise/Datasmith/DatasmithARCHICADExporter/Private/Utils/LibPartInfo.cpp
new file mode 100644
--- /dev/null
+++ b/Engine/Source/Programs/Enterprise/Datasmith/DatasmithARCHICADExporter/Private/Utils/LibPartInfo.cpp
@@ -0,0 +1,165 @@
+// Copyright Epic Games, Inc. All Rights Reserved.
+
+#include "LibPartInfo.h"
+
+#include <stdexcept>
+
+BEGIN_NAMESPACE_UE_AC
+
+const char* FGSUnID::UnIDNullStr = "{00000000-0000-0000-0000-000000000000}-{00000000-0000-0000-0000-000000000000}";
+
+GS::GSErrCode FGSUnID::InitWithString(const Buffer InStr)
+{
+	GS::Guid	main, rev;
+	char		guidStr[60];
+	const Int32 strLen = sizeof(guidStr) / sizeof(guidStr[0]);
+	Int32		i;
+
+	const char* uiStr = InStr;
+	while (*uiStr == ' ' || *uiStr == '\t')
+	{
+		uiStr++;
+	}
+
+	if (*uiStr++ != '{')
+	{
+		return Error;
+	}
+
+	for (i = 0; i < strLen - 1 && *uiStr != 0 && *uiStr != '}'; i++, uiStr++)
+	{
+		guidStr[i] = *uiStr;
+	}
+	if (*uiStr != '}')
+	{
+		return Error;
+	}
+	guidStr[i] = 0;
+
+	if (main.ConvertFromString(guidStr) != NoError)
+	{
+		return Error;
+	}
+
+	if (*(++uiStr)++ != '-' || *uiStr++ != '{')
+	{
+		return Error;
+	}
+
+	for (i = 0; i < strLen - 1 && *uiStr != 0 && *uiStr != '}'; i++, uiStr++)
+	{
+		guidStr[i] = *uiStr;
+	}
+	if (*uiStr != '}')
+	{
+		return Error;
+	}
+	guidStr[i] = 0;
+
+	if (rev.ConvertFromString(guidStr) != NoError)
+	{
+		return Error;
+	}
+
+	Main = main;
+	Rev = rev;
+
+	return NoError;
+}
+
+// Intialize
+void FLibPartInfo::Initialize(GS::Int32 InIndex)
+{
+	if (Index != 0)
+	{
+		return;
+	}
+	Index = InIndex;
+	API_LibPart LibPart;
+	Zap(&LibPart);
+	LibPart.index = InIndex;
+	GSErrCode err = ACAPI_LibPart_Get(&LibPart);
+	if (err == APIERR_MISSINGDEF)
+	{
+		UE_AC_TraceF("FLibPartInfo::FLibPartInfo - LibPart missing index=%d, name=\"%s\"\n", InIndex,
+					 GS::UniString(LibPart.docu_UName).ToUtf8());
+	}
+	else if (err != GS::NoError)
+	{
+		throw std::runtime_error(
+			Utf8StringFormat("FLibPartInfo::FLibPartInfo - LibPart index=%d, error=%d\n", InIndex, err));
+	}
+	Name = LibPart.docu_UName;
+	Guid.InitWithString(LibPart.ownUnID);
+	if (Guid.Main.IsNull())
+	{
+		// LibPart id can be null
+		if (strncmp(LibPart.ownUnID, FGSUnID::UnIDNullStr, sizeof(LibPart.ownUnID)) != 0)
+			UE_AC_DebugF("FLibPartInfo::FLibPartInfo - LibPart id is null index=%d, name=\"%s\", id=\"%s\"\n", InIndex,
+						 GS::UniString(LibPart.docu_UName).ToUtf8(), LibPart.ownUnID);
+
+		Guid.Main = APIGuid2GSGuid(String2API_Guid(Name)); // Simulate guid from name
+	}
+}
+
+const API_AddParType* GetParameter(API_AddParType** InParameters, const char* InParameterName)
+{
+	UInt32 nParams = BMGetHandleSize((GSHandle)InParameters) / sizeof(API_AddParType);
+
+	GS::UInt32 j = 0;
+	for (j = 0; j < nParams; ++j)
+	{
+		const API_AddParType* Parameter = *InParameters + j;
+		if (strncmp(Parameter->name, InParameterName, API_NameLen) == 0)
+		{
+			return Parameter;
+		}
+	}
+	return nullptr;
+}
+
+bool GetParameter(API_AddParType** InParameters, const char* InParameterName, GS::UniString* OutString)
+{
+	const API_AddParType* Parameter = GetParameter(InParameters, InParameterName);
+	if (Parameter == nullptr || Parameter->typeMod != API_ParSimple ||
+		(Parameter->typeID != APIParT_CString && Parameter->typeID == APIParT_Title))
+	{
+		return false;
+	}
+
+	*OutString = GS::UniString(Parameter->value.uStr);
+	return true;
+}
+
+bool GetParameter(API_AddParType** InParameters, const char* InParameterName, double* OutValue)
+{
+	const API_AddParType* Parameter = GetParameter(InParameters, InParameterName);
+	if (Parameter == nullptr || Parameter->typeMod != API_ParSimple)
+	{
+		return false;
+	}
+	switch (Parameter->typeID)
+	{
+		case APIParT_Integer:
+		case APIParT_LightSw:
+		case APIParT_Intens:
+		case APIParT_LineTyp:
+		case APIParT_Mater:
+		case APIParT_FillPat:
+		case APIParT_PenCol:
+		case APIParT_Boolean:
+		case APIParT_Length:
+		case APIParT_Angle:
+		case APIParT_RealNum:
+		case APIParT_ColRGB:
+			break;
+
+		default:
+			return false;
+	}
+
+	*OutValue = Parameter->value.real;
+	return true;
+}
+
+END_NAMESPACE_UE_AC
diff --git a/Engine/Source/Programs/Enterprise/Datasmith/DatasmithARCHICADExporter/Private/Utils/LibPartInfo.h b/Engine/Source/Programs/Enterprise/Datasmith/DatasmithARCHICADExporter/Private/Utils/LibPartInfo.h
new file mode 100644
--- /dev/null
+++ b/Engine/Source/Programs/Enterprise/Datasmith/DatasmithARCHICADExporter/Private/Utils/LibPartInfo.h
@@ -0,0 +1,50 @@
+// Copyright Epic Games, Inc. All Rights Reserved.
+
+#pragma once
+
+#include "AddonTools.h"
+
+BEGIN_NAMESPACE_UE_AC
+
+class FGSUnID
+{
+  public:
+	typedef char Buffer[128];
+	GS::Guid	 Main;
+	GS::Guid	 Rev;
+
+	static const char* UnIDNullStr;
+
+	GS::GSErrCode InitWithString(const Buffer InStr);
+
+	// Equality operator needed for use as FMap key
+	bool operator==(const FGSUnID& InOther) const { return Main == InOther.Main && Rev == InOther.Rev; }
+};
+
+inline uint32 GetTypeHash(const FGSUnID& A)
+{
+	return FCrc::TypeCrc32(A);
+}
+
+class FLibPartInfo
+{
+  public:
+	// Constructor
+	FLibPartInfo()
+		: Index(0)
+	{
+	}
+
+	// Intialize
+	void Initialize(GS::Int32 InIndex);
+
+	GS::Int32	  Index; // LibPart index
+	FGSUnID		  Guid; // LibPart Guid
+	GS::UniString Name; // LibPart name
+};
+
+const API_AddParType* GetParameter(API_AddParType** InParameters, const char* InParameterName);
+bool GetParameter(API_AddParType** InParameters, const char* InParameterName, GS::UniString* OutString);
+bool GetParameter(API_AddParType** InParameters, const char* InParameterName, double* OutValue);
+
+END_NAMESPACE_UE_AC
diff --git a/Engine/Source/Programs/Enterprise/Datasmith/DatasmithARCHICADExporter/Private/Utils/LocalizeTools.cpp b/Engine/Source/Programs/Enterprise/Datasmith/DatasmithARCHICADExporter/Private/Utils/LocalizeTools.cpp
--- a/Engine/Source/Programs/Enterprise/Datasmith/DatasmithARCHICADExporter/Private/Utils/LocalizeTools.cpp
+++ b/Engine/Source/Programs/Enterprise/Datasmith/DatasmithARCHICADExporter/Private/Utils/LocalizeTools.cpp
@@ -177,7 +177,7 @@ const GS::UniString& FMultiString::GetGSString(Int32 InStrIndex)
 {
 	if (!IsValidIndex(InStrIndex))
 	{
-		static GS::UniString invalid("(Invalid string index)");
+		static const GS::UniString invalid("(Invalid string index)");
 		return invalid;
 	}
 	return mGSStrings[size_t(InStrIndex - StartIndex)];
diff --git a/Engine/Source/Programs/Enterprise/Datasmith/DatasmithARCHICADExporter/Private/Utils/Progression.cpp b/Engine/Source/Programs/Enterprise/Datasmith/DatasmithARCHICADExporter/Private/Utils/Progression.cpp
new file mode 100644
--- /dev/null
+++ b/Engine/Source/Programs/Enterprise/Datasmith/DatasmithARCHICADExporter/Private/Utils/Progression.cpp
@@ -0,0 +1,266 @@
+// Copyright Epic Games, Inc. All Rights Reserved.
+
+#include "Progression.h"
+
+#include "Error.h"
+
+#include "CurrentOS.h"
+
+#include "Guard.hpp"
+
+#ifdef WIN32
+#else
+	#include <sys/time.h>
+#endif
+
+BEGIN_NAMESPACE_UE_AC
+
+inline void* ACAPI_Interface_Str(const GS::UniString& InStr)
+{
+	return (void*)&InStr;
+}
+
+// Constructor (SetUp progression window)
+FProgression::FProgression(int InResID, EPhaseStrId InTitle, short InNbPhases, EMode InCancelMode,
+						   volatile bool* OutUserCancelled)
+	: ResID(InResID)
+	, CurrentPhase(kCommonPhaseInvalid)
+	, CurrentValue(0)
+	, bProgressionShown(false)
+	, ErrorCode(NoError)
+	, ErrorMgs(NULL)
+	, CancelMode(InCancelMode)
+	, bUserCancelled(OutUserCancelled)
+	, NextPhase(kCommonPhaseInvalid)
+	, NextPhaseMaxValue(0)
+	, NextCurrentValue(0)
+	, mCV(AccessControl)
+{
+#ifdef WIN32
+	LARGE_INTEGER freq;
+	QueryPerformanceFrequency(&freq);
+	InvFreq = 1.0 / freq.QuadPart;
+#endif
+	LastUpdate = GetHiResTime();
+
+	// Intialize the progression window and specify the number of phases
+	API_ProcessControlTypeID DisableMenu = API_MenuCommandDisabled;
+	ErrorCode =
+		ACAPI_Interface(APIIo_InitProcessWindowID, ACAPI_Interface_Str(GetText(InTitle)), &InNbPhases, &DisableMenu);
+	if (ErrorCode == NoError)
+	{
+		bProgressionShown = true;
+	}
+	else
+	{
+		ErrorMgs = "initialising progression window";
+	}
+}
+
+// Destructor (Automatically close the progression window)
+FProgression::~FProgression()
+{
+	// If we succeed to initialize progression window
+	if (bProgressionShown)
+	{
+		// Close the progression window
+		GSErrCode ErrCode = ACAPI_Interface(APIIo_CloseProcessWindowID, nullptr, nullptr);
+		if (ErrorCode == NoError)
+		{
+			ErrorCode = ErrCode;
+			ErrorMgs = "closing progression window";
+		}
+	}
+
+	// If we got an progression windows error, we report to stderr
+	if (ErrorCode != NoError)
+	{
+		fprintf(stderr, "Got error (%d) with progress window when %s", (int)ErrorCode, ErrorMgs);
+	}
+}
+
+// Set new cancel mode
+FProgression::EMode FProgression::SetCancelMode(EMode InNewCancelMode)
+{
+	EMode PreviousMode = CancelMode;
+	CancelMode = InNewCancelMode;
+	return PreviousMode;
+}
+
+// Start the next phase
+void FProgression::NewPhase(EPhaseStrId InPhaseId, int InMaxValue)
+{
+	// If we got an error from AC before, we do nothing after
+	if (ErrorCode != NoError)
+	{
+		return;
+	}
+
+	NextPhase = kCommonPhaseInvalid;
+	NextCurrentValue = 0;
+
+	// Setup the next phase
+	CurrentValue = 0;
+	CurrentPhase = InPhaseId;
+	Ratio100 = 100.0;
+	bool showPercent = false;
+	if (InMaxValue > 0)
+	{
+		Ratio100 = 100.0 / InMaxValue;
+		showPercent = true;
+		InMaxValue = 100;
+	}
+	ErrorCode = ACAPI_Interface(APIIo_SetNextProcessPhaseID, ACAPI_Interface_Str(GetText(CurrentPhase)), &InMaxValue,
+								&showPercent);
+	ErrorMgs = "setting next phase";
+
+	LastUpdate = GetHiResTime();
+}
+
+// Start the next phase
+void FProgression::NewPhaseX(EPhaseStrId InPhaseId, int InMaxValue)
+{
+	GS::Guard< GS::Lock > lck(AccessControl);
+	NextPhase = InPhaseId;
+	NextPhaseMaxValue = InMaxValue;
+}
+
+// Advance progression bar for the current value
+void FProgression::NewCurrentValue(int InCurrentValue)
+{
+	// If we got an error from AC before, we do nothing after
+	if (ErrorCode != NoError)
+	{
+		return;
+	}
+
+	// Special case of progression: When we are waiting and not advancing
+	if (InCurrentValue != -1)
+	{
+		InCurrentValue = int(Ratio100 * InCurrentValue + 0.5);
+	}
+	else
+	{
+		InCurrentValue = 0;
+	}
+
+	// Send the value to ArchiCAD
+	if (CurrentValue != InCurrentValue)
+	{
+		CurrentValue = InCurrentValue;
+		ErrorCode = ACAPI_Interface(APIIo_SetProcessValueID, &CurrentValue, nullptr);
+		ErrorMgs = "setting current value";
+	}
+
+	// Check that at least 1/30 of seconds elapsed since last update
+	double CurrentTime = GetHiResTime();
+	if ((CurrentTime - LastUpdate) < (1.0 / 30))
+	{
+		return;
+	}
+	// Take note of the last update time
+	LastUpdate = CurrentTime;
+
+	// Check if user want to cancel
+	CheckForCancel();
+}
+
+// Advance progression bar for the current value
+void FProgression::NewCurrentValueX(int InCurrentValue)
+{
+	GS::Guard< GS::Lock > lck(AccessControl);
+	NextCurrentValue = InCurrentValue;
+}
+
+// Call from the main thread (Usualy when Join wait)
+void FProgression::Update()
+{
+	GS::Guard< GS::Lock > lck(AccessControl);
+	if (NextPhase)
+	{
+		int NewValue = NextCurrentValue;
+		NewPhase(NextPhase, NextPhaseMaxValue);
+		if (NewValue)
+		{
+			NewCurrentValue(NewValue);
+		}
+	}
+	else if (NextCurrentValue)
+	{
+		NewCurrentValue(NextCurrentValue);
+	}
+	else
+	{
+		NewCurrentValue();
+	}
+}
+
+// Throw a kUserCancelled error if user push the cancel button
+void FProgression::CheckForCancel()
+{
+	if (CancelMode == kNoCancel) // Dont show cancel interface
+		return;
+	if (CancelMode == kSetFlags && bUserCancelled == nullptr) // Specify to set flag, but not flag to set ?
+		return;
+	if (CancelMode == kThrowOnCancel && bUserCancelled != nullptr && *bUserCancelled)
+	{
+		// Previous mode was kSetFlags and user has cancelled
+		*bUserCancelled = false;
+		throw UE_AC_Error("User cancelled operation", UE_AC_Error::kUserCancelled);
+	}
+
+	// Show cancel interface (when mouse is over progression window)
+	if (ACAPI_Interface(APIIo_IsProcessCanceledID, nullptr, nullptr))
+	{
+		// True if user cancelled
+		if (CancelMode == kSetFlags)
+		{
+			*bUserCancelled = true;
+		}
+		else
+		{
+			throw UE_AC_Error("User cancelled operation", UE_AC_Error::kUserCancelled);
+		}
+	}
+}
+
+// Throw an exception if user cancelled current operation
+void FProgression::ThrowIfUserCancelled() const
+{
+	if (bUserCancelled != nullptr && *bUserCancelled)
+	{
+		throw UE_AC_Error("User cancelled operation", UE_AC_Error::kUserCancelled);
+	}
+}
+
+// Return the localized string specified by the index
+GS::UniString FProgression::GetText(int InIndex)
+{
+	GS::UniString s;
+	try
+	{
+		s = GetUniString(ResID, InIndex);
+	}
+	catch (...)
+	{
+		s = "Invalide text resources";
+	}
+
+	return s;
+}
+
+// Return time value in sec with hi resolution
+double FProgression::GetHiResTime()
+{
+#ifdef WIN32
+	LARGE_INTEGER tick;
+	QueryPerformanceCounter(&tick);
+	return tick.QuadPart * InvFreq;
+#else
+	struct timeval CurrentTime;
+	gettimeofday(&CurrentTime, nullptr);
+	return CurrentTime.tv_sec + (CurrentTime.tv_usec * (1.0 / 1000000));
+#endif
+}
+
+END_NAMESPACE_UE_AC
diff --git a/Engine/Source/Programs/Enterprise/Datasmith/DatasmithARCHICADExporter/Private/Utils/Progression.h b/Engine/Source/Programs/Enterprise/Datasmith/DatasmithARCHICADExporter/Private/Utils/Progression.h
new file mode 100644
--- /dev/null
+++ b/Engine/Source/Programs/Enterprise/Datasmith/DatasmithARCHICADExporter/Private/Utils/Progression.h
@@ -0,0 +1,152 @@
+// Copyright Epic Games, Inc. All Rights Reserved.
+
+#pragma once
+
+#include "AddonTools.h"
+
+#include "Condition.hpp"
+#include "Lock.hpp"
+#include "../ResourcesIDs.h"
+
+BEGIN_NAMESPACE_UE_AC
+
+enum EPhaseStrId
+{
+	kCommonPhaseInvalid = 0,
+	kCommonProjectInfos,
+	kCommonCollectElements,
+	kCommonConvertElements,
+	kCommonCollectMetaDatas,
+	kCommonSetUpMaterials,
+	kCommonSetUpLights,
+	kCommonSetUpCameras,
+
+	kDebugSaveScene,
+
+	kSyncSnapshot,
+	kExportSaving,
+
+	kSyncTitle,
+	kExportTitle,
+
+	kNbPhases = kCommonSetUpCameras - kCommonPhaseInvalid + 1
+};
+
+// Class for progression window, make difficult to forget closing it.
+class FProgression
+{
+  public:
+	enum EMode
+	{
+		kNoCancel = 0,
+		kThrowOnCancel,
+		kSetFlags
+	};
+
+	class AutoCancelMode
+	{
+	  public:
+		AutoCancelMode(FProgression* IOProgression, EMode InMode)
+			: Progression(IOProgression)
+			, PreviousMode(IOProgression ? IOProgression->SetCancelMode(InMode) : kNoCancel)
+		{
+		}
+
+		~AutoCancelMode()
+		{
+			if (Progression)
+			{
+				Progression->SetCancelMode(PreviousMode);
+			}
+		}
+
+		FProgression* Progression;
+		EMode		  PreviousMode;
+	};
+
+	// Constructor (SetUp progression window)
+	FProgression(int InResID, EPhaseStrId InTitle, short InNbPhases, EMode InCancelMode,
+				 volatile bool* OutUserCancelled);
+
+	// Destructor (Automatically close the progression window)
+	~FProgression();
+
+	// Set new cancel mode
+	EMode SetCancelMode(EMode InNewCancelMode);
+
+	// Start the next phase
+	void NewPhase(EPhaseStrId InPhaseId, int InMaxValue = 0);
+
+	void NewPhaseX(EPhaseStrId InPhaseId, int InMaxValue = 0);
+
+	// Advance progression bar for the current value
+	void NewCurrentValue(int InCurrentValue = -1);
+
+	void NewCurrentValueX(int InCurrentValue = -1);
+
+	// Call from the main thread (Usualy when Join wait)
+	void Update();
+
+	// Test if user cancelled the operation
+	void CheckForCancel();
+
+	// Return true if user cancelled the current phase
+	bool IsUserCancelled() const { return bUserCancelled ? *bUserCancelled : false; }
+
+	// Throw an exception if user cancelled current operation
+	void ThrowIfUserCancelled() const;
+
+  private:
+	// Illegal to copy an FProgression
+	FProgression& operator=(const FProgression&) { return *this; }
+
+	// Return the string specifien by the index
+	GS::UniString GetText(int InIndex);
+
+	// Return time value in sec with hi resolution
+	double GetHiResTime();
+
+	// The string ressource id
+	int ResID;
+
+	// The current phase counter
+	EPhaseStrId CurrentPhase;
+
+	double Ratio100 = 100.0;
+
+	// The current item processed
+	int CurrentValue;
+
+	// The last item update time (Used to not refresh screen too often)
+	double LastUpdate;
+
+	// The inverse of the frequecy, used to compute
+	double InvFreq;
+
+	// True if we initalize without error
+	bool bProgressionShown;
+
+	// The error code value if we got one
+	GSErrCode ErrorCode;
+
+	// The task when we got an error message
+	const char* ErrorMgs;
+
+	// The cancel mode
+	EMode CancelMode;
+
+	// When user cancel, we set this flag to true
+	volatile bool* bUserCancelled;
+
+	EPhaseStrId NextPhase;
+	int			NextPhaseMaxValue;
+	int			NextCurrentValue;
+
+	// Control access on this object (for queue operations)
+	GS::Lock AccessControl;
+
+	// Condition variable
+	GS::Condition mCV;
+};
+
+END_NAMESPACE_UE_AC
diff --git a/Engine/Source/Programs/Enterprise/Datasmith/DatasmithARCHICADExporter/Private/Utils/ResourcesUtils.h b/Engine/Source/Programs/Enterprise/Datasmith/DatasmithARCHICADExporter/Private/Utils/ResourcesUtils.h
--- a/Engine/Source/Programs/Enterprise/Datasmith/DatasmithARCHICADExporter/Private/Utils/ResourcesUtils.h
+++ b/Engine/Source/Programs/Enterprise/Datasmith/DatasmithARCHICADExporter/Private/Utils/ResourcesUtils.h
@@ -33,11 +33,10 @@
 #define kStrListENames 331
 
 #define kAlertUnhandledError 350
-#define kAlertSaveFileError 351
-#define kAlertACDBError 352
-#define kAlertPlugInError 353
-#define kAlertNot3DViewError 354
-#define kAlertUserCancelledError 355
+#define kAlertACDBError 351
+#define kAlertPlugInError 352
+#define kAlertNot3DViewError 353
+#define kAlertUserCancelledError 354
 
 #define LocalizeResId(id) short(UE_AC::GetCurrentLanguage() * 1000 + id)
 
@@ -45,20 +44,4 @@
 #define UE_AC_STRINGIZE_A(arg) UE_AC_STRINGIZE_I arg
 #define UE_AC_STRINGIZE_I(text) #text
 
-#ifndef RICKAD
-	#define RICKAD 0
-#endif
-
-#if RIKCAD
-	#define AC_RK_LR(l, r) AC_RK_TXT(l RIKCAD r)
-	#define AC_RK_L(l) AC_RK_TXT(l RIKCAD)
-	#define AC_RK_R(r) AC_RK_TXT(RIKCAD r)
-#else
-	#define AC_RK_LR(l, r) AC_RK_TXT(l ARCHICAD r)
-	#define AC_RK_L(l) AC_RK_TXT(l ARCHICAD)
-	#define AC_RK_R(r) AC_RK_TXT(ARCHICAD r)
-#endif
-#define AC_RK_TXT(text) #text
-
-#define ADDON_SYNC_NAME AC_RK_LR(Datasmith, Exporter)
 #define ADDON_MENU_TITLE "Datasmith"
diff --git a/Engine/Source/Programs/Enterprise/Datasmith/DatasmithARCHICADExporter/Private/Utils/TaskMgr.cpp b/Engine/Source/Programs/Enterprise/Datasmith/DatasmithARCHICADExporter/Private/Utils/TaskMgr.cpp
new file mode 100644
--- /dev/null
+++ b/Engine/Source/Programs/Enterprise/Datasmith/DatasmithARCHICADExporter/Private/Utils/TaskMgr.cpp
@@ -0,0 +1,209 @@
+// Copyright Epic Games, Inc. All Rights Reserved.
+
+#include "TaskMgr.h"
+#include "CurrentOS.h"
+
+#include "Thread.hpp"
+#include "Guard.hpp"
+#include "System.hpp"
+
+#include <stdexcept>
+
+BEGIN_NAMESPACE_UE_AC
+
+static FTaskMgr* STaskMgr;
+
+FTaskMgr* FTaskMgr::GetMgr()
+{
+	if (STaskMgr == nullptr)
+	{
+		STaskMgr = new FTaskMgr();
+	}
+	return STaskMgr;
+}
+
+void FTaskMgr::DeleteMgr()
+{
+	if (STaskMgr)
+	{
+		FTaskMgr* Tmp = STaskMgr;
+		STaskMgr = nullptr;
+		delete Tmp;
+	}
+}
+
+// Class that implement FTaskMgr thread using GS::Thread
+class FTaskMgr::FThread
+{
+  public:
+	// Constructor
+	FThread(FTaskMgr* InTaskMgr, const GS::UniString& inThreadName)
+		: Task(new FTaskRunner(InTaskMgr))
+		, Thread(Task, inThreadName)
+	{
+		Thread.Start();
+	}
+
+	// Start the thread
+	void Start() { Thread.Start(); }
+
+	// Join the thread
+	void Join() { Thread.Join(); }
+
+  private:
+	// Implement a GS::Runnable that will run tasks
+	class FTaskRunner : public GS::Runnable
+	{
+	  public:
+		// Constructor
+		FTaskRunner(FTaskMgr* InTaskMgr)
+			: TaskMgr(InTaskMgr)
+		{
+		}
+
+		// Run task after task
+		void Run()
+		{
+#ifdef WIN32
+			SetThreadName(GS::Thread::GetCurrent().GetName().ToUtf8());
+#else
+			pthread_setname_np(GS::Thread::GetCurrent().GetName().ToUtf8());
+#endif
+			FTask* ATask = TaskMgr->GetTask();
+			while (ATask)
+			{
+				try
+				{
+					ATask->Run();
+					delete ATask;
+				}
+				catch (std::exception& e)
+				{
+					UE_AC_DebugF("FTaskMgr::FTaskRunner::Run - Catch std exception %s\n", e.what());
+				}
+				catch (GS::GSException& gs)
+				{
+					UE_AC_DebugF("FTaskMgr::FTaskRunner::Run - Catch gs exception %s\n", gs.GetMessage().ToUtf8());
+				}
+				catch (...)
+				{
+					UE_AC_DebugF("FTaskMgr::FTaskRunner::Run - Catch unknown exception\n");
+				}
+				ATask = TaskMgr->GetTask();
+			}
+		}
+
+	  private:
+		FTaskMgr* TaskMgr;
+	};
+
+	GS::RunnableTask Task;
+	GS::Thread		 Thread;
+};
+
+// Constructor
+FTaskMgr::FTaskMgr()
+	: bTerminate(false)
+	, bTreadingEnabled(true)
+	, NbTaskWaiting(0)
+	, CV(AccessControl)
+{
+	if (bTreadingEnabled)
+	{
+		// One thread by processor
+		unsigned nbProcessors = GS::System::GetNumberOfActiveProcessors();
+		if (nbProcessors == 0)
+		{
+			nbProcessors = 1;
+		}
+		Treads.resize(nbProcessors);
+		for (unsigned i = 0; i < nbProcessors; i++)
+		{
+			Treads[i].reset(new FThread(this, GS::UniString::Printf("AC2TM Task #%d", i)));
+		}
+	}
+}
+
+// Destructor
+FTaskMgr::~FTaskMgr()
+{
+	{
+		GS::Guard< GS::Lock > lck(AccessControl);
+		bTerminate = true;
+	}
+	CV.NotifyAll();
+
+	size_t nbThreads = Treads.size();
+	for (size_t i = 0; i < nbThreads; i++)
+	{
+		Treads[i]->Join();
+	}
+
+	// Delete all not processed task. This happen when Join has'n been called on the thread manager
+	if (!TaskQueue.empty())
+	{
+		FTask* Task = TaskQueue.front();
+		TaskQueue.pop();
+		delete Task;
+	}
+}
+
+// Add task
+void FTaskMgr::AddTask(FTask* InTask, ERunMode inMode)
+{
+	if (bTreadingEnabled && inMode == kSchedule)
+	{
+		GS::Guard< GS::Lock > lck(AccessControl);
+		if (bTerminate)
+		{
+			throw std::runtime_error("Adding task to a terminated FTaskMgr");
+		}
+		TaskQueue.push(InTask);
+		CV.Notify();
+	}
+	else
+	{
+		InTask->Run();
+		delete InTask;
+	}
+}
+
+// Wait until all task have been processed
+void FTaskMgr::Join(FProgression* inProgression)
+{
+	GS::Guard< GS::Lock > lck(AccessControl);
+	while (!TaskQueue.empty() || NbTaskWaiting != Treads.size())
+	{
+		CV.NotifyAll();
+		CV.Wait(10, true, nullptr);
+		if (inProgression)
+		{
+			inProgression->Update();
+		}
+	}
+#if 0
+	UE_AC_TraceF("FTaskMgr::Join %lu %u\n", (unsigned long)mTaskQueue.size(), mNbTaskWaiting);
+#endif
+}
+
+// Get the next task to be run
+FTaskMgr::FTask* FTaskMgr::GetTask()
+{
+	CV.NotifyAll();
+	GS::Guard< GS::Lock > lck(AccessControl);
+	NbTaskWaiting++;
+	FTask* Task = nullptr;
+	while (Task == nullptr && !bTerminate)
+	{
+		CV.Wait(100, true, nullptr);
+		if (!TaskQueue.empty())
+		{
+			Task = TaskQueue.front();
+			TaskQueue.pop();
+			NbTaskWaiting--;
+		}
+	}
+	return Task;
+}
+
+END_NAMESPACE_UE_AC
diff --git a/Engine/Source/Programs/Enterprise/Datasmith/DatasmithARCHICADExporter/Private/Utils/TaskMgr.h b/Engine/Source/Programs/Enterprise/Datasmith/DatasmithARCHICADExporter/Private/Utils/TaskMgr.h
new file mode 100644
--- /dev/null
+++ b/Engine/Source/Programs/Enterprise/Datasmith/DatasmithARCHICADExporter/Private/Utils/TaskMgr.h
@@ -0,0 +1,87 @@
+// Copyright Epic Games, Inc. All Rights Reserved.
+
+#pragma once
+
+#include "AddonTools.h"
+
+#include "Progression.h"
+
+#include "Condition.hpp"
+#include "Lock.hpp"
+
+#include <queue>
+#include <memory>
+
+BEGIN_NAMESPACE_UE_AC
+
+class FTaskMgr
+{
+  public:
+	// Task class, your class must inherit from this class.
+	class FTask
+	{
+	  public:
+		virtual void Run() = 0;
+		virtual ~FTask() {}
+	};
+
+	// Constructor
+	FTaskMgr();
+
+	// Destructor
+	~FTaskMgr();
+
+	// Task can be added for schedule or run immediatly in the current thread).
+	enum ERunMode
+	{
+		kSchedule,
+		kRunNow
+	};
+
+	// Add task (if inRunNow, the task is run immediatly in current thread)
+	void AddTask(FTask* InTask, ERunMode InMode = kSchedule);
+
+	// Wait until all task have been processed
+	void Join(FProgression* InProgression = nullptr);
+
+	bool EnableThreading(bool InEnable)
+	{
+		bool tmp = bTreadingEnabled;
+		bTreadingEnabled = InEnable;
+		return tmp;
+	}
+
+	static FTaskMgr* GetMgr();
+
+	static void DeleteMgr();
+
+  private:
+	// Get the next task to be run
+	FTask* GetTask();
+
+	// Class that implement CTaskMgr thread
+	class FThread;
+
+	// Threads used (≈ nomber of cpu)
+	std::vector< std::unique_ptr< FThread > > Treads;
+
+	// Fifo tasks queue
+	std::queue< FTask* > TaskQueue;
+
+	// Control access on this object (for queue operations)
+	GS::Lock AccessControl;
+
+	// Condition variable
+	GS::Condition CV;
+
+	// Flag to signal that we are deleting this task mgr
+	volatile bool bTerminate;
+
+	// Number of thread waiting
+	volatile int NbTaskWaiting;
+
+	// False: tasks to run in main threads, true: Task run in thread
+	bool bTreadingEnabled;
+};
+
+END_NAMESPACE_UE_AC
diff --git a/Engine/Source/Programs/Enterprise/Datasmith/DatasmithARCHICADExporter/Private/Utils/WarningsDisabler.h b/Engine/Source/Programs/Enterprise/Datasmith/DatasmithARCHICADExporter/Private/Utils/WarningsDisabler.h
--- a/Engine/Source/Programs/Enterprise/Datasmith/DatasmithARCHICADExporter/Private/Utils/WarningsDisabler.h
+++ b/Engine/Source/Programs/Enterprise/Datasmith/DatasmithARCHICADExporter/Private/Utils/WarningsDisabler.h
@@ -33,7 +33,8 @@
     _Pragma( "clang diagnostic ignored \"-Wdocumentation\"" ) \
 	_Pragma( "clang diagnostic ignored \"-Wshorten-64-to-32\"" ) \
 	_Pragma( "clang diagnostic ignored \"-Wdefaulted-function-deleted\"" ) \
-	_Pragma( "clang diagnostic ignored \"-Wcomma\"" )
+	_Pragma( "clang diagnostic ignored \"-Wcomma\"" ) \
+	_Pragma( "clang diagnostic ignored \"-Wnon-c-typedef-for-linkage\"" )
 
 #define DISABLE_SDK_WARNINGS_END \
     _Pragma( "clang diagnostic pop" )
diff --git a/Engine/Source/Programs/Enterprise/Datasmith/DatasmithARCHICADExporter/Private/Utils/WinOS.cpp b/Engine/Source/Programs/Enterprise/Datasmith/DatasmithARCHICADExporter/Private/Utils/WinOS.cpp
--- a/Engine/Source/Programs/Enterprise/Datasmith/DatasmithARCHICADExporter/Private/Utils/WinOS.cpp
+++ b/Engine/Source/Programs/Enterprise/Datasmith/DatasmithARCHICADExporter/Private/Utils/WinOS.cpp
@@ -99,4 +99,38 @@ GS::UniString GetHomeDirectory()
 	return GS::UniString((GS::UniChar::Layout*)Path);
 }
 
+void SetThreadName(const char* InName)
+{
+#if defined(DEBUG) || 1
+	const DWORD MS_VC_EXCEPTION = 0x406D1388;
+
+	#pragma pack(push, 8)
+	typedef struct tagTHREADNAME_INFO
+	{
+		DWORD  dwType; // Must be 0x1000.
+		LPCSTR szName; // Pointer to name (in user addr space).
+		DWORD  dwThreadID; // Thread ID (-1=caller thread).
+		DWORD  dwFlags; // Reserved for future use, must be zero.
+	} THREADNAME_INFO;
+	#pragma pack(pop)
+
+	THREADNAME_INFO info;
+	info.dwType = 0x1000;
+	info.szName = InName;
+	info.dwThreadID = (DWORD)-1;
+	info.dwFlags = 0;
+
+	// This will work only when we are debugged. The __try/__except block is there for when there are no debugger.
+	__try
+	{
+		RaiseException(MS_VC_EXCEPTION, 0, sizeof(info) / sizeof(ULONG_PTR), (ULONG_PTR*)&info);
+	}
+	__except (EXCEPTION_EXECUTE_HANDLER)
+	{
+	}
+#else
+	(void)InName;
+#endif
+}
+
 END_NAMESPACE_UE_AC
diff --git a/Engine/Source/Programs/Enterprise/Datasmith/DatasmithARCHICADExporter/Resources/RFIX.WIN/DatasmithARCHICADExporter.rc2 b/Engine/Source/Programs/Enterprise/Datasmith/DatasmithARCHICADExporter/Resources/RFIX.WIN/DatasmithARCHICADExporter.rc2
--- a/Engine/Source/Programs/Enterprise/Datasmith/DatasmithARCHICADExporter/Resources/RFIX.WIN/DatasmithARCHICADExporter.rc2
+++ b/Engine/Source/Programs/Enterprise/Datasmith/DatasmithARCHICADExporter/Resources/RFIX.WIN/DatasmithARCHICADExporter.rc2
@@ -8,7 +8,6 @@
 // [SG compatible] - Yes
 // ****************************************************************************/
 
-#include "DatasmithARCHICADExporter.grc.rc2"
 #include "DatasmithARCHICADExporterFix.grc.rc2"
 
 1	ICON    LOADONCALL MOVEABLE IMPURE  ACAP.ico
diff --git a/Engine/Source/Programs/Enterprise/Datasmith/DatasmithARCHICADExporter/Resources/RFIX.mac/compileGRCs.pl b/Engine/Source/Programs/Enterprise/Datasmith/DatasmithARCHICADExporter/Resources/RFIX.mac/compileGRCs.pl
--- a/Engine/Source/Programs/Enterprise/Datasmith/DatasmithARCHICADExporter/Resources/RFIX.mac/compileGRCs.pl
+++ b/Engine/Source/Programs/Enterprise/Datasmith/DatasmithARCHICADExporter/Resources/RFIX.mac/compileGRCs.pl
@@ -9,15 +9,14 @@ use feature 'unicode_strings';
 #-----------------------------------------------------------------------
 # Executes one GRC conversion command
 #-----------------------------------------------------------------------
-sub DoGRC ($$$$$$)
+sub DoGRC ($$$$$$$)
 {
-	my ($tool, $headerDir, $sourceDir, $inputFile, $destinationDir, $outputFile) = @_ ;
+	my ($tool, $headerDir, $sourceDir, $intSourceDir, $inputFile, $destinationDir, $outputFile) = @_ ;
 
 	system ("mkdir -p \"" . $destinationDir . "\"") ;
-	print "\t$inputFile\n";
 	system ("vim -e -s +\"set bomb|set encoding=utf-8|wq\" \"" . catfile ($sourceDir, $inputFile) . "\"");
 
-	system ("xcrun clang -x c++ -E -P " . $ENV{GRC_DEFINES} . " -I \"" . $headerDir . "\" \"" . catfile ($sourceDir, $inputFile) . "\" > \"" . catfile ($destinationDir, $outputFile .".i") . "\"");
+	system ("xcrun clang -x c++ -E -P " . $ENV{GRC_DEFINES} . " -I \"" . $headerDir . "\" -I \"" . $intSourceDir . "\" \"" . catfile ($sourceDir, $inputFile) . "\" > \"" . catfile ($destinationDir, $outputFile .".i") . "\"");
 	system ($tool . " -m r -q utf8 utf16 -T M -n -i \"" . catfile ($destinationDir, $outputFile .".i") . "\" -p \"" . catfile ($sourceDir, "Images") . "\" -o \"" . catfile ($destinationDir, $outputFile) . "\" -w 2");
 }
 
@@ -45,17 +44,8 @@ opendir (DIR, $rfixFolder);
 closedir (DIR);
 
 foreach $path (@files) {
-	($volume,$directories,$file) = splitpath ($path);
-	DoGRC ($makeResConvTool, $srcFolder, $rfixFolder, $file, $roFolder, $file . ".ro") if ($file =~ /.*\.grc$/i);
-}
-
-opendir (DIR, $rintFolder);
-@files = readdir (DIR);
-closedir (DIR);
-
-foreach $path (@files) {
-	($volume,$directories,$file) = splitpath ($path);
-	DoGRC ($makeResConvTool, $srcFolder, $rintFolder, $file, $roFolder, $file . ".ro") if ($file =~ /.*\.grc$/i);
+	($_,$_,$file) = splitpath ($path);
+	DoGRC ($makeResConvTool, $srcFolder, $rfixFolder, $rintFolder, $file, $roFolder, $file . ".ro") if ($file =~ /.*\.grc$/i);
 }
 
 opendir (DIR, $roFolder);
