Hide Auto Sync feature

diff --git a/Engine/Source/Programs/Enterprise/Datasmith/DatasmithARCHICADExporter/Build/Config/UE_AC.xcconfig b/Engine/Source/Programs/Enterprise/Datasmith/DatasmithARCHICADExporter/Build/Config/UE_AC.xcconfig
--- a/Engine/Source/Programs/Enterprise/Datasmith/DatasmithARCHICADExporter/Build/Config/UE_AC.xcconfig
+++ b/Engine/Source/Programs/Enterprise/Datasmith/DatasmithARCHICADExporter/Build/Config/UE_AC.xcconfig
@@ -14,7 +14,7 @@ UE_INCLUDES_PATH = "$(UE_Core)" "$(UE_Core)/Internationalization" "$(UE_Core)/As
 
 USER_HEADER_SEARCH_PATHS = "." $(AC_INCLUDES_PATH) $(UE_INCLUDES_PATH)
 
-AUTO_SYNC = 1
+AUTO_SYNC = 0
 
 VMajor = 1
 VMinor = 0
