Datasmith Revit: Structural Steel Connection Pivots are default

diff --git a/Engine/Source/Programs/Enterprise/Datasmith/DatasmithRevitExporter/Private/DatasmithRevitDocumentData.cs b/Engine/Source/Programs/Enterprise/Datasmith/DatasmithRevitExporter/Private/DatasmithRevitDocumentData.cs
--- a/Engine/Source/Programs/Enterprise/Datasmith/DatasmithRevitExporter/Private/DatasmithRevitDocumentData.cs
+++ b/Engine/Source/Programs/Enterprise/Datasmith/DatasmithRevitExporter/Private/DatasmithRevitDocumentData.cs
@@ -119,6 +119,8 @@ namespace DatasmithRevitExporter
 
 			public Dictionary<string, int>		MeshMaterialsMap = new Dictionary<string, int>();
 
+			public Transform WorldTransform;
+
 			public List<FBaseElementData>	ChildElements = new List<FBaseElementData>();
 
 			public FBaseElementData			Parent = null;
@@ -763,6 +765,8 @@ namespace DatasmithRevitExporter
 					FBaseElementData InElement
 			)
 			{
+				InElement.WorldTransform = InWorldTransform;
+
 				// Create a new Datasmith mesh.
 				// Hash the Datasmith mesh name to shorten it.
 				string HashedMeshName = FDatasmithFacadeElement.GetStringHash("M:" + GetMeshName());
@@ -1297,6 +1301,31 @@ namespace DatasmithRevitExporter
 			{
 				InstanceData.UpdateMeshName();
 			}
+			else
+			{
+				/* Instance has no mesh.
+				 * Handle the case where instance has valid transform, but parent element has valid mesh (exported after instance gets finished),
+				 * in which case we want to apply the instance transform as a pivot transform.
+				 * This is a case currently encountered for steel beams.
+				 */
+				bool bElementHasMesh = ContainsMesh(CurrentElement.DatasmithMeshElement.GetName()) || (CurrentElement.DatasmithPolymesh.Vertices.Count > 0 && CurrentElement.DatasmithPolymesh.Faces.Count > 0);
+
+				if (CurrentElement.CurrentElement.GetType() == typeof(FamilyInstance) && !bElementHasMesh)
+				{
+					if (!CurrentElement.WorldTransform.IsIdentity)
+					{
+						CurrentElement.WorldTransform = CurrentElement.WorldTransform * InstanceData.WorldTransform;
+					} 
+					else
+					{
+						CurrentElement.WorldTransform = InstanceData.WorldTransform;
+					}
+
+					CurrentElement.MeshPointsTransform = InstanceData.WorldTransform.Inverse;
+
+					SetActorTransform(CurrentElement.WorldTransform, CurrentElement.ElementActor);
+				}
+			}
 
 			// Collect the element Datasmith mesh into the mesh dictionary.
 			CollectMesh(InstanceData.DatasmithPolymesh, InstanceData.DatasmithMeshElement, InDatasmithScene);
