Sketchup: fixed exporting png alpha channel not connected to opacity

diff --git a/Engine/Source/Programs/Enterprise/Datasmith/DatasmithSketchUpRubyExporter/Private/DatasmithSketchUpMaterial.cpp b/Engine/Source/Programs/Enterprise/Datasmith/DatasmithSketchUpRubyExporter/Private/DatasmithSketchUpMaterial.cpp
--- a/Engine/Source/Programs/Enterprise/Datasmith/DatasmithSketchUpRubyExporter/Private/DatasmithSketchUpMaterial.cpp
+++ b/Engine/Source/Programs/Enterprise/Datasmith/DatasmithSketchUpRubyExporter/Private/DatasmithSketchUpMaterial.cpp
@@ -149,11 +149,6 @@ namespace DatasmithSketchUp
 
 		DatasmithMaterialElementPtr->SetTwoSided(false);// todo: consider this
 
-		IDatasmithMaterialExpressionColor* ExpressionColor = DatasmithMaterialElementPtr->AddMaterialExpression<IDatasmithMaterialExpressionColor>();
-		ExpressionColor->SetName(TEXT("Base Color"));
-		ExpressionColor->GetColor() = LinearColor;
-		ExpressionColor->ConnectExpression(DatasmithMaterialElementPtr->GetBaseColor());
-
 		bool bTranslucent = InMaterial.bSourceColorAlphaUsed;
 		if (SUIsValid(InMaterial.TextureRef))
 		{
@@ -193,20 +188,41 @@ namespace DatasmithSketchUp
 			ExpressionTexture->ConnectExpression(DatasmithMaterialElementPtr->GetBaseColor());
 
 			bTranslucent = bTranslucent || Texture->GetTextureUseAlphaChannel();
-		}
 
-		if (bTranslucent)
+			// Set the Datasmith material element opacity.
+			if (Texture->GetTextureUseAlphaChannel())
+			{
+				// Invert texture translarency to get Unreal opacity
+				IDatasmithMaterialExpressionGeneric* ExpressionOpacity = DatasmithMaterialElementPtr->AddMaterialExpression<IDatasmithMaterialExpressionGeneric>();
+				ExpressionOpacity->SetExpressionName(TEXT("OneMinus"));
+
+
+				ExpressionTexture->ConnectExpression(*ExpressionOpacity->GetInput(0), 3);
+
+				ExpressionOpacity->ConnectExpression(DatasmithMaterialElementPtr->GetOpacity());
+			}
+		}
+		else
 		{
-			DatasmithMaterialElementPtr->SetBlendMode(/*EBlendMode::BLEND_Translucent*/2);
+			IDatasmithMaterialExpressionColor* ExpressionColor = DatasmithMaterialElementPtr->AddMaterialExpression<IDatasmithMaterialExpressionColor>();
+			ExpressionColor->SetName(TEXT("Base Color"));
+			ExpressionColor->GetColor() = LinearColor;
+			ExpressionColor->ConnectExpression(DatasmithMaterialElementPtr->GetBaseColor());
+
+			// Set the Datasmith material element opacity.
+			if (InMaterial.bSourceColorAlphaUsed)
+			{
+				IDatasmithMaterialExpressionScalar* ExpressionOpacity = DatasmithMaterialElementPtr->AddMaterialExpression<IDatasmithMaterialExpressionScalar>();
+				ExpressionOpacity->SetName(TEXT("Opacity"));
+				ExpressionOpacity->GetScalar() = float(InMaterial.SourceColor.alpha) / float(255);
+				ExpressionOpacity->ConnectExpression(DatasmithMaterialElementPtr->GetOpacity());
+			}
 		}
 
-		// Set the Datasmith material element opacity.
-		if (InMaterial.bSourceColorAlphaUsed)
+
+		if (bTranslucent)
 		{
-			IDatasmithMaterialExpressionScalar* ExpressionOpacity = DatasmithMaterialElementPtr->AddMaterialExpression<IDatasmithMaterialExpressionScalar>();
-			ExpressionOpacity->SetName(TEXT("Opacity"));
-			ExpressionOpacity->GetScalar() = float(InMaterial.SourceColor.alpha) / float(255);
-			ExpressionOpacity->ConnectExpression(DatasmithMaterialElementPtr->GetOpacity());
+			DatasmithMaterialElementPtr->SetBlendMode(/*EBlendMode::BLEND_Translucent*/2);
 		}
 
 		Context.DatasmithScene->AddMaterial(DatasmithMaterialElementPtr);
