Align DatasmithRhinoExporter UI Terminology

diff --git a/Engine/Source/Programs/Enterprise/Datasmith/DatasmithRhinoExporter/Private/Commands/DatasmithExportCommand.cs b/Engine/Source/Programs/Enterprise/Datasmith/DatasmithRhinoExporter/Private/Commands/DatasmithExportCommand.cs
--- a/Engine/Source/Programs/Enterprise/Datasmith/DatasmithRhinoExporter/Private/Commands/DatasmithExportCommand.cs
+++ b/Engine/Source/Programs/Enterprise/Datasmith/DatasmithRhinoExporter/Private/Commands/DatasmithExportCommand.cs
@@ -45,7 +45,7 @@ namespace DatasmithRhino.Commands
 			Eto.Forms.FileFilter DatasmithFileFilter = new Eto.Forms.FileFilter("Unreal Datasmith", new string[] { ".udatasmith" });
 			SaveDialog.Filters.Add(DatasmithFileFilter);
 			SaveDialog.CurrentFilter = DatasmithFileFilter;
-			SaveDialog.Title = "Export to Datasmith Scene";
+			SaveDialog.Title = "Export to Datasmith file";
 			SaveDialog.FileName = string.IsNullOrEmpty(RhinoDocument.Name) ? "Untitled" : System.IO.Path.GetFileNameWithoutExtension(RhinoDocument.Name);
 			if (!string.IsNullOrEmpty(RhinoDocument.Path))
 			{
