Sketchup: Observers optimization and Layers visibility tracking

- Observers optimization. Only current editing context is watched for Entities(instances, groups, faces) modificaitons. This removes about 9 seconds from Stadium export.

- Layers visibility tracking. Toggling layer visibility and reassinging layers on instances, groups and faces is reflected on the updated datasmith scene: when layer made invisible faces are excluded from generated datasmith meshes and invisible instances remove datasmith actor hierarchy, toggling layer visible brings corresponding items. Also layer inheritance(overriding in containing instances) works correctly.

diff --git a/Engine/Source/Programs/Enterprise/Datasmith/DatasmithSketchUpRubyExporter/Private/DatasmithSketchUpComponent.cpp b/Engine/Source/Programs/Enterprise/Datasmith/DatasmithSketchUpRubyExporter/Private/DatasmithSketchUpComponent.cpp
--- a/Engine/Source/Programs/Enterprise/Datasmith/DatasmithSketchUpRubyExporter/Private/DatasmithSketchUpComponent.cpp
+++ b/Engine/Source/Programs/Enterprise/Datasmith/DatasmithSketchUpRubyExporter/Private/DatasmithSketchUpComponent.cpp
@@ -32,9 +32,6 @@ void FNodeOccurence::AddChildOccurrence(FExportContext& Context, FComponentInsta
 {
 	FNodeOccurence& ChildNode = ChildComponentInstance.CreateNodeOccurrence(Context, /*Parent*/ *this);
 
-	SULayerRef ChildEffectiveLayerRef = DatasmithSketchUpUtils::GetEffectiveLayer(ChildComponentInstance.GetComponentInstanceRef(), EffectiveLayerRef);
-	ChildNode.EffectiveLayerRef = ChildEffectiveLayerRef;
-
 	// todo: rename ToDatasmith - this just creates hierarchy
 	ChildNode.ToDatasmith(Context);
 }
@@ -75,15 +72,6 @@ void FNodeOccurence::ToDatasmith(FExportContext& Context)
 
 void FNodeOccurence::UpdateMeshActors(FExportContext& Context)
 {
-	FDefinition* EntityDefinition = Entity.GetDefinition();
-
-	DatasmithSketchUp::FEntitiesGeometry* EntitiesGeometry = EntityDefinition->GetEntities().EntitiesGeometry.Get();
-
-	if (!EntitiesGeometry)
-	{
-		return;
-	}
-
 	// Remove old mesh actors
 	// todo: reuse old mesh actors (also can keep instances when removing due to say hidden)
 	if (DatasmithActorElement.IsValid())
@@ -100,15 +88,22 @@ void FNodeOccurence::UpdateMeshActors(FExportContext& Context)
 			Context.DatasmithScene->RemoveActor(MeshActor, EDatasmithActorRemovalRule::RemoveChildren);
 		}
 	}
-	MeshActors.Reset(EntitiesGeometry->GetMeshCount());
+	MeshActors.Reset();
 
-#ifdef REMOVE_MESHES_WHEN_INVISIBLE
-	if (!bVisible) // remove mesh actors for invisible node
+	if (!bVisible) // no mesh actors for invisible node
+	{
+		return;
+	}
+
+	FDefinition* EntityDefinition = Entity.GetDefinition();
+	DatasmithSketchUp::FEntitiesGeometry* EntitiesGeometry = EntityDefinition->GetEntities().EntitiesGeometry.Get();
+
+	if (!EntitiesGeometry)
 	{
-		
 		return;
 	}
-#endif 
+
+	MeshActors.Reset(EntitiesGeometry->GetMeshCount());
 
 	FString ComponentActorName = GetActorName();
 	FString MeshActorLabel = GetActorLabel();
@@ -297,6 +292,42 @@ void FNodeOccurence::SetVisibility(bool bValue)
 	Entity.EntityOccurrenceVisible(this, bVisible);
 }
 
+void FNodeOccurence::RemoveDatasmithActorHierarchy(FExportContext& Context)
+{
+	if (!DatasmithActorElement)
+	{
+		return;
+	}
+
+	// Remove depth-first
+	for (FNodeOccurence* ChildNode : Children)
+	{
+		ChildNode->RemoveDatasmithActorHierarchy(Context);
+	}
+
+	for (TSharedPtr<IDatasmithMeshActorElement> MeshActor : MeshActors)
+	{
+		DatasmithActorElement->RemoveChild(MeshActor);
+	}
+	MeshActors.Reset();
+
+	if (const TSharedPtr<IDatasmithActorElement>& ParentActor = DatasmithActorElement->GetParentActor())
+	{
+		ParentActor->RemoveChild(DatasmithActorElement);
+	}
+	else
+	{
+		Context.DatasmithScene->RemoveActor(DatasmithActorElement, EDatasmithActorRemovalRule::RemoveChildren);
+	}
+	DatasmithActorElement.Reset();
+
+	if (DatasmithMetadataElement)
+	{
+		Context.DatasmithScene->RemoveMetaData(DatasmithMetadataElement);
+	}
+	DatasmithMetadataElement.Reset();
+}
+
 FModelDefinition::FModelDefinition(SUModelRef InModel) : Model(InModel)
 {
 }
@@ -533,6 +564,8 @@ void FComponentDefinition::RemoveComponentDefinition(FExportContext& Context)
 			DatasmithSketchUpUtils::GetComponentID(ComponentDefinitionRef),
 			DatasmithSketchUpUtils::GetGroupID(GroupRef));
 	}
+
+	Context.EntitiesObjects.UnregisterEntities(GetEntities());
 }
 
 void FComponentDefinition::AddInstance(FExportContext& Context, TSharedPtr<FComponentInstance> Instance)
@@ -578,8 +611,6 @@ void FEntity::UpdateOccurrence(FExportContext& Context, FNodeOccurence& Node)
 			}
 		}
 	}
-
-	
 }
 
 void FEntity::UpdateEntityGeometry(FExportContext& Context)
@@ -735,10 +766,23 @@ void FComponentInstance::InvalidateOccurrencesGeometry(FExportContext& Context)
 
 void FComponentInstance::InvalidateOccurrencesProperties(FExportContext& Context)
 {
-	bool bNewHidden = RetrieveHidden();
-	if (bHidden != bNewHidden) // ComponentInstance 'hidden' flag changed
+	// When ComponentInstance is modified we need to determine if its visibility might have changed foremost
+	// because this determines whether corresponding node would exist in the Datasmith scene 
+	// Two things affect this - Hidden instance flag and layer(tag):
+
+	bool bNewHidden = false;
+	SUDrawingElementRef DrawingElementRef = SUComponentInstanceToDrawingElement(GetComponentInstanceRef());
+	SUDrawingElementGetHidden(DrawingElementRef, &bNewHidden);
+
+	SUDrawingElementGetLayer(DrawingElementRef, &LayerRef);
+	bool bNewLayerVisible = true;
+	SULayerGetVisibility(LayerRef, &bNewLayerVisible);
+
+
+	if (bHidden != bNewHidden || bLayerVisible != bNewLayerVisible)
 	{
 		bHidden = bNewHidden;
+		bLayerVisible = bNewLayerVisible;
 		for (FNodeOccurence* Node : Occurrences)
 		{
 			Node->bVisibilityInvalidated = true;
@@ -770,14 +814,6 @@ SUComponentInstanceRef FComponentInstance::GetComponentInstanceRef()
 	return SUComponentInstanceFromEntity(EntityRef);
 }
 
-bool FComponentInstance::RetrieveHidden()
-{
-	bool bResult = false;
-	SUResult bSuResult  = SUDrawingElementGetHidden(SUComponentInstanceToDrawingElement(GetComponentInstanceRef()), &bResult);
-	ensure(bSuResult == SU_ERROR_NONE);
-	return bResult;
-}
-
 void FComponentInstance::FillOccurrenceActorMetadata(FNodeOccurence& Node)
 {
 	if (!Node.DatasmithMetadataElement)
@@ -810,13 +846,13 @@ void FComponentInstance::FillOccurrenceActorMetadata(FNodeOccurence& Node)
 
 void FComponentInstance::UpdateOccurrenceVisibility(FExportContext& Context, FNodeOccurence& Node)
 {
-	// Get the effective layer of the SketckUp normal component instance.
-	SULayerRef ChildEffectiveLayerRef = DatasmithSketchUpUtils::GetEffectiveLayer(GetComponentInstanceRef(), Node.ParentNode->EffectiveLayerRef);
-	bool bChildEffectiveLayerVisible = true;
-	SULayerGetVisibility(ChildEffectiveLayerRef, &bChildEffectiveLayerVisible);
+	Node.EffectiveLayerRef = DatasmithSketchUpUtils::GetEffectiveLayer(GetComponentInstanceRef(), Node.ParentNode->EffectiveLayerRef);
+
+	bool bEffectiveLayerVisible = true;
+	SULayerGetVisibility(Node.EffectiveLayerRef, &bEffectiveLayerVisible);
 
 	// Parent node, component instance and layer - all should be visible to have node visible
-	Node.SetVisibility(Node.ParentNode->bVisible && !bHidden && bChildEffectiveLayerVisible);
+	Node.SetVisibility(Node.ParentNode->bVisible && !bHidden && bEffectiveLayerVisible);
 
 	if (Node.bVisible)
 	{
@@ -829,22 +865,11 @@ void FComponentInstance::UpdateOccurrenceVisibility(FExportContext& Context, FNo
 			Node.InvalidateMeshActors();
 			Node.InvalidateProperties();
 		}
-		Node.DatasmithActorElement->SetVisibility(true);
-#ifdef REMOVE_MESHES_WHEN_INVISIBLE
 		Node.InvalidateMeshActors();
-#endif
-
 	}
 	else
 	{
-		if (Node.DatasmithActorElement)
-		{
-			Node.DatasmithActorElement->SetVisibility(false);
-
-#ifdef REMOVE_MESHES_WHEN_INVISIBLE
-			Node.InvalidateMeshActors();
-#endif
-		}
+		Node.RemoveDatasmithActorHierarchy(Context);
 	}
 
 	for (FNodeOccurence* ChildNode : Node.Children)
diff --git a/Engine/Source/Programs/Enterprise/Datasmith/DatasmithSketchUpRubyExporter/Private/DatasmithSketchUpComponent.h b/Engine/Source/Programs/Enterprise/Datasmith/DatasmithSketchUpRubyExporter/Private/DatasmithSketchUpComponent.h
--- a/Engine/Source/Programs/Enterprise/Datasmith/DatasmithSketchUpRubyExporter/Private/DatasmithSketchUpComponent.h
+++ b/Engine/Source/Programs/Enterprise/Datasmith/DatasmithSketchUpRubyExporter/Private/DatasmithSketchUpComponent.h
@@ -65,6 +65,8 @@ namespace DatasmithSketchUp
 		void InvalidateMeshActors();
 		void SetVisibility(bool);
 
+		void RemoveDatasmithActorHierarchy(FExportContext& Context);
+
 		void ToDatasmith(FExportContext& Context); // Build actor hierarchy 
 
 		FString GetActorName();
@@ -269,6 +271,8 @@ namespace DatasmithSketchUp
 		bool IsMeshUsingInheritedMaterial(int32 MeshIndex);
 
 		TArray<TSharedPtr<FDatasmithInstantiatedMesh>> Meshes;
+		TSet<int32> FaceIds; // EntityId of all the VISIBLE faces composing the mesh
+		TSet<DatasmithSketchUp::FEntityIDType> Layers; // EntityId of all layers assigned to geometry faces(needed to identify if geometry needs to be rebuilt when layer visibility changes)
 		TSet<FMaterialIDType> MaterialsUsed;
 
 		// todo: update reusing datasmith elements? 
@@ -365,11 +369,11 @@ namespace DatasmithSketchUp
 		FComponentInstanceIDType GetComponentInstanceId();
 		SUComponentInstanceRef GetComponentInstanceRef(); 
 
-		bool RetrieveHidden(); // Retrieve SketchUp value of Hidden flag
-
 		void FillOccurrenceActorMetadata(FNodeOccurence& Node);
 
-		bool bHidden = false; // S
+		bool bHidden = false;
+		SULayerRef LayerRef = SU_INVALID;
+		bool bLayerVisible = true;
 
 		TArray<FNodeOccurence*> Occurrences;
 
diff --git a/Engine/Source/Programs/Enterprise/Datasmith/DatasmithSketchUpRubyExporter/Private/DatasmithSketchUpExportContext.cpp b/Engine/Source/Programs/Enterprise/Datasmith/DatasmithSketchUpRubyExporter/Private/DatasmithSketchUpExportContext.cpp
--- a/Engine/Source/Programs/Enterprise/Datasmith/DatasmithSketchUpRubyExporter/Private/DatasmithSketchUpExportContext.cpp
+++ b/Engine/Source/Programs/Enterprise/Datasmith/DatasmithSketchUpRubyExporter/Private/DatasmithSketchUpExportContext.cpp
@@ -284,12 +284,38 @@ TSharedPtr<FComponentDefinition>* FComponentDefinitionCollection::FindComponentD
 	return ComponentDefinitionMap.Find(ComponentDefinitionID);
 }
 
-void FEntitiesObjectCollection::RegisterEntitiesFaces(DatasmithSketchUp::FEntities& Entities, const TSet<int32>& FaceIds)
+void FEntitiesObjectCollection::RegisterEntities(DatasmithSketchUp::FEntities& Entities)
 {
-	for (int32 FaceId : FaceIds)
+	for (int32 FaceId : Entities.EntitiesGeometry->FaceIds)
 	{
 		FaceIdForEntitiesMap.Add(FaceId, &Entities);
 	}
+
+	for (DatasmithSketchUp::FEntityIDType LayerId : Entities.EntitiesGeometry->Layers)
+	{
+		LayerIdForEntitiesMap.FindOrAdd(LayerId).Add(&Entities);
+	}
+}
+
+void FEntitiesObjectCollection::UnregisterEntities(DatasmithSketchUp::FEntities& Entities)
+{
+	if (!Entities.EntitiesGeometry)
+	{
+		return;
+	}
+
+	for (int32 FaceId : Entities.EntitiesGeometry->FaceIds)
+	{
+		FaceIdForEntitiesMap.Remove(FaceId);
+	}
+
+	for (DatasmithSketchUp::FEntityIDType LayerId : Entities.EntitiesGeometry->Layers)
+	{
+		if(TSet<DatasmithSketchUp::FEntities*>* Ptr = LayerIdForEntitiesMap.Find(LayerId))
+		{
+			Ptr->Remove(&Entities);
+		}
+	}
 }
 
 TSharedPtr<FEntities> FEntitiesObjectCollection::AddEntities(FDefinition& InDefinition, SUEntitiesRef EntitiesRef)
@@ -309,6 +335,16 @@ DatasmithSketchUp::FEntities* FEntitiesObjectCollection::FindFace(int32 FaceId)
 	return nullptr;
 }
 
+void FEntitiesObjectCollection::LayerModified(FEntityIDType LayerId)
+{
+	if (TSet<DatasmithSketchUp::FEntities*>* Ptr = LayerIdForEntitiesMap.Find(LayerId))
+	{
+		for (DatasmithSketchUp::FEntities* Entities : *Ptr)
+		{
+			Entities->Definition.InvalidateDefinitionGeometry();
+		}
+	}
+}
 
 TSharedPtr<FComponentInstance> FComponentInstanceCollection::AddComponentInstance(FDefinition& ParentDefinition, SUComponentInstanceRef InComponentInstanceRef)
 {
@@ -421,6 +457,18 @@ void FComponentInstanceCollection::UpdateGeometry()
 	}
 }
 
+void FComponentInstanceCollection::LayerModified(DatasmithSketchUp::FEntityIDType LayerId)
+{
+	for (const auto& KeyValue : ComponentInstanceMap)
+	{
+		TSharedPtr<DatasmithSketchUp::FComponentInstance> ComponentInstance = KeyValue.Value;
+		if (SUIsValid(ComponentInstance->LayerRef) && (LayerId == DatasmithSketchUpUtils::GetEntityID(SULayerToEntity(ComponentInstance->LayerRef))))
+		{
+			ComponentInstance->InvalidateEntityProperties();
+		}
+	}
+}
+
 void FMaterialCollection::PopulateFromModel(SUModelRef InModelRef)
 {
 	DefaultMaterial = FMaterial::CreateDefaultMaterial(Context);
diff --git a/Engine/Source/Programs/Enterprise/Datasmith/DatasmithSketchUpRubyExporter/Private/DatasmithSketchUpExportContext.h b/Engine/Source/Programs/Enterprise/Datasmith/DatasmithSketchUpRubyExporter/Private/DatasmithSketchUpExportContext.h
--- a/Engine/Source/Programs/Enterprise/Datasmith/DatasmithSketchUpRubyExporter/Private/DatasmithSketchUpExportContext.h
+++ b/Engine/Source/Programs/Enterprise/Datasmith/DatasmithSketchUpRubyExporter/Private/DatasmithSketchUpExportContext.h
@@ -58,15 +58,17 @@ namespace DatasmithSketchUp
 		void UpdateProperties();
 		void UpdateGeometry();
 
+		void LayerModified(DatasmithSketchUp::FEntityIDType LayerId);
+
 		TSharedPtr<FComponentInstance>* FindComponentInstance(FComponentInstanceIDType ComponentInstanceID)
 		{
 			return ComponentInstanceMap.Find(ComponentInstanceID);
 		}
 
+		TMap<FComponentInstanceIDType, TSharedPtr<FComponentInstance>> ComponentInstanceMap;
 	private:
 		FExportContext& Context;
 
-		TMap<FComponentInstanceIDType, TSharedPtr<FComponentInstance>> ComponentInstanceMap;
 	};
 
 	class FComponentDefinitionCollection
@@ -122,12 +124,16 @@ namespace DatasmithSketchUp
 
 		TSharedPtr<DatasmithSketchUp::FEntities> AddEntities(FDefinition& InDefinition, SUEntitiesRef EntitiesRef);
 
-		void RegisterEntitiesFaces(DatasmithSketchUp::FEntities&, const TSet<int32>& FaceIds);
+		void RegisterEntities(DatasmithSketchUp::FEntities&);
+		void UnregisterEntities(DatasmithSketchUp::FEntities&);
 		DatasmithSketchUp::FEntities* FindFace(int32 FaceId);
 
+		void LayerModified(FEntityIDType LayerId);
+
 	private:
 		FExportContext& Context;
 		TMap<int32, DatasmithSketchUp::FEntities*> FaceIdForEntitiesMap; // Identify Entities for each Face
+		TMap<DatasmithSketchUp::FEntityIDType, TSet<DatasmithSketchUp::FEntities*>> LayerIdForEntitiesMap; // Identify Entities for each Face
 	};
 
 	// Tracks information related to SketchUp "Scenes"(or "Pages" in older UI)
diff --git a/Engine/Source/Programs/Enterprise/Datasmith/DatasmithSketchUpRubyExporter/Private/DatasmithSketchUpExporter.cpp b/Engine/Source/Programs/Enterprise/Datasmith/DatasmithSketchUpRubyExporter/Private/DatasmithSketchUpExporter.cpp
--- a/Engine/Source/Programs/Enterprise/Datasmith/DatasmithSketchUpRubyExporter/Private/DatasmithSketchUpExporter.cpp
+++ b/Engine/Source/Programs/Enterprise/Datasmith/DatasmithSketchUpRubyExporter/Private/DatasmithSketchUpExporter.cpp
@@ -314,45 +314,6 @@ public:
 		return false;
 	}
 
-	bool OnEntityModified(SUEntityRef Entity)
-	{
-		switch (SUEntityGetType(Entity))
-		{
-		case SURefType_Group:
-		case SURefType_ComponentInstance:
-		{
-			OnComponentInstanceChanged(Entity);
-			break;
-		}
-		case SURefType_Face:
-		{
-			int32_t FaceId = 0;
-			if (SUEntityGetID(Entity, &FaceId) == SU_ERROR_NONE)
-			{
-				InvalidateGeometryForFace(DatasmithSketchUp::FEntityIDType(FaceId));
-				// todo: can we change DatasmithMesh in-place?
-				// - In case mesh need to be updated on DatasmithMeshActors:
-				// - and reset on every MeshActor, (or create if was empty? or not - modified, not added, not expected)
-			}
-			else
-			{
-				// todo: unexpected
-			}
-			break;
-		}
-		case SURefType_Material:
-		{
-			Context.Materials.InvalidateMaterial(SUMaterialFromEntity(Entity));
-			break;
-		}
-		default:
-		{
-			// todo: not expected
-		}
-		}
-		return true;
-	}
-
 	bool OnGeometryModified(DatasmithSketchUp::FEntityIDType EntityId)
 	{
 		DatasmithSketchUp::FDefinition* DefinitionPtr = Context.GetDefinition(EntityId);
@@ -451,6 +412,16 @@ public:
 		Context.Materials.CreateMaterial(EntityId);
 		return true;
 	}
+
+	bool OnLayerModified(SUEntityRef Entity)
+	{
+		DatasmithSketchUp::FEntityIDType LayerId = DatasmithSketchUpUtils::GetEntityID(Entity);
+		Context.ComponentInstances.LayerModified(LayerId);
+		Context.EntitiesObjects.LayerModified(LayerId);
+		return true;
+	}
+
+
 };
 
 ///////////////////////////////////////////////////////////////////////////
@@ -577,48 +548,49 @@ VALUE DatasmithSketchUpDirectLinkExporter_on_component_instance_changed(VALUE se
 	return Qtrue;
 }
 
-VALUE DatasmithSketchUpDirectLinkExporter_on_entity_modified(VALUE self, VALUE ruby_entity)
+VALUE DatasmithSketchUpDirectLinkExporter_on_entity_added(VALUE self, VALUE ruby_parent_entity, VALUE ruby_entity)
 {
 	// Converting args
 	FDatasmithSketchUpDirectLinkExporter* Ptr;
 	Data_Get_Struct(self, FDatasmithSketchUpDirectLinkExporter, Ptr);
 
-	int RubyEntityRubyObjectType = TYPE(ruby_entity);
+	SUEntityRef ParentEntity = SU_INVALID;
+	if (!NIL_P(ruby_parent_entity) && (SUEntityFromRuby(ruby_parent_entity, &ParentEntity) != SU_ERROR_NONE)) {
+		
+		rb_raise(rb_eTypeError, "Expected SketchUp Entity but found '%s'", StringValue(ruby_parent_entity));
+	}
 
 	SUEntityRef Entity = SU_INVALID;
-
 	if (SUEntityFromRuby(ruby_entity, &Entity) != SU_ERROR_NONE) {
 		rb_raise(rb_eTypeError, "Expected SketchUp Entity or nil");
 	}
 	// Done converting args
 
-	Ptr->OnEntityModified(Entity);
+	Ptr->OnEntityAdded(ParentEntity, Entity);
 
 	return Qtrue;
 }
 
-VALUE DatasmithSketchUpDirectLinkExporter_on_entity_added(VALUE self, VALUE ruby_parent_entity, VALUE ruby_entity)
+VALUE DatasmithSketchUpDirectLinkExporter_on_layer_modified(VALUE self, VALUE ruby_entity)
 {
 	// Converting args
 	FDatasmithSketchUpDirectLinkExporter* Ptr;
 	Data_Get_Struct(self, FDatasmithSketchUpDirectLinkExporter, Ptr);
 
-	SUEntityRef ParentEntity = SU_INVALID;
-	if (!NIL_P(ruby_parent_entity) && (SUEntityFromRuby(ruby_parent_entity, &ParentEntity) != SU_ERROR_NONE)) {
-		
-		rb_raise(rb_eTypeError, "Expected SketchUp Entity but found '%s'", StringValue(ruby_parent_entity));
-	}
+	int RubyEntityRubyObjectType = TYPE(ruby_entity);
 
 	SUEntityRef Entity = SU_INVALID;
+
 	if (SUEntityFromRuby(ruby_entity, &Entity) != SU_ERROR_NONE) {
 		rb_raise(rb_eTypeError, "Expected SketchUp Entity or nil");
 	}
 	// Done converting args
 
-	Ptr->OnEntityAdded(ParentEntity, Entity);
+	Ptr->OnLayerModified(Entity);
 
 	return Qtrue;
 }
+
 #endif
 
 VALUE DatasmithSketchUpDirectLinkExporter_on_entity_modified_by_id(VALUE self, VALUE ruby_entity_id)
@@ -771,8 +743,8 @@ extern "C" DLLEXPORT void Init_DatasmithSketchUpRuby()
 
 #ifndef SKP_SDK_2019
 	rb_define_method(DatasmithSketchUpDirectLinkExporterCRubyClass, "on_component_instance_changed", ToRuby(DatasmithSketchUpDirectLinkExporter_on_component_instance_changed), 1);
-	rb_define_method(DatasmithSketchUpDirectLinkExporterCRubyClass, "on_entity_modified", ToRuby(DatasmithSketchUpDirectLinkExporter_on_entity_modified), 1);
 	rb_define_method(DatasmithSketchUpDirectLinkExporterCRubyClass, "on_entity_added", ToRuby(DatasmithSketchUpDirectLinkExporter_on_entity_added), 2);
+	rb_define_method(DatasmithSketchUpDirectLinkExporterCRubyClass, "on_layer_modified", ToRuby(DatasmithSketchUpDirectLinkExporter_on_layer_modified), 1);
 #endif
 
 	rb_define_method(DatasmithSketchUpDirectLinkExporterCRubyClass, "on_entity_modified_by_id", ToRuby(DatasmithSketchUpDirectLinkExporter_on_entity_modified_by_id), 1);
diff --git a/Engine/Source/Programs/Enterprise/Datasmith/DatasmithSketchUpRubyExporter/Private/DatasmithSketchUpMesh.cpp b/Engine/Source/Programs/Enterprise/Datasmith/DatasmithSketchUpRubyExporter/Private/DatasmithSketchUpMesh.cpp
--- a/Engine/Source/Programs/Enterprise/Datasmith/DatasmithSketchUpRubyExporter/Private/DatasmithSketchUpMesh.cpp
+++ b/Engine/Source/Programs/Enterprise/Datasmith/DatasmithSketchUpRubyExporter/Private/DatasmithSketchUpMesh.cpp
@@ -343,7 +343,7 @@ const TCHAR* FEntitiesGeometry::GetMeshElementName(int32 MeshIndex)
 	return Meshes[MeshIndex]->DatasmithMesh->GetName();
 }
 
-void ScanSketchUpEntitiesFaces(SUEntitiesRef EntitiesRef, TSet<int32>& ScannedFaceIDSet, TFunctionRef<void(TSharedPtr<FDatasmithSketchUpMesh> ExtractedMesh)> OnNewExtractedMesh);
+void ScanSketchUpEntitiesFaces(SUEntitiesRef EntitiesRef, FEntitiesGeometry& Geometry, TFunctionRef<void(TSharedPtr<FDatasmithSketchUpMesh> ExtractedMesh)> OnNewExtractedMesh);
 
 void FEntities::UpdateGeometry(FExportContext& Context)
 {
@@ -356,6 +356,9 @@ void FEntities::UpdateGeometry(FExportContext& Context)
 			Context.DatasmithScene->RemoveMesh(Mesh->DatasmithMesh);
 		}
 
+		Context.EntitiesObjects.UnregisterEntities(*this);
+		EntitiesGeometry->FaceIds.Reset();
+		EntitiesGeometry->Layers.Reset();
 	}
 	else
 	{
@@ -415,15 +418,12 @@ void FEntities::UpdateGeometry(FExportContext& Context)
 		}
 	};
 
-	TSet<int32> FaceIds;
-	ScanSketchUpEntitiesFaces(EntitiesRef, FaceIds, ProcessExtractedMesh);
-
+	ScanSketchUpEntitiesFaces(EntitiesRef, *EntitiesGeometry, ProcessExtractedMesh);
 	EntitiesGeometry->Meshes.SetNum(MeshCount);
 
-	Context.EntitiesObjects.RegisterEntitiesFaces(*this, FaceIds);
+	Context.EntitiesObjects.RegisterEntities(*this);
 }
 
-
 void FEntities::AddMeshesToDatasmithScene(FExportContext& Context)
 {
 	for (TSharedPtr<FDatasmithInstantiatedMesh> Mesh : EntitiesGeometry->Meshes)
@@ -483,7 +483,7 @@ TArray<SUComponentInstanceRef> FEntities::GetComponentInstances()
 }
 
 
-void ScanSketchUpEntitiesFaces(SUEntitiesRef EntitiesRef, TSet<int32>& ScannedFaceIDSet, TFunctionRef<void(TSharedPtr<FDatasmithSketchUpMesh> ExtractedMesh)> OnNewExtractedMesh)
+void ScanSketchUpEntitiesFaces(SUEntitiesRef EntitiesRef, FEntitiesGeometry& Geometry, TFunctionRef<void(TSharedPtr<FDatasmithSketchUpMesh> ExtractedMesh)> OnNewExtractedMesh)
 {
 	// Get the number of faces in the source SketchUp entities.
 	size_t SFaceCount = 0;
@@ -511,7 +511,7 @@ void ScanSketchUpEntitiesFaces(SUEntitiesRef EntitiesRef, TSet<int32>& ScannedFa
 		int32 SSourceFaceID = DatasmithSketchUpUtils::GetFaceID(SSourceFaceRef);
 
 		// Do not scan more than once a valid SketckUp face.
-		if (SUIsInvalid(SSourceFaceRef) || ScannedFaceIDSet.Contains(SSourceFaceID))
+		if (SUIsInvalid(SSourceFaceRef) || Geometry.FaceIds.Contains(SSourceFaceID))
 		{
 			continue;
 		}
@@ -523,7 +523,7 @@ void ScanSketchUpEntitiesFaces(SUEntitiesRef EntitiesRef, TSet<int32>& ScannedFa
 		// The source SketchUp face needs to be scanned once.
 		TArray<SUFaceRef> FacesToScan;
 		FacesToScan.Add(SSourceFaceRef);
-		ScannedFaceIDSet.Add(SSourceFaceID);
+		Geometry.FaceIds.Add(SSourceFaceID);
 
 		// Collect all connected faces
 		while (FacesToScan.Num() > 0)
@@ -534,10 +534,16 @@ void ScanSketchUpEntitiesFaces(SUEntitiesRef EntitiesRef, TSet<int32>& ScannedFa
 			// SUEntityGetPersistentID(SUFaceToEntity(SScannedFaceRef), &SFacePID);
 			// ADD_TRACE_LINE(TEXT("   Face %lld"), SFacePID);
 
+			// Record every face's layer(even for invisible faces!). When face layer visibility changes 
+			// this geometry needs to be rebuilt
+			SULayerRef LayerRef = SU_INVALID;
+			SUDrawingElementGetLayer(SUFaceToDrawingElement(SScannedFaceRef), &LayerRef);
+			Geometry.Layers.Add(DatasmithSketchUpUtils::GetEntityID(SULayerToEntity(LayerRef)));
+
+
 			// Get whether or not the SketckUp face is visible in the current SketchUp scene.
 			if (DatasmithSketchUpUtils::IsVisible(SScannedFaceRef))
 			{
-				// Tessellate the SketchUp face into a triangle mesh merged into the combined mesh.
 				ExtractedMesh.AddFace(SScannedFaceRef);
 			}
 
@@ -577,9 +583,9 @@ void ScanSketchUpEntitiesFaces(SUEntitiesRef EntitiesRef, TSet<int32>& ScannedFa
 						int32 SFaceID = DatasmithSketchUpUtils::GetFaceID(SFaceRef);
 
 						// Avoid scanning more than once this SketckUp face.
-						if (!ScannedFaceIDSet.Contains(SFaceID))
+						if (!Geometry.FaceIds.Contains(SFaceID))
 						{
-							ScannedFaceIDSet.Add(SFaceID);
+							Geometry.FaceIds.Add(SFaceID);
 
 							// This SketchUp face is connected and needs to be scanned further.
 							FacesToScan.Add(SFaceRef);
