Added support for 3ds max 2022

diff --git a/Engine/Source/Programs/Enterprise/Datasmith/DatasmithMaxExporter/DatasmithMax2017.Build.cs b/Engine/Source/Programs/Enterprise/Datasmith/DatasmithMaxExporter/DatasmithMax2017.Build.cs
--- a/Engine/Source/Programs/Enterprise/Datasmith/DatasmithMaxExporter/DatasmithMax2017.Build.cs
+++ b/Engine/Source/Programs/Enterprise/Datasmith/DatasmithMaxExporter/DatasmithMax2017.Build.cs
@@ -26,20 +26,20 @@ namespace UnrealBuildTool.Rules
 				string MaxVersionString = GetMaxVersion();
 				string MaxSDKLocation = "";
 
-				// Try with custom setup
-				string MaxSDKEnvVar = System.Environment.GetEnvironmentVariable("ADSK_3DSMAX_SDK_" + MaxVersionString);
-				if (MaxSDKEnvVar != null && MaxSDKEnvVar != "")
+				// Try with build machine setup
+				string SDKRootEnvVar = System.Environment.GetEnvironmentVariable("UE_SDKS_ROOT");
+				if (SDKRootEnvVar != null && SDKRootEnvVar != "")
 				{
-					MaxSDKLocation = MaxSDKEnvVar;
+					MaxSDKLocation = Path.Combine(SDKRootEnvVar, "HostWin64", "Win64", "3dsMax", MaxVersionString);
 				}
 
 				if (!Directory.Exists(MaxSDKLocation))
 				{
-					// Try with build machine setup
-					string SDKRootEnvVar = System.Environment.GetEnvironmentVariable("UE_SDKS_ROOT");
-					if (SDKRootEnvVar != null && SDKRootEnvVar != "")
+					// Try with custom setup
+					string MaxSDKEnvVar = System.Environment.GetEnvironmentVariable("ADSK_3DSMAX_SDK_" + MaxVersionString);
+					if (MaxSDKEnvVar != null && MaxSDKEnvVar != "")
 					{
-						MaxSDKLocation = Path.Combine(SDKRootEnvVar, "HostWin64", "Win64", "3dsMax", MaxVersionString);
+						MaxSDKLocation = MaxSDKEnvVar;
 					}
 				}
 
diff --git a/Engine/Source/Programs/Enterprise/Datasmith/DatasmithMaxExporter/DatasmithMax2022.Build.cs b/Engine/Source/Programs/Enterprise/Datasmith/DatasmithMaxExporter/DatasmithMax2022.Build.cs
new file mode 100644
--- /dev/null
+++ b/Engine/Source/Programs/Enterprise/Datasmith/DatasmithMaxExporter/DatasmithMax2022.Build.cs
@@ -0,0 +1,15 @@
+// Copyright Epic Games, Inc. All Rights Reserved.
+
+namespace UnrealBuildTool.Rules
+{
+	[SupportedPlatforms("Win64")]
+	public class DatasmithMax2022 : DatasmithMaxBase
+	{
+		public DatasmithMax2022(ReadOnlyTargetRules Target)
+			: base(Target)
+		{
+		}
+
+		public override string GetMaxVersion() { return "2022"; }
+	}
+}
\ No newline at end of file
diff --git a/Engine/Source/Programs/Enterprise/Datasmith/DatasmithMaxExporter/DatasmithMax2022.Target.cs b/Engine/Source/Programs/Enterprise/Datasmith/DatasmithMaxExporter/DatasmithMax2022.Target.cs
new file mode 100644
--- /dev/null
+++ b/Engine/Source/Programs/Enterprise/Datasmith/DatasmithMaxExporter/DatasmithMax2022.Target.cs
@@ -0,0 +1,15 @@
+// Copyright Epic Games, Inc. All Rights Reserved.
+
+using UnrealBuildTool;
+using System.Collections.Generic;
+
+public class DatasmithMax2022Target : DatasmithMaxBaseTarget
+{
+	public DatasmithMax2022Target(TargetInfo Target) : base(Target)
+	{
+		LaunchModuleName = "DatasmithMax2022";
+		ExeBinariesSubFolder = @"3DSMax\2022";
+
+		AddCopyPostBuildStep(Target);
+	}
+}
diff --git a/Engine/Source/Programs/Enterprise/Datasmith/DatasmithMaxExporter/Private/DatasmithMaxCameraExporter.cpp b/Engine/Source/Programs/Enterprise/Datasmith/DatasmithMaxExporter/Private/DatasmithMaxCameraExporter.cpp
--- a/Engine/Source/Programs/Enterprise/Datasmith/DatasmithMaxExporter/Private/DatasmithMaxCameraExporter.cpp
+++ b/Engine/Source/Programs/Enterprise/Datasmith/DatasmithMaxExporter/Private/DatasmithMaxCameraExporter.cpp
@@ -2,11 +2,12 @@
 #include "DatasmithMaxCameraExporter.h"
 
 #include "DatasmithMaxExporterDefines.h"
-#include "DatasmithSceneFactory.h"
-#include "DatasmithSceneExporter.h"
+#include "DatasmithMaxHelper.h"
 #include "DatasmithMaxSceneParser.h"
 #include "DatasmithMaxWriter.h"
 #include "DatasmithMaxLogger.h"
+#include "DatasmithSceneFactory.h"
+#include "DatasmithSceneExporter.h"
 #include "VRayLights.h"
 
 #include "GenericPlatform/GenericPlatformFile.h"
@@ -85,7 +86,7 @@ namespace DatasmithMaxCameraExporterImpl
 
 	FMaxCameraParameters ParseCameraParameters(Animatable& MaxObject)
 	{
-		const float CustomUnit = FMath::Abs( (float)GetMasterScale(UNITS_CENTIMETERS) );
+		const float CustomUnit = FMath::Abs( (float)GetSystemUnitScale(UNITS_CENTIMETERS) );
 		const TimeValue CurrentTime = GetCOREInterface()->GetTime();
 
 		FMaxCameraParameters CameraParameters;
@@ -296,11 +297,11 @@ namespace DatasmithMaxCameraExporterImpl
 					}
 					else if ( FCString::Stricmp(ParamDefinition.int_name, TEXT("focus_distance")) == 0 )
 					{
-						CameraParameters.FocusDistance = ParamBlock2->GetFloat( ParamDefinition.ID, CurrentTime, ValidInterval ) * (float)GetMasterScale( UNITS_CENTIMETERS );
+						CameraParameters.FocusDistance = ParamBlock2->GetFloat( ParamDefinition.ID, CurrentTime, ValidInterval ) * (float)GetSystemUnitScale( UNITS_CENTIMETERS );
 					}
 					else if ( FCString::Stricmp( ParamDefinition.int_name, TEXT("target_distance") ) == 0 )
 					{
-						CameraParameters.TargetDistance = ParamBlock2->GetFloat( ParamDefinition.ID, CurrentTime, ValidInterval ) * (float)GetMasterScale( UNITS_CENTIMETERS );
+						CameraParameters.TargetDistance = ParamBlock2->GetFloat( ParamDefinition.ID, CurrentTime, ValidInterval ) * (float)GetSystemUnitScale( UNITS_CENTIMETERS );
 					}
 					else if ( FCString::Stricmp( ParamDefinition.int_name, TEXT("ISO") ) == 0 )
 					{
@@ -405,7 +406,7 @@ bool FDatasmithMaxCameraExporter::ExportCamera(INode& Node, TSharedRef< IDatasmi
 	Camera->SetFocalLength( GetCOREInterface()->GetRendApertureWidth() / (2.f * FMath::Tan(CamObj->GetFOV(GetCOREInterface()->GetTime()) * 0.5f)) );
 
 	int NumParamBlocks = CamObj->NumParamBlocks();
-	float CustomUnit = FMath::Abs((float)GetMasterScale(UNITS_CENTIMETERS));
+	float CustomUnit = FMath::Abs((float)GetSystemUnitScale(UNITS_CENTIMETERS));
 
 	DatasmithMaxCameraExporterImpl::FMaxCameraParameters MaxCameraParameters = DatasmithMaxCameraExporterImpl::ParseCameraParameters( *CamObj );
 
diff --git a/Engine/Source/Programs/Enterprise/Datasmith/DatasmithMaxExporter/Private/DatasmithMaxExporter.cpp b/Engine/Source/Programs/Enterprise/Datasmith/DatasmithMaxExporter/Private/DatasmithMaxExporter.cpp
--- a/Engine/Source/Programs/Enterprise/Datasmith/DatasmithMaxExporter/Private/DatasmithMaxExporter.cpp
+++ b/Engine/Source/Programs/Enterprise/Datasmith/DatasmithMaxExporter/Private/DatasmithMaxExporter.cpp
@@ -19,6 +19,7 @@
 #include "DatasmithMaxAttributes.h"
 #include "DatasmithMaxClassIDs.h"
 #include "DatasmithMaxExporterDefines.h"
+#include "DatasmithMaxHelper.h"
 #include "DatasmithMaxMeshExporter.h"
 #include "DatasmithMaxSceneParser.h"
 #include "DatasmithMaxSceneExporter.h"
@@ -59,6 +60,10 @@ public:
 	bool UseOnlyInternalNameForMAXScriptExposure() { return true; }
 	const TCHAR* InternalName() { return TEXT("DatasmithExport"); }
 
+#if MAX_PRODUCT_YEAR_NUMBER > 2021
+	virtual const MCHAR* NonLocalizedClassName() override { return _M("datasmithmaxexporter"); }
+#endif
+
 #ifdef UNICODE_
 	virtual const MCHAR* GetRsrcString(int id) { return _M(""); }
 #else
@@ -339,7 +344,7 @@ namespace DatasmithMaxExporterUtils
 
 int FDatasmithMaxExporter::ExportScene(FDatasmithSceneExporter* DatasmithSceneExporter, TSharedRef< IDatasmithScene > DatasmithScene, FDatasmithMaxSceneParser& InSceneParser, const TCHAR* Filename, const FDatasmithMaxExportOptions& ExporterOptions, TSharedPtr < FDatasmithMaxProgressManager >& ProgressManager)
 {
-	float UnitMultiplier = FMath::Abs(GetMasterScale(UNITS_CENTIMETERS));
+	float UnitMultiplier = FMath::Abs(GetSystemUnitScale(UNITS_CENTIMETERS));
 
 	MSTR Renderer;
 	FString Host;
diff --git a/Engine/Source/Programs/Enterprise/Datasmith/DatasmithMaxExporter/Private/DatasmithMaxSceneExporter.cpp b/Engine/Source/Programs/Enterprise/Datasmith/DatasmithMaxExporter/Private/DatasmithMaxSceneExporter.cpp
--- a/Engine/Source/Programs/Enterprise/Datasmith/DatasmithMaxExporter/Private/DatasmithMaxSceneExporter.cpp
+++ b/Engine/Source/Programs/Enterprise/Datasmith/DatasmithMaxExporter/Private/DatasmithMaxSceneExporter.cpp
@@ -232,8 +232,12 @@ FTransform FDatasmithMaxSceneExporter::GetPivotTransform( INode* Node, float Uni
 
 	if ( !bIsInWorldSpace )
 	{
+// Matrix3::Matrix3(BOOL) is deprecated in 3ds max 2022 SDK
+#if MAX_PRODUCT_YEAR_NUMBER < 2022
 		Matrix3 ObjectOffset(1);
-
+#else
+		Matrix3 ObjectOffset;
+#endif
 		Point3 OffsetPos = Node->GetObjOffsetPos(); // Object translation in Node (pivot) space
 		ObjectOffset.PreTranslate( OffsetPos );
 
@@ -830,7 +834,12 @@ bool FDatasmithMaxSceneExporter::ParseTransformAnimation(INode* Node, TSharedRef
 	Matrix3 NodeTransform = Node->GetNodeTM(TIME_NegInfinity, &ValidInterval);
 	INode* Parent = Node->GetParentNode();
 
+// Matrix3::Matrix3(BOOL) is deprecated in 3ds max 2022 SDK
+#if MAX_PRODUCT_YEAR_NUMBER < 2022
 	Matrix3 ParentTransform = Parent ? Parent->GetNodeTM(TIME_NegInfinity, &ValidInterval) : Matrix3(true);
+#else
+	Matrix3 ParentTransform = Parent ? Parent->GetNodeTM(TIME_NegInfinity, &ValidInterval) : Matrix3();
+#endif
 	Matrix3 LocalTransform;
 
 	// An infinite interval indicates there's no animation in the transform
@@ -852,7 +861,12 @@ bool FDatasmithMaxSceneExporter::ParseTransformAnimation(INode* Node, TSharedRef
 		// The parent node could change at each frame because of the Link Constraint
 		NodeTransform = Node->GetNodeTM(CurrentTime, &ValidInterval);
 		Parent = Node->GetParentNode();
+// Matrix3::Matrix3(BOOL) is deprecated in 3ds max 2022 SDK
+#if MAX_PRODUCT_YEAR_NUMBER < 2022
 		ParentTransform = Parent ? Parent->GetNodeTM(CurrentTime, &ValidInterval) : Matrix3(true);
+#else
+		ParentTransform = Parent ? Parent->GetNodeTM(CurrentTime, &ValidInterval) : Matrix3();
+#endif
 
 		Matrix3 InvParentTransform = Inverse(ParentTransform);
 		LocalTransform = NodeTransform * InvParentTransform;
@@ -1027,7 +1041,7 @@ bool FDatasmithMaxSceneExporter::ParseLight(INode* Node, TSharedRef< IDatasmithL
 		FVector Translation, Scale;
 		FQuat Rotation;
 
-		const float UnitMultiplier = (float)GetMasterScale(UNITS_CENTIMETERS);
+		const float UnitMultiplier = (float)GetSystemUnitScale(UNITS_CENTIMETERS);
 		const FMaxLightCoordinateConversionParams LightParams = FMaxLightCoordinateConversionParams(Node, EDatasmithLightShape::Cylinder);
 		if (Node->GetWSMDerivedObject() != nullptr)
 		{
@@ -1066,7 +1080,7 @@ bool FDatasmithMaxSceneExporter::ParseLightObject(LightObject& Light, TSharedRef
 
 		if ( LState.useAtten != 0 ) // Using the LightState here because calling Light.GetUseAtten() on a Corona Light crashes
 		{
-			PointLightElement->SetAttenuationRadius( LState.attenEnd * (float)GetMasterScale( UNITS_CENTIMETERS ) );
+			PointLightElement->SetAttenuationRadius( LState.attenEnd * (float)GetSystemUnitScale( UNITS_CENTIMETERS ) );
 		}
 	}
 
@@ -1197,11 +1211,11 @@ bool FDatasmithMaxSceneExporter::ParseCoronaLight(LightObject& Light, TSharedRef
 			}
 			else if ( FCString::Stricmp(ParamDefinition.int_name, TEXT("width")) == 0 )
 			{
-				AreaLightElement->SetWidth( ParamBlock2->GetFloat( ParamDefinition.ID, GetCOREInterface()->GetTime() ) * (float)GetMasterScale( UNITS_CENTIMETERS ) );
+				AreaLightElement->SetWidth( ParamBlock2->GetFloat( ParamDefinition.ID, GetCOREInterface()->GetTime() ) * (float)GetSystemUnitScale( UNITS_CENTIMETERS ) );
 			}
 			else if ( FCString::Stricmp(ParamDefinition.int_name, TEXT("height")) == 0 )
 			{
-				AreaLightElement->SetLength( ParamBlock2->GetFloat( ParamDefinition.ID, GetCOREInterface()->GetTime() ) * (float)GetMasterScale( UNITS_CENTIMETERS ) );
+				AreaLightElement->SetLength( ParamBlock2->GetFloat( ParamDefinition.ID, GetCOREInterface()->GetTime() ) * (float)GetSystemUnitScale( UNITS_CENTIMETERS ) );
 			}
 			else if ( FCString::Stricmp(ParamDefinition.int_name, TEXT("visibleDirectly")) == 0 )
 			{
@@ -1417,16 +1431,16 @@ bool FDatasmithMaxSceneExporter::ParsePhotometricLight(LightObject& Light, TShar
 		case LightscapeLight::DISC_TYPE:
 		case LightscapeLight::TARGET_DISC_TYPE:
 			AreaLightElement->SetLightShape( EDatasmithLightShape::Disc );
-			AreaLightElement->SetWidth( PhotometricLight.GetRadius( GetCOREInterface()->GetTime() ) * 2.f * (float)GetMasterScale( UNITS_CENTIMETERS ) );
-			AreaLightElement->SetLength( PhotometricLight.GetRadius( GetCOREInterface()->GetTime() ) * 2.f * (float)GetMasterScale( UNITS_CENTIMETERS ) );
+			AreaLightElement->SetWidth( PhotometricLight.GetRadius( GetCOREInterface()->GetTime() ) * 2.f * (float)GetSystemUnitScale( UNITS_CENTIMETERS ) );
+			AreaLightElement->SetLength( PhotometricLight.GetRadius( GetCOREInterface()->GetTime() ) * 2.f * (float)GetSystemUnitScale( UNITS_CENTIMETERS ) );
 			bIsOneSided = true;
 
 			break;
 		case LightscapeLight::AREA_TYPE:
 		case LightscapeLight::TARGET_AREA_TYPE:
 			AreaLightElement->SetLightShape( EDatasmithLightShape::Rectangle );
-			AreaLightElement->SetWidth( PhotometricLight.GetWidth( GetCOREInterface()->GetTime() ) * (float)GetMasterScale( UNITS_CENTIMETERS ) );
-			AreaLightElement->SetLength( PhotometricLight.GetLength( GetCOREInterface()->GetTime() ) * (float)GetMasterScale( UNITS_CENTIMETERS ) );
+			AreaLightElement->SetWidth( PhotometricLight.GetWidth( GetCOREInterface()->GetTime() ) * (float)GetSystemUnitScale( UNITS_CENTIMETERS ) );
+			AreaLightElement->SetLength( PhotometricLight.GetLength( GetCOREInterface()->GetTime() ) * (float)GetSystemUnitScale( UNITS_CENTIMETERS ) );
 			bIsOneSided = true;
 
 			break;
@@ -1435,15 +1449,15 @@ bool FDatasmithMaxSceneExporter::ParsePhotometricLight(LightObject& Light, TShar
 		case LightscapeLight::TARGET_LINEAR_TYPE:
 			AreaLightElement->SetLightShape( EDatasmithLightShape::Rectangle );
 			AreaLightElement->SetWidth( 1.f ); // 1 cm
-			AreaLightElement->SetLength( PhotometricLight.GetLength( GetCOREInterface()->GetTime() ) * (float)GetMasterScale( UNITS_CENTIMETERS ) );
+			AreaLightElement->SetLength( PhotometricLight.GetLength( GetCOREInterface()->GetTime() ) * (float)GetSystemUnitScale( UNITS_CENTIMETERS ) );
 
 			break;
 
 		case LightscapeLight::SPHERE_TYPE:
 		case LightscapeLight::TARGET_SPHERE_TYPE:
 			AreaLightElement->SetLightShape( EDatasmithLightShape::Sphere );
-			AreaLightElement->SetWidth( PhotometricLight.GetRadius( GetCOREInterface()->GetTime() ) * 2.f * (float)GetMasterScale( UNITS_CENTIMETERS ) );
-			AreaLightElement->SetLength( PhotometricLight.GetRadius( GetCOREInterface()->GetTime() ) * 2.f * (float)GetMasterScale( UNITS_CENTIMETERS ) );
+			AreaLightElement->SetWidth( PhotometricLight.GetRadius( GetCOREInterface()->GetTime() ) * 2.f * (float)GetSystemUnitScale( UNITS_CENTIMETERS ) );
+			AreaLightElement->SetLength( PhotometricLight.GetRadius( GetCOREInterface()->GetTime() ) * 2.f * (float)GetSystemUnitScale( UNITS_CENTIMETERS ) );
 			bHasVolume = true;
 			
 			break;
@@ -1451,8 +1465,8 @@ bool FDatasmithMaxSceneExporter::ParsePhotometricLight(LightObject& Light, TShar
 		case LightscapeLight::CYLINDER_TYPE:
 		case LightscapeLight::TARGET_CYLINDER_TYPE:
 			AreaLightElement->SetLightShape( EDatasmithLightShape::Cylinder );
-			AreaLightElement->SetWidth( PhotometricLight.GetRadius( GetCOREInterface()->GetTime() ) * 2.f * (float)GetMasterScale( UNITS_CENTIMETERS ) );
-			AreaLightElement->SetLength( PhotometricLight.GetLength( GetCOREInterface()->GetTime() ) * (float)GetMasterScale( UNITS_CENTIMETERS ) );
+			AreaLightElement->SetWidth( PhotometricLight.GetRadius( GetCOREInterface()->GetTime() ) * 2.f * (float)GetSystemUnitScale( UNITS_CENTIMETERS ) );
+			AreaLightElement->SetLength( PhotometricLight.GetLength( GetCOREInterface()->GetTime() ) * (float)GetSystemUnitScale( UNITS_CENTIMETERS ) );
 			bHasVolume = true;
 
 			break;
@@ -1540,8 +1554,8 @@ bool FDatasmithMaxSceneExporter::ParseVRayLight(LightObject& Light, TSharedRef<
 	}
 	
 
-	float VRayLightSize0 = ParamBlock2->GetFloat( (short)EVrayLightsParams::Size0 ) * 2.f * (float)GetMasterScale( UNITS_CENTIMETERS );
-	float VRayLightSize1 = ParamBlock2->GetFloat( (short)EVrayLightsParams::Size1 ) * 2.f * (float)GetMasterScale( UNITS_CENTIMETERS );
+	float VRayLightSize0 = ParamBlock2->GetFloat( (short)EVrayLightsParams::Size0 ) * 2.f * (float)GetSystemUnitScale( UNITS_CENTIMETERS );
+	float VRayLightSize1 = ParamBlock2->GetFloat( (short)EVrayLightsParams::Size1 ) * 2.f * (float)GetSystemUnitScale( UNITS_CENTIMETERS );
 	AreaLightElement->SetWidth( VRayLightSize0 );
 
 	const double CentimeterToMeter = 0.01;
@@ -1630,8 +1644,8 @@ bool FDatasmithMaxSceneExporter::ParseVRayLightPortal(LightObject& Light, TShare
 {
 	IParamBlock2* ParamBlock2 = Light.GetParamBlockByID( (short)EVRayLightParamBlocks::Params );
 
-	float VRayLightSize0 = ParamBlock2->GetFloat( (short)EVrayLightsParams::Size0 ) * (float)GetMasterScale( UNITS_CENTIMETERS );
-	float VRayLightSize1 = ParamBlock2->GetFloat( (short)EVrayLightsParams::Size1 ) * (float)GetMasterScale( UNITS_CENTIMETERS );
+	float VRayLightSize0 = ParamBlock2->GetFloat( (short)EVrayLightsParams::Size0 ) * (float)GetSystemUnitScale( UNITS_CENTIMETERS );
+	float VRayLightSize1 = ParamBlock2->GetFloat( (short)EVrayLightsParams::Size1 ) * (float)GetSystemUnitScale( UNITS_CENTIMETERS );
 	
 	FVector CurrentScale = LightPortalElement->GetScale();
 	LightPortalElement->SetScale( 10.f, VRayLightSize0 * CurrentScale.X, VRayLightSize1 * CurrentScale.Y);
@@ -1767,19 +1781,19 @@ bool FDatasmithMaxSceneExporter::ParseVRayLightIES(LightObject& Light, TSharedRe
 				{
 					if ( FCString::Stricmp(ParamDefinition.int_name, TEXT("width")) == 0 )
 					{
-						AreaLightElement->SetWidth( ParamBlock2->GetFloat( ParamDefinition.ID, GetCOREInterface()->GetTime() ) * (float)GetMasterScale( UNITS_CENTIMETERS ) );
+						AreaLightElement->SetWidth( ParamBlock2->GetFloat( ParamDefinition.ID, GetCOREInterface()->GetTime() ) * (float)GetSystemUnitScale( UNITS_CENTIMETERS ) );
 					}
 					else if ( FCString::Stricmp(ParamDefinition.int_name, TEXT("length")) == 0 )
 					{
-						AreaLightElement->SetLength( ParamBlock2->GetFloat( ParamDefinition.ID, GetCOREInterface()->GetTime() ) * (float)GetMasterScale( UNITS_CENTIMETERS ) );
+						AreaLightElement->SetLength( ParamBlock2->GetFloat( ParamDefinition.ID, GetCOREInterface()->GetTime() ) * (float)GetSystemUnitScale( UNITS_CENTIMETERS ) );
 					}
 				}
 				else
 				{
 					if ( FCString::Stricmp(ParamDefinition.int_name, TEXT("diameter")) == 0 )
 					{
-						AreaLightElement->SetWidth( ParamBlock2->GetFloat( ParamDefinition.ID, GetCOREInterface()->GetTime() ) * (float)GetMasterScale( UNITS_CENTIMETERS ) );
-						AreaLightElement->SetLength( ParamBlock2->GetFloat( ParamDefinition.ID, GetCOREInterface()->GetTime() ) * (float)GetMasterScale( UNITS_CENTIMETERS ) );
+						AreaLightElement->SetWidth( ParamBlock2->GetFloat( ParamDefinition.ID, GetCOREInterface()->GetTime() ) * (float)GetSystemUnitScale( UNITS_CENTIMETERS ) );
+						AreaLightElement->SetLength( ParamBlock2->GetFloat( ParamDefinition.ID, GetCOREInterface()->GetTime() ) * (float)GetSystemUnitScale( UNITS_CENTIMETERS ) );
 					}
 				}
 			}
@@ -1892,15 +1906,15 @@ bool FDatasmithMaxSceneExporter::ParseLightParameters(EMaxLightClass LightClass,
 
 				if (FCString::Stricmp(ParamDefinition.int_name, TEXT("size0")) == 0)
 				{
-					PointLightElement->SetSourceRadius( ParamBlock2->GetFloat(ParamDefinition.ID, GetCOREInterface()->GetTime()) * (float)GetMasterScale( UNITS_CENTIMETERS ) );
+					PointLightElement->SetSourceRadius( ParamBlock2->GetFloat(ParamDefinition.ID, GetCOREInterface()->GetTime()) * (float)GetSystemUnitScale( UNITS_CENTIMETERS ) );
 				}
 				else if (FCString::Stricmp(ParamDefinition.int_name, TEXT("width")) == 0 || FCString::Stricmp(ParamDefinition.int_name, TEXT("light_Width")) == 0 || FCString::Stricmp(ParamDefinition.int_name, TEXT("quadX")) == 0)
 				{
-					PointLightElement->SetSourceRadius( ParamBlock2->GetFloat(ParamDefinition.ID, GetCOREInterface()->GetTime()) / 2.0f * (float)GetMasterScale( UNITS_CENTIMETERS ) );
+					PointLightElement->SetSourceRadius( ParamBlock2->GetFloat(ParamDefinition.ID, GetCOREInterface()->GetTime()) / 2.0f * (float)GetSystemUnitScale( UNITS_CENTIMETERS ) );
 				}
 				else if (FCString::Stricmp(ParamDefinition.int_name, TEXT("size1")) == 0)
 				{
-					PointLightElement->SetSourceLength( ParamBlock2->GetFloat(ParamDefinition.ID, GetCOREInterface()->GetTime()) * (float)GetMasterScale( UNITS_CENTIMETERS ) );
+					PointLightElement->SetSourceLength( ParamBlock2->GetFloat(ParamDefinition.ID, GetCOREInterface()->GetTime()) * (float)GetSystemUnitScale( UNITS_CENTIMETERS ) );
 				}
 				else if (FCString::Stricmp(ParamDefinition.int_name, TEXT("height")) == 0 || FCString::Stricmp(ParamDefinition.int_name, TEXT("length")) == 0 ||
 					FCString::Stricmp(ParamDefinition.int_name, TEXT("light_length")) == 0 || FCString::Stricmp(ParamDefinition.int_name, TEXT("quadY")) == 0)
@@ -1908,7 +1922,7 @@ bool FDatasmithMaxSceneExporter::ParseLightParameters(EMaxLightClass LightClass,
 					//Arnold has height parameter but it is not the right one we are looking for quadY
 					if (LightClass != EMaxLightClass::ArnoldLight || FCString::Stricmp(ParamDefinition.int_name, TEXT("height")) != 0)
 					{
-						PointLightElement->SetSourceLength(ParamBlock2->GetFloat(ParamDefinition.ID, GetCOREInterface()->GetTime()) / 2.0f * (float)GetMasterScale( UNITS_CENTIMETERS ));
+						PointLightElement->SetSourceLength(ParamBlock2->GetFloat(ParamDefinition.ID, GetCOREInterface()->GetTime()) / 2.0f * (float)GetSystemUnitScale( UNITS_CENTIMETERS ));
 					}
 				}
 			}
diff --git a/Engine/Source/Programs/Enterprise/Datasmith/DatasmithMaxExporter/Public/DatasmithMaxHelper.h b/Engine/Source/Programs/Enterprise/Datasmith/DatasmithMaxExporter/Public/DatasmithMaxHelper.h
--- a/Engine/Source/Programs/Enterprise/Datasmith/DatasmithMaxExporter/Public/DatasmithMaxHelper.h
+++ b/Engine/Source/Programs/Enterprise/Datasmith/DatasmithMaxExporter/Public/DatasmithMaxHelper.h
@@ -7,9 +7,15 @@
 #include "Windows/AllowWindowsPlatformTypes.h"
 MAX_INCLUDES_START
 #include "bitmap.h"
+#include "maxversion.h"
+#include "units.h"
 MAX_INCLUDES_END
 #include "Windows/HideWindowsPlatformTypes.h"
 
+// In 3ds max 2022 SDK, GetMasterScale is deprecated and replaced with GetSystemUnitScale
+#if MAX_PRODUCT_YEAR_NUMBER < 2022
+inline double GetSystemUnitScale(int type) { return GetMasterScale(type); }
+#endif
 
 // Helper structure to help manage loading and deleting bitmaps when we only have its BitmapInfo
 struct FScopedBitMapPtr
