SketchUp - assigned material slots remapped to keep them stable between reexports

diff --git a/Engine/Source/Programs/Enterprise/Datasmith/DatasmithSketchUpRubyExporter/Private/DatasmithSketchUpMesh.cpp b/Engine/Source/Programs/Enterprise/Datasmith/DatasmithSketchUpRubyExporter/Private/DatasmithSketchUpMesh.cpp
--- a/Engine/Source/Programs/Enterprise/Datasmith/DatasmithSketchUpRubyExporter/Private/DatasmithSketchUpMesh.cpp
+++ b/Engine/Source/Programs/Enterprise/Datasmith/DatasmithSketchUpRubyExporter/Private/DatasmithSketchUpMesh.cpp
@@ -71,6 +71,21 @@ namespace DatasmithSketchUp
 			SUFaceRef InSFaceRef // valid source SketchUp face to tessellate and combine
 		);
 
+		int32 GetOrCreateSlotForMaterial(FMaterialIDType MaterialID)
+		{
+			if (int32* SlotIdPtr = SlotIdForMaterialId.Find(MaterialID))
+			{
+				return *SlotIdPtr;
+			}
+			else
+			{
+				int32 SlotId = MaterialIDForSlotId.Num();
+				MaterialIDForSlotId.Add(MaterialID);// Assign material to slot
+				SlotIdForMaterialId.Add(MaterialID, SlotId); // store back reference from material to slot
+				return SlotId;
+			}
+		}
+
 		// Return whether or not the combined mesh contains geometry.
 		bool ContainsGeometry() const;
 
@@ -87,10 +102,12 @@ namespace DatasmithSketchUp
 		TArray<SMeshTriangleIndices> MeshTriangleIndices;
 
 		// Combined mesh triangle material IDs.
-		TArray<FEntityIDType> MeshTriangleMaterialIDs;
+		TArray<int32> MeshTriangleSlotIds;
+
+		TArray<FEntityIDType> MaterialIDForSlotId;
+		TMap<FEntityIDType, int32> SlotIdForMaterialId;
+		bool bHasFacesWithDefaultMaterial = false;;
 
-		// Set of all the material IDs used by the combined mesh triangles.
-		TSet<FEntityIDType> MeshTriangleMaterialIDSet;
 	};
 
 
@@ -242,8 +259,8 @@ namespace DatasmithSketchUp
 		// Release SketchUp face UV helper.
 		SUUVHelperRelease(&UVHelperRef); // we can ignore the returned SU_RESULT
 
-		// Inherit SketchUp material by default.
-		FMaterialIDType MaterialID = FMaterial::INHERITED_MATERIAL_ID;
+
+		int32 SlotId = 0; // Default material slot
 
 		// Get the SketckUp material ID.
 		if (bUseFrontMaterial)
@@ -251,7 +268,7 @@ namespace DatasmithSketchUp
 			if (SUIsValid(FrontMaterialRef))
 			{
 				// Get the front material ID of the SketckUp front material.
-				MaterialID = DatasmithSketchUpUtils::GetMaterialID(FrontMaterialRef);
+				SlotId = GetOrCreateSlotForMaterial(DatasmithSketchUpUtils::GetMaterialID(FrontMaterialRef));
 			}
 		}
 		else // bUseBackMaterial
@@ -259,17 +276,23 @@ namespace DatasmithSketchUp
 			if (SUIsValid(BackMaterialRef))
 			{
 				// Get the back material ID of the SketckUp back material.
-				MaterialID = DatasmithSketchUpUtils::GetMaterialID(BackMaterialRef);
+				SlotId = GetOrCreateSlotForMaterial(DatasmithSketchUpUtils::GetMaterialID(BackMaterialRef));
 			}
 		}
 
-		// Combine the mesh triangle material IDs.
-		TArray<FMaterialIDType> MaterialIDs;
-		MaterialIDs.Init(MaterialID, TriangleCount);
-		MeshTriangleMaterialIDs.Append(MaterialIDs);
 
-		// Add the material ID to the set of all the material IDs used by the combined mesh triangles.
-		MeshTriangleMaterialIDSet.Add(MaterialID);
+		if (SlotId == 0)
+		{
+			// todo: it's possible to skip adding slot=0 when there's no faces with 'default' material
+			// for this need to compute MeshTriangleSlotIds afterwards(when all materials are known)
+			bHasFacesWithDefaultMaterial = true;
+		}
+
+		MeshTriangleSlotIds.Reserve(MeshTriangleSlotIds.Num());
+		for(int32 TriangleIndex = 0; TriangleIndex < TriangleCount; ++TriangleIndex)
+		{
+			MeshTriangleSlotIds.Add(SlotId);
+		}
 	}
 
 	void FDatasmithSketchUpMesh::ConvertMeshToDatasmith(FDatasmithMesh& OutDMesh) const
@@ -316,7 +339,7 @@ namespace DatasmithSketchUp
 
 			// Set the triangle vertex indices in the exported Datasmith mesh.
 			SMeshTriangleIndices const& TriangleIndices = MeshTriangleIndices[TriangleNo];
-			OutDMesh.SetFace(TriangleNo, int32(TriangleIndices.IndexA), int32(TriangleIndices.IndexB), int32(TriangleIndices.IndexC), MeshTriangleMaterialIDs[TriangleNo].EntityID);
+			OutDMesh.SetFace(TriangleNo, int32(TriangleIndices.IndexA), int32(TriangleIndices.IndexB), int32(TriangleIndices.IndexC), MeshTriangleSlotIds[TriangleNo]);
 
 			// Set the triangle vertex normals in the exported Datasmith mesh.
 			SMeshTriangleNormals TriangleNormals = { MeshVertexNormals[TriangleIndices.IndexA],
@@ -390,15 +413,15 @@ void FEntities::UpdateGeometry(FExportContext& Context)
 
 			Mesh->DatasmithMesh->SetName(*MeshElementName);
 			Mesh->DatasmithMesh->SetLabel(*MeshLabel);
-			Mesh->bIsUsingInheritedMaterial = ExtractedMeshPtr->MeshTriangleMaterialIDSet.Contains(FMaterial::INHERITED_MATERIAL_ID);
+			Mesh->bIsUsingInheritedMaterial = ExtractedMeshPtr->bHasFacesWithDefaultMaterial;
 
 			// Add the non-inherited materials used by the combined mesh triangles.
-			for (FMaterialIDType MeshMaterialID : ExtractedMeshPtr->MeshTriangleMaterialIDSet)
+			for (int32 SlotId = 0;SlotId < ExtractedMeshPtr->MaterialIDForSlotId.Num(); ++SlotId)
 			{
+				FMaterialIDType MeshMaterialID = ExtractedMeshPtr->MaterialIDForSlotId[SlotId];
 				// Default or (somehow)missing materials are also assigned to mesh(as a default material)
 				if (FMaterialOccurrence* Material = Context.Materials.RegisterGeometry(MeshMaterialID, EntitiesGeometry.Get()))
 				{
-					int32 SlotId = MeshMaterialID.EntityID; // Triangles assigned Material Ids as SlotIds
 					Mesh->DatasmithMesh->SetMaterial(Material->GetName(), SlotId);
 				}
 			}
@@ -519,6 +542,7 @@ void ScanSketchUpEntitiesFaces(SUEntitiesRef EntitiesRef, FEntitiesGeometry& Geo
 		// Create a mesh combining the geometry of the SketchUp connected faces.
 		TSharedPtr<FDatasmithSketchUpMesh> ExtractedMeshPtr = MakeShared<FDatasmithSketchUpMesh>();
 		FDatasmithSketchUpMesh& ExtractedMesh = *ExtractedMeshPtr;
+		ExtractedMesh.GetOrCreateSlotForMaterial(FMaterial::INHERITED_MATERIAL_ID); // Add default material to Slot=0
 
 		// The source SketchUp face needs to be scanned once.
 		TArray<SUFaceRef> FacesToScan;
