Create exporter for Revit 2022

diff --git a/Engine/Source/Programs/Enterprise/Datasmith/DatasmithRevitExporter/DatasmithRevit2022/DatasmithDynamoNode.csproj b/Engine/Source/Programs/Enterprise/Datasmith/DatasmithRevitExporter/DatasmithRevit2022/DatasmithDynamoNode.csproj
new file mode 100644
--- /dev/null
+++ b/Engine/Source/Programs/Enterprise/Datasmith/DatasmithRevitExporter/DatasmithRevit2022/DatasmithDynamoNode.csproj
@@ -0,0 +1,72 @@
+﻿<?xml version="1.0" encoding="utf-8"?>
+<Project ToolsVersion="15.0" xmlns="http://schemas.microsoft.com/developer/msbuild/2003">
+  <Import Project="$(MSBuildExtensionsPath)\$(MSBuildToolsVersion)\Microsoft.Common.props" Condition="Exists('$(MSBuildExtensionsPath)\$(MSBuildToolsVersion)\Microsoft.Common.props')" />
+  <PropertyGroup>
+    <ResolveAssemblyWarnOrErrorOnTargetArchitectureMismatch>
+      None
+    </ResolveAssemblyWarnOrErrorOnTargetArchitectureMismatch>
+  </PropertyGroup>
+  <PropertyGroup>
+    <Configuration Condition=" '$(Configuration)' == '' ">Debug</Configuration>
+    <Platform Condition=" '$(Platform)' == '' ">AnyCPU</Platform>
+    <ProjectGuid>{FB926111-B972-4163-972B-FE58B2CC0AD1}</ProjectGuid>
+    <OutputType>Library</OutputType>
+    <AppDesignerFolder>Properties</AppDesignerFolder>
+    <RootNamespace>DatasmithDynamoNode</RootNamespace>
+    <AssemblyName>DatasmithDynamoNode</AssemblyName>
+    <TargetFrameworkVersion>v4.8</TargetFrameworkVersion>
+    <FileAlignment>512</FileAlignment>
+    <Deterministic>true</Deterministic>
+  </PropertyGroup>
+  <PropertyGroup Condition=" '$(Configuration)|$(Platform)' == 'Debug|AnyCPU' ">
+    <DebugSymbols>true</DebugSymbols>
+    <DebugType>full</DebugType>
+    <Optimize>false</Optimize>
+    <OutputPath>bin\Debug\</OutputPath>
+    <DefineConstants>TRACE;DEBUG;REVIT_API_2022</DefineConstants>
+    <ErrorReport>prompt</ErrorReport>
+    <WarningLevel>4</WarningLevel>
+  </PropertyGroup>
+  <PropertyGroup Condition=" '$(Configuration)|$(Platform)' == 'Release|AnyCPU' ">
+    <DebugType>pdbonly</DebugType>
+    <Optimize>true</Optimize>
+    <OutputPath>bin\Release\</OutputPath>
+    <DefineConstants>TRACE;REVIT_API_2022</DefineConstants>
+    <ErrorReport>prompt</ErrorReport>
+    <WarningLevel>4</WarningLevel>
+  </PropertyGroup>
+  <ItemGroup>
+    <Reference Include="RevitAPI">
+      <HintPath>$(Revit_2022_API)\RevitAPI.dll</HintPath>
+      <Private>False</Private>
+    </Reference>
+    <Reference Include="RevitAPIUI">
+      <HintPath>$(Revit_2022_API)\RevitAPIUI.dll</HintPath>
+      <Private>False</Private>
+    </Reference>
+    <Reference Include="System" />
+    <Reference Include="System.Core" />
+    <Reference Include="System.Xml.Linq" />
+    <Reference Include="System.Data.DataSetExtensions" />
+    <Reference Include="Microsoft.CSharp" />
+    <Reference Include="System.Data" />
+    <Reference Include="System.Net.Http" />
+    <Reference Include="System.Xml" />
+  </ItemGroup>
+  <ItemGroup>
+    <Compile Include="..\Private\DatasmithDynamoNode.cs" />
+    <Compile Include="Properties\AssemblyInfo.cs" />
+  </ItemGroup>
+  <ItemGroup>
+    <ProjectReference Include="DatasmithRevit2022.csproj">
+      <Project>{1851EF38-66E2-4749-9C23-DCCC4D71D9C5}</Project>
+      <Name>DatasmithRevit2022</Name>
+    </ProjectReference>
+  </ItemGroup>
+  <Import Project="$(MSBuildToolsPath)\Microsoft.CSharp.targets" />
+  <PropertyGroup>
+    <PostBuildEvent>set PUBLISH_DIR=$(ProjectDir)..\..\..\..\..\..\Binaries\Win64\Revit\2022\
+xcopy "$(ProjectDir)$(OutputPath)*.dll" "%25PUBLISH_DIR%25" /Y
+</PostBuildEvent>
+  </PropertyGroup>
+</Project>
\ No newline at end of file
diff --git a/Engine/Source/Programs/Enterprise/Datasmith/DatasmithRevitExporter/DatasmithRevit2022/DatasmithRevit2022.csproj b/Engine/Source/Programs/Enterprise/Datasmith/DatasmithRevitExporter/DatasmithRevit2022/DatasmithRevit2022.csproj
new file mode 100644
--- /dev/null
+++ b/Engine/Source/Programs/Enterprise/Datasmith/DatasmithRevitExporter/DatasmithRevit2022/DatasmithRevit2022.csproj
@@ -0,0 +1,122 @@
+﻿<?xml version="1.0" encoding="utf-8"?>
+<Project ToolsVersion="14.0" DefaultTargets="Build" xmlns="http://schemas.microsoft.com/developer/msbuild/2003">
+  <Import Project="$(MSBuildExtensionsPath)\$(MSBuildToolsVersion)\Microsoft.Common.props" Condition="Exists('$(MSBuildExtensionsPath)\$(MSBuildToolsVersion)\Microsoft.Common.props')" />
+  <PropertyGroup>
+    <ResolveAssemblyWarnOrErrorOnTargetArchitectureMismatch>
+      None
+    </ResolveAssemblyWarnOrErrorOnTargetArchitectureMismatch>
+  </PropertyGroup>
+  <PropertyGroup>
+    <Configuration Condition=" '$(Configuration)' == '' ">Debug</Configuration>
+    <Platform Condition=" '$(Platform)' == '' ">AnyCPU</Platform>
+    <ProjectGuid>{1851EF38-66E2-4749-9C23-DCCC4D71D9C5}</ProjectGuid>
+    <OutputType>Library</OutputType>
+    <AppDesignerFolder>Properties</AppDesignerFolder>
+    <RootNamespace>DatasmithRevit2022</RootNamespace>
+    <AssemblyName>DatasmithRevit2022</AssemblyName>
+    <TargetFrameworkVersion>v4.8</TargetFrameworkVersion>
+    <FileAlignment>512</FileAlignment>
+  </PropertyGroup>
+  <PropertyGroup Condition=" '$(Configuration)|$(Platform)' == 'Debug|AnyCPU' ">
+    <DebugSymbols>true</DebugSymbols>
+    <DebugType>full</DebugType>
+    <Optimize>false</Optimize>
+    <OutputPath>bin\Debug\</OutputPath>
+    <DefineConstants>DEBUG;TRACE;REVIT_API_2022</DefineConstants>
+    <ErrorReport>prompt</ErrorReport>
+    <WarningLevel>4</WarningLevel>
+    <StartAction>Program</StartAction>
+    <StartProgram>$(ProgramW6432)\Autodesk\Revit 2022\Revit.exe</StartProgram>
+    <EnableUnmanagedDebugging>true</EnableUnmanagedDebugging>
+  </PropertyGroup>
+  <PropertyGroup Condition=" '$(Configuration)|$(Platform)' == 'Release|AnyCPU' ">
+    <DebugType>pdbonly</DebugType>
+    <Optimize>true</Optimize>
+    <OutputPath>bin\Release\</OutputPath>
+    <DefineConstants>TRACE;REVIT_API_2022</DefineConstants>
+    <ErrorReport>prompt</ErrorReport>
+    <WarningLevel>4</WarningLevel>
+    <StartAction>Program</StartAction>
+    <StartProgram>$(ProgramW6432)\Autodesk\Revit 2022\Revit.exe</StartProgram>
+  </PropertyGroup>
+  <PropertyGroup>
+    <RunPostBuildEvent>OnOutputUpdated</RunPostBuildEvent>
+  </PropertyGroup>
+  <ItemGroup>
+    <Reference Include="PresentationCore" />
+    <Reference Include="RevitAPI">
+      <HintPath>$(Revit_2022_API)\RevitAPI.dll</HintPath>
+      <Private>False</Private>
+    </Reference>
+    <Reference Include="RevitAPIUI">
+      <HintPath>$(Revit_2022_API)\RevitAPIUI.dll</HintPath>
+      <Private>False</Private>
+    </Reference>
+    <Reference Include="System" />
+    <Reference Include="System.Drawing" />
+    <Reference Include="System.Windows.Forms" />
+  </ItemGroup>
+  <ItemGroup>
+    <Compile Include="Properties\AssemblyInfo.cs" />
+    <Compile Include="..\Private\DatasmithRevitDocumentData.cs">
+      <Link>DatasmithRevitExporter\DatasmithRevitDocumentData.cs</Link>
+    </Compile>
+    <Compile Include="..\Private\DatasmithRevitDirectLink.cs">
+      <Link>DatasmithRevitExporter\DatasmithRevitDirectLink.cs</Link>
+    </Compile>
+    <Compile Include="..\Private\DatasmithRevitApplication.cs">
+      <Link>DatasmithRevitExporter\DatasmithRevitApplication.cs</Link>
+    </Compile>
+    <Compile Include="..\Private\DatasmithRevitCommand.cs">
+      <Link>DatasmithRevitExporter\DatasmithRevitCommand.cs</Link>
+    </Compile>
+    <Compile Include="..\Private\DatasmithRevitExportContext.cs">
+      <Link>DatasmithRevitExporter\DatasmithRevitExportContext.cs</Link>
+    </Compile>
+    <Compile Include="..\Private\DatasmithRevitExportMessages.cs">
+      <Link>DatasmithRevitExporter\DatasmithRevitExportMessages.cs</Link>
+      <SubType>Form</SubType>
+    </Compile>
+    <Compile Include="..\Private\DatasmithRevitExportOptions.cs">
+      <Link>DatasmithRevitExporter\DatasmithRevitExportOptions.cs</Link>
+      <SubType>Form</SubType>
+    </Compile>
+    <Compile Include="..\Private\DatasmithRevitLight.cs">
+      <Link>DatasmithRevitExporter\DatasmithRevitLight.cs</Link>
+    </Compile>
+    <Compile Include="..\Private\DatasmithRevitMaterial.cs">
+      <Link>DatasmithRevitExporter\DatasmithRevitMaterial.cs</Link>
+    </Compile>
+    <Compile Include="..\..\..\..\..\..\Binaries\Win64\DatasmithFacadeCSharp\Public\*.cs">
+      <Link>DatasmithSwigFacade\%(Filename)%(Extension)</Link>
+    </Compile>
+    <Content Include="DatasmithRevit2022.addin" />
+  </ItemGroup>
+  <ItemGroup>
+    <ProjectReference Include="..\Resources\DatasmithRevitResources.csproj">
+      <Project>{bc17d734-8f8f-44c1-a84e-af0e616fbe9e}</Project>
+      <Name>DatasmithRevitResources</Name>
+    </ProjectReference>
+  </ItemGroup>
+  <Import Project="$(MSBuildToolsPath)\Microsoft.CSharp.targets" />
+  <PropertyGroup>
+    <PostBuildEvent>
+      set PUBLISH_DIR=$(ProjectDir)..\..\..\..\..\..\Binaries\Win64\Revit\2022\
+      xcopy "$(ProjectDir)$(TargetName).addin" "%25PUBLISH_DIR%25" /Y
+      xcopy "$(ProjectDir)$(OutputPath)*.dll" "%25PUBLISH_DIR%25" /Y
+      xcopy "$(ProjectDir)..\Resources\Windows\*.png" "%25PUBLISH_DIR%25" /Y
+      xcopy "$(ProjectDir)..\Resources\Windows\DatasmithBatchExporter.dyn" "%25PUBLISH_DIR%25" /Y
+</PostBuildEvent>
+  </PropertyGroup>
+  <!-- To modify your build process, add your task inside one of the targets below and uncomment it.
+       Other similar extension points exist, see Microsoft.Common.targets.
+  <Target Name="BeforeBuild">
+  </Target>
+  <Target Name="AfterBuild">
+  </Target>
+  -->
+  <Target Name="AfterClean">
+    <RemoveDir Directories="DatasmithRevitExporter" Condition="Exists('DatasmithRevitExporter')" />
+    <RemoveDir Directories="DatasmithSwigFacade" Condition="Exists('DatasmithSwigFacade')" />
+  </Target>
+</Project>
\ No newline at end of file
diff --git a/Engine/Source/Programs/Enterprise/Datasmith/DatasmithRevitExporter/DatasmithRevit2022/Properties/AssemblyInfo.cs b/Engine/Source/Programs/Enterprise/Datasmith/DatasmithRevitExporter/DatasmithRevit2022/Properties/AssemblyInfo.cs
new file mode 100644
--- /dev/null
+++ b/Engine/Source/Programs/Enterprise/Datasmith/DatasmithRevitExporter/DatasmithRevit2022/Properties/AssemblyInfo.cs
@@ -0,0 +1,31 @@
+// Copyright Epic Games, Inc. All Rights Reserved.
+using System.Resources;
+
+using System.Reflection;
+using System.Runtime.CompilerServices;
+using System.Runtime.InteropServices;
+
+// General Information about an assembly is controlled through the following set of attributes.
+// Change these attribute values to modify the information associated with an assembly.
+[assembly: AssemblyTitle("DatasmithRevit2022")]
+[assembly: AssemblyDescription("Export a Revit 2022 3D View to Unreal Datasmith")]
+[assembly: AssemblyConfiguration("")]
+[assembly: AssemblyCompany("Epic Games, Inc.")]
+[assembly: AssemblyProduct("DatasmithRevitExporter Add-In for Revit 2022")]
+[assembly: AssemblyCopyright("Copyright Epic Games, Inc. All Rights Reserved.")]
+[assembly: AssemblyTrademark("")]
+[assembly: AssemblyCulture("")]
+
+// Setting ComVisible to false makes the types in this assembly not visible to COM components.
+// If you need to access a type in this assembly from COM, set the ComVisible attribute to true on that type.
+[assembly: ComVisible(false)]
+
+// The following GUID is for the ID of the typelib if this project is exposed to COM
+[assembly: Guid("241EFA34-AF10-4FDD-B11D-E32B7C9A229C")]
+
+// Version information for an assembly consists of the following four values: <major version>.<minor version>.<build number>.<revision number>
+// You can specify all the values or you can default the build number and revision number by using the '*': [assembly: AssemblyVersion("1.0.*")]
+[assembly: AssemblyVersion("4.21.0.0")]
+[assembly: AssemblyFileVersion("4.21.0.0")]
+[assembly: NeutralResourcesLanguage("en-US")]
+
diff --git a/Engine/Source/Programs/Enterprise/Datasmith/DatasmithRevitExporter/Private/DatasmithRevitMaterial.cs b/Engine/Source/Programs/Enterprise/Datasmith/DatasmithRevitExporter/Private/DatasmithRevitMaterial.cs
--- a/Engine/Source/Programs/Enterprise/Datasmith/DatasmithRevitExporter/Private/DatasmithRevitMaterial.cs
+++ b/Engine/Source/Programs/Enterprise/Datasmith/DatasmithRevitExporter/Private/DatasmithRevitMaterial.cs
@@ -1889,7 +1889,7 @@ namespace DatasmithRevitExporter
 
 					if (distanceProperty != null && distanceProperty.Value != 0.0)
 					{
-#if REVIT_API_2021
+#if REVIT_API_2021 || REVIT_API_2022
 						return (float) UnitUtils.Convert(distanceProperty.Value, distanceProperty.GetUnitTypeId(), UnitTypeId.Feet);
 #else
 						return (float) UnitUtils.Convert(distanceProperty.Value, distanceProperty.DisplayUnitType, DisplayUnitType.DUT_DECIMAL_FEET);
diff --git a/Engine/Source/Programs/Enterprise/Datasmith/DatasmithRevitExporter/ReadMe.txt b/Engine/Source/Programs/Enterprise/Datasmith/DatasmithRevitExporter/ReadMe.txt
--- a/Engine/Source/Programs/Enterprise/Datasmith/DatasmithRevitExporter/ReadMe.txt
+++ b/Engine/Source/Programs/Enterprise/Datasmith/DatasmithRevitExporter/ReadMe.txt
@@ -14,4 +14,4 @@ on the developer's workstation.
 Inside of Epic Games, when opening the Visual Studio C# project directly, environment variable Revit_<year>_API must be set
 to the developer's Perforce workspace directory Enterprise\Source\ThirdParty\NotForLicensees\Revit\Revit_<year>_API.
 
-At this time, <year> can be in the range [2018..2021].
+At this time, <year> can be in the range [2018..2022].
