Fixed build breakage for compilation of  Revit exporter

diff --git a/Engine/Source/Programs/Enterprise/Datasmith/DatasmithRevitExporter/DatasmithRevit2022.Build.cs b/Engine/Source/Programs/Enterprise/Datasmith/DatasmithRevitExporter/DatasmithRevit2022.Build.cs
new file mode 100644
--- /dev/null
+++ b/Engine/Source/Programs/Enterprise/Datasmith/DatasmithRevitExporter/DatasmithRevit2022.Build.cs
@@ -0,0 +1,20 @@
+// Copyright Epic Games, Inc. All Rights Reserved.
+
+using System.IO;
+
+namespace UnrealBuildTool.Rules
+{
+	[SupportedPlatforms("Win64")]
+	public class DatasmithRevit2022 : DatasmithRevitBase
+	{
+		public DatasmithRevit2022(ReadOnlyTargetRules Target)
+			: base(Target)
+		{
+		}
+		
+		public override string GetRevitVersion()
+		{
+			return "2022";
+		}
+	}
+}
diff --git a/Engine/Source/Programs/Enterprise/Datasmith/DatasmithRevitExporter/DatasmithRevit2022.Target.cs b/Engine/Source/Programs/Enterprise/Datasmith/DatasmithRevitExporter/DatasmithRevit2022.Target.cs
new file mode 100644
--- /dev/null
+++ b/Engine/Source/Programs/Enterprise/Datasmith/DatasmithRevitExporter/DatasmithRevit2022.Target.cs
@@ -0,0 +1,14 @@
+﻿// Copyright Epic Games, Inc. All Rights Reserved.
+
+using UnrealBuildTool;
+using System.IO;
+
+public class DatasmithRevit2022Target : DatasmithRevitBaseTarget
+{
+	public DatasmithRevit2022Target(TargetInfo Target)
+		: base(Target)
+	{
+	}
+
+	public override string GetVersion() { return "2022"; }
+}
