DatasmithSDK - don't build by default with the solution

diff --git a/Engine/Source/Programs/Enterprise/Datasmith/DatasmithSDK/DatasmithSDK.Target.cs b/Engine/Source/Programs/Enterprise/Datasmith/DatasmithSDK/DatasmithSDK.Target.cs
--- a/Engine/Source/Programs/Enterprise/Datasmith/DatasmithSDK/DatasmithSDK.Target.cs
+++ b/Engine/Source/Programs/Enterprise/Datasmith/DatasmithSDK/DatasmithSDK.Target.cs
@@ -10,6 +10,7 @@ public class DatasmithSDKTarget : TargetRules
 	{
 		Type = TargetType.Program;
 		SolutionDirectory = "Programs/Datasmith";
+		bBuildInSolutionByDefault = false;
 
 		LaunchModuleName = "DatasmithSDK";
 		ExeBinariesSubFolder = "DatasmithSDK";
