DatasmithSDK - support spaces in project paths

diff --git a/Engine/Source/Programs/Enterprise/Datasmith/DatasmithSDK/DatasmithSDK.Target.cs b/Engine/Source/Programs/Enterprise/Datasmith/DatasmithSDK/DatasmithSDK.Target.cs
--- a/Engine/Source/Programs/Enterprise/Datasmith/DatasmithSDK/DatasmithSDK.Target.cs
+++ b/Engine/Source/Programs/Enterprise/Datasmith/DatasmithSDK/DatasmithSDK.Target.cs
@@ -41,7 +41,7 @@ public class DatasmithSDKTarget : TargetRules
 	public void PostBuildCopy(string SrcPath, string DestPath)
 	{
 		PostBuildSteps.Add(string.Format("echo Copying {0} to {1}", SrcPath, DestPath));
-		PostBuildSteps.Add(string.Format("xcopy {0} {1} /R /Y /S", SrcPath, DestPath));
+		PostBuildSteps.Add(string.Format("xcopy \"{0}\" \"{1}\" /R /Y /S", SrcPath, DestPath));
 	}
 
 	public void AddWindowsPostBuildSteps()
