3ds Max DirectLink: AutoSync(prototype)

- toggle AutoSync from action/toolbar
- every 500 ms check that user was idle and in that case Update(with progress bar not shown)

diff --git a/Engine/Source/Programs/Enterprise/Datasmith/DatasmithMaxExporter/Private/DatasmithMaxDirectLinkMaxScript.cpp b/Engine/Source/Programs/Enterprise/Datasmith/DatasmithMaxExporter/Private/DatasmithMaxDirectLinkMaxScript.cpp
--- a/Engine/Source/Programs/Enterprise/Datasmith/DatasmithMaxExporter/Private/DatasmithMaxDirectLinkMaxScript.cpp
+++ b/Engine/Source/Programs/Enterprise/Datasmith/DatasmithMaxExporter/Private/DatasmithMaxDirectLinkMaxScript.cpp
@@ -22,6 +22,7 @@ MAX_INCLUDES_START
 	#include "maxscript/foundation/numbers.h"
 	#include "maxscript/foundation/arrays.h"
 	#include "maxscript\macros\define_instantiation_functions.h"
+
 MAX_INCLUDES_END
 
 extern HINSTANCE HInstanceMax;
@@ -177,6 +178,19 @@ Value* DirectLinkUpdateScene_cf(Value** arg_list, int count)
 }
 Primitive DirectLinkUpdateScene_pf(_M("Datasmith_DirectLinkUpdateScene"), DirectLinkUpdateScene_cf);
 
+Value* ToggleAutoSync_cf(Value** arg_list, int count) 
+{
+	check_arg_count(ToggleAutoSync, 0, count);
+
+	if (GetExporter())
+	{
+		return bool_result(GetExporter()->ToggleAutoSync());
+	}
+
+	return &false_value;
+}
+Primitive ToggleAutoSync_pf(_M("Datasmith_ToggleAutoSync"), ToggleAutoSync_cf);
+
 Value* OpenDirectlinkUi_cf(Value** arg_list, int count) 
 {
 	check_arg_count(OpenDirectlinkUi, 0, count);
@@ -271,7 +285,7 @@ public:
 			{
 				if (GetExporter())
 				{
-					// todo: GetExporter()->ToggleAutoSync();
+					GetExporter()->ToggleAutoSync();
 				}
 				return true;
 			}
diff --git a/Engine/Source/Programs/Enterprise/Datasmith/DatasmithMaxExporter/Private/DatasmithMaxExporterWithDirectLink.cpp b/Engine/Source/Programs/Enterprise/Datasmith/DatasmithMaxExporter/Private/DatasmithMaxExporterWithDirectLink.cpp
--- a/Engine/Source/Programs/Enterprise/Datasmith/DatasmithMaxExporter/Private/DatasmithMaxExporterWithDirectLink.cpp
+++ b/Engine/Source/Programs/Enterprise/Datasmith/DatasmithMaxExporter/Private/DatasmithMaxExporterWithDirectLink.cpp
@@ -1685,6 +1685,51 @@ public:
 		DirectLinkImpl->UpdateScene(ExportedScene.GetDatasmithScene());
 	}
 
+	static VOID AutoSyncTimerProc(HWND, UINT, UINT_PTR TimerIdentifier, DWORD)
+	{
+		reinterpret_cast<FExporter*>(TimerIdentifier)->UpdateAutoSync(); 
+	}
+
+	// Update is user was idle for some time
+	void UpdateAutoSync()
+	{
+		LASTINPUTINFO LastInputInfo;
+		LastInputInfo.cbSize = sizeof(LASTINPUTINFO);
+		LastInputInfo.dwTime = 0;
+		if (GetLastInputInfo(&LastInputInfo))
+		{
+			DWORD CurrentTime = GetTickCount();
+			DWORD IdlePeriod = GetTickCount() - LastInputInfo.dwTime;
+			LogDebug(FString::Printf(TEXT("CurrentTime: %ld, Idle time: %ld, IdlePeriod: %ld"), CurrentTime, LastInputInfo.dwTime, IdlePeriod));
+
+			if (IdlePeriod > 500)
+			{
+				// Don't create progress bar for autosync - it steals focus, closes listener and what else
+				// todo: consider creating progress when a big change in scene is detected, e.g. number of nodes?
+				UpdateScene(true);
+				UpdateDirectLinkScene();
+			}
+		}
+	}
+
+	bool bAutoSyncEnabled = false;
+
+	bool ToggleAutoSync()
+	{
+		if (bAutoSyncEnabled)
+		{
+			KillTimer(GetCOREInterface()->GetMAXHWnd(), reinterpret_cast<UINT_PTR>(this));
+		}
+		else
+		{
+			SetTimer(GetCOREInterface()->GetMAXHWnd(), reinterpret_cast<UINT_PTR>(this), 500, AutoSyncTimerProc);
+		}
+		bAutoSyncEnabled = !bAutoSyncEnabled;
+
+		LogDebug(bAutoSyncEnabled ? TEXT("AutoSync ON") : TEXT("AutoSync OFF"));
+		return bAutoSyncEnabled;
+	}
+
 	FNotifications NotificationsHandler;
 
 	// Install change notification systems
diff --git a/Engine/Source/Programs/Enterprise/Datasmith/DatasmithMaxExporter/Public/DatasmithMaxDirectLink.h b/Engine/Source/Programs/Enterprise/Datasmith/DatasmithMaxExporter/Public/DatasmithMaxDirectLink.h
--- a/Engine/Source/Programs/Enterprise/Datasmith/DatasmithMaxExporter/Public/DatasmithMaxDirectLink.h
+++ b/Engine/Source/Programs/Enterprise/Datasmith/DatasmithMaxExporter/Public/DatasmithMaxDirectLink.h
@@ -58,6 +58,7 @@ class IExporter
 	// DirectLink
 	virtual void InitializeDirectLinkForScene() = 0;
 	virtual void UpdateDirectLinkScene() = 0;
+	virtual bool ToggleAutoSync() = 0;
 };
 
 
