3ds Max DirectLink: Animation export

diff --git a/Engine/Source/Programs/Enterprise/Datasmith/DatasmithMaxExporter/Private/DatasmithMaxExporterWithDirectLink.cpp b/Engine/Source/Programs/Enterprise/Datasmith/DatasmithMaxExporter/Private/DatasmithMaxExporterWithDirectLink.cpp
--- a/Engine/Source/Programs/Enterprise/Datasmith/DatasmithMaxExporter/Private/DatasmithMaxExporterWithDirectLink.cpp
+++ b/Engine/Source/Programs/Enterprise/Datasmith/DatasmithMaxExporter/Private/DatasmithMaxExporterWithDirectLink.cpp
@@ -676,6 +676,39 @@ public:
 		LogDebug("Scene update: done");
 	}
 
+	void ExportAnimations()
+	{
+		FDatasmithConverter Converter;
+		// Use the same name for the unique level sequence as the scene name
+		TSharedRef<IDatasmithLevelSequenceElement> LevelSequence = FDatasmithSceneFactory::CreateLevelSequence(ExportedScene.GetDatasmithScene()->GetName());
+		LevelSequence->SetFrameRate(GetFrameRate());
+
+		for (TPair<FNodeKey, FNodeTrackerHandle> NodeKeyAndNodeTracker: NodeTrackers)
+		{
+			FNodeTracker* NodeTracker = NodeKeyAndNodeTracker.Value.GetNodeTracker();
+			
+			if (NodeTracker->DatasmithActorElement)
+			{
+				
+				if (Lights.Contains(NodeTracker))
+				{
+					const TSharedPtr<IDatasmithLightActorElement> LightElement = StaticCastSharedPtr< IDatasmithLightActorElement >(NodeTracker->DatasmithActorElement);
+					const FMaxLightCoordinateConversionParams LightParams(NodeTracker->Node,
+						LightElement->IsA(EDatasmithElementType::AreaLight) ? StaticCastSharedPtr<IDatasmithAreaLightElement>(LightElement)->GetLightShape() : EDatasmithLightShape::None);
+					FDatasmithMaxSceneExporter::ExportAnimation(LevelSequence, NodeTracker->Node, NodeTracker->DatasmithActorElement->GetName(), Converter.UnitToCentimeter, LightParams);
+				}
+				else
+				{
+					FDatasmithMaxSceneExporter::ExportAnimation(LevelSequence, NodeTracker->Node, NodeTracker->DatasmithActorElement->GetName(), Converter.UnitToCentimeter);
+				}
+			}
+		}
+		if (LevelSequence->GetAnimationsCount() > 0)
+		{
+			ExportedScene.GetDatasmithScene()->AddLevelSequence(LevelSequence);
+		}
+	}
+
 	FORCENOINLINE
 	FNodeTrackerHandle& AddNode(FNodeKey NodeKey, INode* Node)
 	{
@@ -1670,6 +1703,7 @@ public:
 	{
 		SceneTracker.ParseScene();
 		SceneTracker.Update(bQuiet);
+		SceneTracker.ExportAnimations();
 		ExportedScene.GetSceneExporter().Export(ExportedScene.GetDatasmithScene(), false);
 		return true;
 	}
