Datasmith Solidworks: Missing materials from solidworks exporter

diff --git a/Engine/Source/Programs/Enterprise/Datasmith/DatasmithSolidworksExporter/Private/AssemblyDocument.cs b/Engine/Source/Programs/Enterprise/Datasmith/DatasmithSolidworksExporter/Private/AssemblyDocument.cs
--- a/Engine/Source/Programs/Enterprise/Datasmith/DatasmithSolidworksExporter/Private/AssemblyDocument.cs
+++ b/Engine/Source/Programs/Enterprise/Datasmith/DatasmithSolidworksExporter/Private/AssemblyDocument.cs
@@ -21,19 +21,21 @@ namespace DatasmithSolidworks
 			Delete
 		};
 
-		class FSyncState
+		public class FSyncState
 		{
 			public Dictionary<string, FPartDocument> PartsMap = new Dictionary<string, FPartDocument>();
 			public ConcurrentDictionary<string, FObjectMaterials> ComponentsMaterialsMap = new ConcurrentDictionary<string, FObjectMaterials>();
 			public Dictionary<string, string> ComponentToPartMap = new Dictionary<string, string>();
+			public Dictionary<string, float[]> ComponentsTransformsMap = new Dictionary<string, float[]>();
 			public HashSet<string> CleanComponents = new HashSet<string>();
 			public Dictionary<string, uint> DirtyComponents = new Dictionary<string, uint>();
 			public HashSet<string> ComponentsToDelete = new HashSet<string>();
+			public Dictionary<string, Component2> ExportedComponentsMap = new Dictionary<string, Component2>();
 		}
 
 		public AssemblyDoc SwAsmDoc { get; private set; } = null;
 
-		private FSyncState SyncState = new FSyncState();
+		public FSyncState SyncState { get; private set; } = new FSyncState();
 
 		public FAssemblyDocument(int InDocId, AssemblyDoc InSwDoc, FDatasmithExporter InExporter) : base(InDocId, InSwDoc as ModelDoc2, InExporter)
 		{
@@ -54,6 +56,9 @@ namespace DatasmithSolidworks
 			{
 				string ActorName = FDatasmithExporter.SanitizeName(CompName);
 				SyncState.ComponentToPartMap.Remove(CompName);
+				SyncState.ExportedComponentsMap.Remove(CompName);
+				SyncState.ComponentsMaterialsMap.TryRemove(CompName, out FObjectMaterials _);
+				SyncState.ComponentsTransformsMap.Remove(CompName);
 				Exporter.RemoveActor(ActorName);
 			}
 
@@ -78,9 +83,18 @@ namespace DatasmithSolidworks
 				{
 					ComponentNamesToExportSet.Add(KVP.Key.Name);
 				}
-				
 			}
-			SyncState.ComponentsMaterialsMap = FObjectMaterials.LoadAssemblyMaterials(this, ComponentNamesToExportSet, swDisplayStateOpts_e.swThisDisplayState, null);
+
+			ConcurrentDictionary<string, FObjectMaterials> ModifiedComponentsMaterials = FObjectMaterials.LoadAssemblyMaterials(this, ComponentNamesToExportSet, swDisplayStateOpts_e.swThisDisplayState, null);
+
+			if (ModifiedComponentsMaterials != null)
+			{
+				foreach (var MatKVP in ModifiedComponentsMaterials)
+				{
+					SyncState.ComponentsMaterialsMap[MatKVP.Key] = MatKVP.Value;
+				}
+			}
+
 			Exporter.ExportMaterials(ExportedMaterialsMap);
 
 			// Export meshes
@@ -187,19 +201,47 @@ namespace DatasmithSolidworks
 			}
 		}
 
+		private float[] GetComponentDatasmithTransform(Component2 InComponent)
+		{
+			MathTransform ComponentTransform = InComponent.GetTotalTransform(true);
+			if (ComponentTransform == null)
+			{
+				ComponentTransform = InComponent.Transform2;
+			}
+
+			float[] DatasmithTransform = null;
+			if (ComponentTransform != null)
+			{
+				DatasmithTransform = MathUtils.ConvertFromSolidworksTransform(ComponentTransform, 100f/*GeomScale*/);
+			}
+			else
+			{
+				DatasmithTransform = new float[]
+				{
+					1.0f, 0.0f, 0.0f, 0.0f,
+					0.0f, 1.0f, 0.0f, 0.0f,
+					0.0f, 0.0f, 1.0f, 0.0f,
+					0.0f, 0.0f, 0.0f, 1.0f
+				};
+			}
+
+			return DatasmithTransform;
+		}
+
 		private void ExportComponentRecursive(Component2 InComponent, Component2 InParent, ref Dictionary<Component2, string> OutMeshesToExportMap)
 		{
-			if (!SyncState.CleanComponents.Contains(InComponent.Name2))
+			bool bHasDirtyTransform = false;
+
+			if (SyncState.ComponentsTransformsMap.ContainsKey(InComponent.Name2))
 			{
-				SetExportStatus(InComponent.Name2);
+				float[] ComponentTm = GetComponentDatasmithTransform(InComponent);
+				bHasDirtyTransform = !MathUtils.TransformsAreEqual(SyncState.ComponentsTransformsMap[InComponent.Name2], ComponentTm);
+			}
 
-				MathTransform ComponentTransform = InComponent.GetTotalTransform(true);
+			if (bHasDirtyTransform || !SyncState.CleanComponents.Contains(InComponent.Name2))
+			{
+				SetExportStatus(InComponent.Name2);
 
-				if (ComponentTransform == null)
-				{
-					ComponentTransform = InComponent.Transform2;
-				}
-			
 				FDatasmithActorExportInfo ActorExportInfo = new FDatasmithActorExportInfo();
 
 				string ComponentName = FDatasmithExporter.SanitizeName(InComponent.Name2);
@@ -210,11 +252,9 @@ namespace DatasmithSolidworks
 				ActorExportInfo.ParentName = InParent?.Name2;
 				ActorExportInfo.bVisible = true;
 				ActorExportInfo.Type = Exporter.GetExportedActorType(ComponentName) ?? EActorType.SimpleActor;
+				ActorExportInfo.Transform = GetComponentDatasmithTransform(InComponent);
 
-				if (ComponentTransform != null)
-				{
-					ActorExportInfo.Transform = MathUtils.ConvertFromSolidworksTransform(ComponentTransform, 100f/*GeomScale*/);
-				}
+				SyncState.ComponentsTransformsMap[InComponent.Name2] = ActorExportInfo.Transform;
 
 				if (!InComponent.IsSuppressed())
 				{
@@ -233,15 +273,21 @@ namespace DatasmithSolidworks
 					ActorExportInfo.bVisible = false;
 				}
 
-				bool bNeedsGeometryExport = !InComponent.IsSuppressed() && (InComponent.GetModelDoc2() is PartDoc);
+				bool bNeedsGeometryExport = false;
 
-				if (bNeedsGeometryExport && SyncState.DirtyComponents.ContainsKey(InComponent.Name2))
+				if (InComponent.GetModelDoc2() is PartDoc) 
 				{
-					uint DirtyState = SyncState.DirtyComponents[InComponent.Name2];
-					bNeedsGeometryExport = 
-						((DirtyState & (1u << (int)EComponentDirtyState.Material)) != 0) || 
-						((DirtyState & (1u << (int)EComponentDirtyState.Geometry)) != 0) || 
-						((DirtyState & (1u << (int)EComponentDirtyState.Delete)) != 0); 
+					bool bFirstExport = !SyncState.ExportedComponentsMap.ContainsKey(InComponent.Name2);
+					bNeedsGeometryExport = bFirstExport;
+
+					if (!bFirstExport && SyncState.DirtyComponents.ContainsKey(InComponent.Name2))
+					{
+						uint DirtyState = SyncState.DirtyComponents[InComponent.Name2];
+						bNeedsGeometryExport =
+							((DirtyState & (1u << (int)EComponentDirtyState.Material)) != 0) ||
+							((DirtyState & (1u << (int)EComponentDirtyState.Geometry)) != 0) ||
+							((DirtyState & (1u << (int)EComponentDirtyState.Delete)) != 0);
+					}
 				}
 
 				if (bNeedsGeometryExport)
@@ -271,6 +317,8 @@ namespace DatasmithSolidworks
 				SyncState.CleanComponents.Add(InComponent.Name2);
 			}
 
+			SyncState.ExportedComponentsMap[InComponent.Name2] = InComponent;
+
 			// Export component children
 			object[] Children = (object[])InComponent.GetChildren();
 
@@ -454,9 +502,6 @@ namespace DatasmithSolidworks
 
 		int OnConfigurationChangeNotify(string ConfigurationName, object Object, int ObjectType, int changeType)
 		{
-			SetDirty(true);
-			SyncState.CleanComponents?.Clear();
-			SyncState.DirtyComponents?.Clear();
 			return 0;
 		}
 
@@ -474,6 +519,23 @@ namespace DatasmithSolidworks
 
 		int OnUndoPostNotify()
 		{
+			// Check each exported component's transform for changes, since 
+			// this callback does not tell us what changed (and there's no other way to know that)!
+			foreach (var KVP in SyncState.ExportedComponentsMap)
+			{
+				Component2 Comp = KVP.Value;
+				float[] PrevCompTransform = null;
+				if (SyncState.ComponentsTransformsMap.TryGetValue(Comp.Name2, out PrevCompTransform))
+				{
+					float[] CompTransform = GetComponentDatasmithTransform(Comp);
+
+					if (!MathUtils.TransformsAreEqual(CompTransform, PrevCompTransform))
+					{
+						SetDirty(true);
+						break;
+					}
+				}
+			}
 			return 0;
 		}
 
@@ -561,6 +623,10 @@ namespace DatasmithSolidworks
 
 		int OnActiveConfigChangeNotify()
 		{
+			SetDirty(true);
+			SyncState.CleanComponents?.Clear();
+			SyncState.DirtyComponents?.Clear();
+			SyncState.ExportedComponentsMap?.Clear();
 			return 0;
 		}
 
diff --git a/Engine/Source/Programs/Enterprise/Datasmith/DatasmithSolidworksExporter/Private/MathUtils.cs b/Engine/Source/Programs/Enterprise/Datasmith/DatasmithSolidworksExporter/Private/MathUtils.cs
--- a/Engine/Source/Programs/Enterprise/Datasmith/DatasmithSolidworksExporter/Private/MathUtils.cs
+++ b/Engine/Source/Programs/Enterprise/Datasmith/DatasmithSolidworksExporter/Private/MathUtils.cs
@@ -2,6 +2,7 @@
 
 using SolidWorks.Interop.sldworks;
 using System;
+using System.Diagnostics;
 using System.Runtime.InteropServices;
 
 namespace DatasmithSolidworks
@@ -14,7 +15,7 @@ namespace DatasmithSolidworks
 
 		public static bool Equals(float F1, float F2)
 		{
-			return (Math.Abs(F2 - F1) <= 0.00000001f);
+			return (Math.Abs(F2 - F1) <= 0.000001f);
 		}
 
 		public static bool Equals(double A, float B)
@@ -77,6 +78,21 @@ namespace DatasmithSolidworks
 			return Result;
 		}
 
+		public static bool TransformsAreEqual(float[] InTransformA, float[] InTransformB)
+		{
+			Debug.Assert(InTransformA.Length == 16 && InTransformA.Length == InTransformB.Length);
+
+			for (int Idx = 0; Idx < 16; ++Idx)
+			{
+				if (!Equals(InTransformA[Idx], InTransformB[Idx]))
+				{
+					return false;
+				}
+			}
+
+			return true;
+		}
+
 		public static MathVector CreateVector(double InX, double InY, double InZ)
 		{
 			MathUtility MUtil = Addin.Instance.SolidworksApp.IGetMathUtility();
diff --git a/Engine/Source/Programs/Enterprise/Datasmith/DatasmithSolidworksExporter/Private/ObjectMaterials.cs b/Engine/Source/Programs/Enterprise/Datasmith/DatasmithSolidworksExporter/Private/ObjectMaterials.cs
--- a/Engine/Source/Programs/Enterprise/Datasmith/DatasmithSolidworksExporter/Private/ObjectMaterials.cs
+++ b/Engine/Source/Programs/Enterprise/Datasmith/DatasmithSolidworksExporter/Private/ObjectMaterials.cs
@@ -343,9 +343,9 @@ namespace DatasmithSolidworks
 			// Check for materials that are not per component (but instead per face, feature, part etc.)
 			Parallel.ForEach(InComponentsSet, CompName =>
 			{
-				Component2 Comp = InAsmDoc.SwAsmDoc.GetComponentByName(CompName);
+				Component2 Comp = null;
 
-				if (Comp != null)
+				if (InAsmDoc.SyncState.ExportedComponentsMap.TryGetValue(CompName, out Comp))
 				{
 					if (!DocMaterials.ContainsKey(CompName))
 					{
