Solidworks - autosync button is disabled by default when first opening Solidworks 2021

diff --git a/Engine/Source/Programs/Enterprise/Datasmith/DatasmithSolidworksExporter/Private/Addin.cs b/Engine/Source/Programs/Enterprise/Datasmith/DatasmithSolidworksExporter/Private/Addin.cs
--- a/Engine/Source/Programs/Enterprise/Datasmith/DatasmithSolidworksExporter/Private/Addin.cs
+++ b/Engine/Source/Programs/Enterprise/Datasmith/DatasmithSolidworksExporter/Private/Addin.cs
@@ -324,6 +324,12 @@ namespace DatasmithSolidworks
 			if (CurrentDocument != null)
 			{
 				CurrentDocument.bDirectLinkAutoSync = !CurrentDocument.bDirectLinkAutoSync;
+
+				if (CurrentDocument.bDirectLinkAutoSync && CurrentDocument.DirectLinkSyncCount == 0)
+				{
+					// Run first sync
+					CurrentDocument.OnDirectLinkSync();
+				}
 			}
 		}
 
@@ -334,7 +340,7 @@ namespace DatasmithSolidworks
 			// 2 Selects and disables the item
 			// 3 Selects and enables the item
 
-			if (CurrentDocument == null || CurrentDocument.DirectLinkSyncCount == 0)
+			if (CurrentDocument == null)
 			{
 				return 0;
 			}
