3ds Max DirectLink: fixed crash on File/New

- also scene is not parsed/updated until first Sync command

diff --git a/Engine/Source/Programs/Enterprise/Datasmith/DatasmithMaxExporter/Private/DatasmithMaxDirectLinkNotifications.cpp b/Engine/Source/Programs/Enterprise/Datasmith/DatasmithMaxExporter/Private/DatasmithMaxDirectLinkNotifications.cpp
--- a/Engine/Source/Programs/Enterprise/Datasmith/DatasmithMaxExporter/Private/DatasmithMaxDirectLinkNotifications.cpp
+++ b/Engine/Source/Programs/Enterprise/Datasmith/DatasmithMaxExporter/Private/DatasmithMaxDirectLinkNotifications.cpp
@@ -404,11 +404,7 @@ FNotifications::FNotifications(IExporter& InExporter)
 
 FNotifications::~FNotifications()
 {
-	if (NodeEventCallback)
-	{
-		GetISceneEventManager()->UnRegisterCallback(NodeEventCallback->CallbackKey);
-	}
-	
+	Reset();
 }
 
 void FNotifications::AddNode(INode* Node)
@@ -426,10 +422,26 @@ void FNotifications::Reset()
 	{
 		MaterialObserver->Reset();
 	}
+
+	for (int Code: NotificationCodesRegistered)
+	{
+		UnRegisterNotification(On3dsMaxNotification, this, Code);
+	}
+	NotificationCodesRegistered.Reset();
+	if (NodeEventCallback)
+	{
+		GetISceneEventManager()->UnRegisterCallback(NodeEventCallback->CallbackKey);
+		NodeEventCallback.Reset();
+	}
+	bRegistered = false;
 }
 
 void FNotifications::RegisterForNotifications()
 {
+	if (bRegistered)
+	{
+		return;
+	}
 	// Build todo: remove strings, for debug/logging
 #pragma warning(push)
 #pragma warning(disable:4995) // disable error on deprecated events, we just assign handlers not firing them
@@ -460,6 +472,7 @@ void FNotifications::RegisterForNotifications()
 	{
 		RegisterNotification(On3dsMaxNotification, this, Code);
 		NotificationCodetoString.Add(Code, Strings[i]);
+		NotificationCodesRegistered.Add(Code);
 		++i;
 	}
 
@@ -467,6 +480,8 @@ void FNotifications::RegisterForNotifications()
 	// Setup Node Event System callback
 	// https://help.autodesk.com/view/3DSMAX/2018/ENU/?guid=__files_GUID_7C91D285_5683_4606_9F7C_B8D3A7CA508B_htm
 	NodeEventCallback->CallbackKey = GetISceneEventManager()->RegisterCallback(NodeEventCallback.Get());
+
+	bRegistered = true;
 }
 
 FString FNotifications::ConvertNotificationCodeToString(int code)
@@ -531,15 +546,14 @@ void FNotifications::On3dsMaxNotification(void* param, NotifyInfo* info)
 		SceneTracker.NodeDeleted(reinterpret_cast<INode*>(info->callParam));
 		break;
 	}
-
-	case NOTIFY_SYSTEM_POST_RESET:
+	// Handle New/Reset events - reset tracking immediately when "Pre events are received - after this point all nodes are invalid, don't wait for "Post" event
+	case NOTIFY_SYSTEM_PRE_NEW:  // Sent when File>New>New All is selected 
+	case NOTIFY_SYSTEM_PRE_RESET:  // Sent when Reset OR File>New>New From Template is selected
 		Exporter->Reset();
-		Exporter->ParseScene();
 		break;
 
 	case NOTIFY_FILE_POST_OPEN:
 		Exporter->Reset();
-		Exporter->ParseScene();
 		break;
 
 	}
diff --git a/Engine/Source/Programs/Enterprise/Datasmith/DatasmithMaxExporter/Private/DatasmithMaxExporterWithDirectLink.cpp b/Engine/Source/Programs/Enterprise/Datasmith/DatasmithMaxExporter/Private/DatasmithMaxExporterWithDirectLink.cpp
--- a/Engine/Source/Programs/Enterprise/Datasmith/DatasmithMaxExporter/Private/DatasmithMaxExporterWithDirectLink.cpp
+++ b/Engine/Source/Programs/Enterprise/Datasmith/DatasmithMaxExporter/Private/DatasmithMaxExporterWithDirectLink.cpp
@@ -384,7 +384,8 @@ public:
 	bool ParseScene()
 	{
 		INode* Node = GetCOREInterface()->GetRootNode();
-		return ParseScene(Node, nullptr);
+		bSceneParsed = ParseScene(Node, nullptr);
+		return bSceneParsed;
 	}
 
 	// Parset scene or XRef scene(in this case attach to parent datasmith actor)
@@ -473,6 +474,7 @@ public:
 
 	void Reset()
 	{
+		bSceneParsed = false;
 		NodeTrackers.Reset();
 		NodeTrackersNames.Reset();
 		CollisionNodes.Reset();
@@ -547,6 +549,11 @@ public:
 	{
 		FUpdateProgress ProgressManager(!bQuiet, 6); // Will shutdown on end of Update
 
+		if (!bSceneParsed) // Parse whole scene only once
+		{
+			ParseScene();
+		}
+
 		ProgressManager.ProgressStage(TEXT("Refresh layers"));
 		{
 			UpdateLayers();
@@ -1641,6 +1648,9 @@ public:
 
 	FDatasmith3dsMaxScene& ExportedScene;
 	FNotifications& NotificationsHandler;
+
+	bool bSceneParsed = false;
+
 	TMap<FNodeKey, FNodeTrackerHandle> NodeTrackers; // All scene nodes
 	TMap<FString, TSet<FNodeTracker*>> NodeTrackersNames; // Nodes grouped by name
 	TSet<FNodeTracker*> InvalidatedNodeTrackers; // Nodes that need to be rebuilt
@@ -1701,7 +1711,6 @@ public:
 	// Just export, parsing scene from scratch
 	bool Export(bool bQuiet)
 	{
-		SceneTracker.ParseScene();
 		SceneTracker.Update(bQuiet);
 		SceneTracker.ExportAnimations();
 		ExportedScene.GetSceneExporter().Export(ExportedScene.GetDatasmithScene(), false);
@@ -1717,6 +1726,7 @@ public:
 	void UpdateDirectLinkScene()
 	{
 		DirectLinkImpl->UpdateScene(ExportedScene.GetDatasmithScene());
+		StartSceneChangeTracking(); // Always track scene changes if it's synced with DirectLink
 	}
 
 	static VOID AutoSyncTimerProc(HWND, UINT, UINT_PTR TimerIdentifier, DWORD)
diff --git a/Engine/Source/Programs/Enterprise/Datasmith/DatasmithMaxExporter/Public/DatasmithMaxDirectLink.h b/Engine/Source/Programs/Enterprise/Datasmith/DatasmithMaxExporter/Public/DatasmithMaxDirectLink.h
--- a/Engine/Source/Programs/Enterprise/Datasmith/DatasmithMaxExporter/Public/DatasmithMaxDirectLink.h
+++ b/Engine/Source/Programs/Enterprise/Datasmith/DatasmithMaxExporter/Public/DatasmithMaxDirectLink.h
@@ -219,8 +219,10 @@ public:
 
 	static void On3dsMaxNotification(void* param, NotifyInfo* info);
 
+	bool bRegistered = false; // Whether notification are being processed
 	IExporter& Exporter;
 	TMap<int, FString> NotificationCodetoString; // todo: remove, just for debug to output strings for notification codes
+	TArray<int> NotificationCodesRegistered;
 	TUniquePtr<FNodeEventCallback> NodeEventCallback;
 	TUniquePtr<FNodeObserver> NodeObserver;
 	TUniquePtr<FMaterialObserver> MaterialObserver;
