3ds Max DirectLink: actions icons

- icons available in max for all existing actions(Sync, AutoSync, Export, Collections, Messages)
- icons are visible in Customize UI actions list and when assigned to a custom toolbars or ribbon
- added example Ribbon from Pierr-Felix (only modified to use icons from actions, not files). Installing to "C:\Program Files\Autodesk\3ds Max 2022\en-US\UI\Ribbon\Extensions" would work

diff --git a/Engine/Source/Programs/Enterprise/Datasmith/DatasmithMaxExporter/Extras/3dsMaxScripts/UI/DatasmithUI.xaml b/Engine/Source/Programs/Enterprise/Datasmith/DatasmithMaxExporter/Extras/3dsMaxScripts/UI/DatasmithUI.xaml
new file mode 100644
--- /dev/null
+++ b/Engine/Source/Programs/Enterprise/Datasmith/DatasmithMaxExporter/Extras/3dsMaxScripts/UI/DatasmithUI.xaml
@@ -0,0 +1,248 @@
+<?xml version="1.0" encoding="utf-8"?>
+<prese_wind:ResourceDictionary xmlns:x="http://schemas.microsoft.com/winfx/2006/xaml" xmlns:context="clr-namespace:WPFCustomControls.ContextQueries;assembly=WPFCustomControls" xmlns:ads="clr-namespace:Autodesk.Windows;assembly=AdWindows" xmlns:wpfcust="clr-namespace:WPFCustomControls;assembly=WPFCustomControls" xmlns:prese_wind="clr-namespace:System.Windows;assembly=PresentationFramework">
+  <wpfcust:MaxRibbonTab
+    IsActive="true"
+    Name="Datasmith"
+    Title="Datasmith"
+    UID="9d925117-9b31-4445-a189-ccca79934d90"
+    x:Key="Datasmith">
+    <wpfcust:MaxModelingRibbonPanel
+      CollapsedPanelImage="pack://application:,,,/AdWindows;component/Resources/Images/MissingImage.ico"
+      FloatingGroup="0"
+      FloatingOrder="0"
+      UID="0b8165dd-f894-42a6-9bfc-64e5b189fa9d">
+      <wpfcust:MaxRibbonPanelSource
+        Title="Direct Link"
+        UID="5f41f38d-6458-4e7a-bb7b-1799477ffa1c">
+        <ads:RibbonButton
+          CommandHandler="{wpfcust:ActionItem ActionId=0, TableId=689133272}"
+          IsCheckable="true"
+          IsToolTipEnabled="false"
+          MinWidth="0"
+          Orientation="Vertical"
+          ResizeStyle="HideText"
+          ShowText="true"
+          Size="Large"
+          Text="Synchronize"
+          UID="e12deb32-dc84-4f23-874f-000722c61b11" />
+        <ads:RibbonButton
+          CommandHandler="{wpfcust:ActionItem ActionId=1, TableId=689133272}"
+          IsCheckable="true"
+          IsToolTipEnabled="false"
+          MinWidth="0"
+          Orientation="Vertical"
+          ResizeStyle="HideText"
+          ShowText="true"
+          Size="Large"
+          Text="Toggle Autos Sync"
+          UID="04cb208c-2fbc-424b-ab75-f206985cc0f4" />
+        <ads:RibbonPanelBreak
+          AllowInToolBar="false"
+          IsToolTipEnabled="true"
+          MinWidth="100"
+          ShowImage="false"
+          ShowText="false"
+          UID="38c8f9a4-a524-4fa5-82d5-f60d80dcef06" />
+        <ads:RibbonRowPanel
+          AllowInToolBar="false"
+          IsToolTipEnabled="true"
+          MinWidth="100"
+          ShowText="false"
+          UID="533984eb-3982-460e-b97b-d71ec74800f8">
+          <ads:RibbonButton
+            CommandHandler="{wpfcust:ActionItem ActionId=2, TableId=689133272}"
+            IsCheckable="true"
+            IsToolTipEnabled="false"
+            MinWidth="0"
+            ResizeStyle="HideText"
+            ShowText="true"
+            Size="Large"
+            Text="Connections"
+            UID="e7499676-fcc7-4d7b-9b23-d71a9beea4e9" />
+        </ads:RibbonRowPanel>
+      </wpfcust:MaxRibbonPanelSource>
+    </wpfcust:MaxModelingRibbonPanel>
+    <wpfcust:MaxModelingRibbonPanel
+      CollapsedPanelImage="pack://application:,,,/AdWindows;component/Resources/Images/MissingImage.ico"
+      FloatingGroup="0"
+      FloatingOrder="0"
+      UID="0113ae1a-ac1c-4b19-bd49-e69ea811778d">
+      <wpfcust:MaxRibbonPanelSource
+        Title="File Export"
+        UID="6bff3880-3957-435d-b72d-e49ca2555709">
+        <ads:RibbonButton
+          CommandHandler="{wpfcust:ActionItem ActionId=3, TableId=689133272}"
+          IsCheckable="true"
+          IsToolTipEnabled="false"
+          MinWidth="0"
+          Orientation="Vertical"
+          ResizeStyle="HideText"
+          ShowText="true"
+          Size="Large"
+          Text="Export Datasmith file"
+          UID="a1de5acd-fef8-45be-8e01-ec194b9b5191" />
+        <ads:RibbonPanelBreak
+          AllowInToolBar="false"
+          IsToolTipEnabled="true"
+          MinWidth="100"
+          ShowText="false"
+          UID="d6de1715-c063-4e7a-b1ec-325f11ee079a" />
+        <ads:RibbonRowPanel
+          AllowInToolBar="false"
+          Image="Icons/MissingImage_16.png"
+          IsToolTipEnabled="true"
+          LargeImage="Icons/MissingImage_32.png"
+          MinWidth="100"
+          ShowText="false"
+          UID="4f43d4b3-f741-4262-8a64-ed43f3354248">
+          <wpfcust:MaxRibbonLabel
+            AllowInToolBar="false"
+            IsToolTipEnabled="false"
+            MinWidth="150"
+            ResizeStyle="NoResize"
+            ShowImage="false"
+            ShowText="true"
+            Text="Include:"
+            UID="8d852bc3-7af5-4060-920a-a0514588d74a" />
+          <ads:RibbonSeparator
+            AllowInToolBar="false"
+            IsToolTipEnabled="true"
+            MinWidth="0"
+            ResizeStyle="NoResize"
+            SeparatorStyle="Spacer"
+            ShowText="false"
+            UID="918543c4-50dc-4203-bbdf-fe5fa4bd10fe" />
+          <wpfcust:MaxRibbonMenuButton
+            IsCheckable="false"
+            IsSplit="false"
+            IsSynchronizedWithCurrentItem="true"
+            IsToolTipEnabled="false"
+            MinWidth="150"
+            ResizeStyle="HideText"
+            ShowImage="false"
+            ShowText="true"
+            UID="f97d81dd-03f8-4a14-b718-06733fb224f2">
+            <wpfcust:MaxRibbonMenuItem
+              AllowInToolBar="false"
+              IsCheckable="true"
+              IsToolTipEnabled="false"
+              MinWidth="100"
+              ShowImage="false"
+              ShowText="true"
+              Text="Visible Objects"
+              UID="6ef627ea-f1fb-4b7a-9235-6f302d5ac4de" />
+            <wpfcust:MaxRibbonMenuItem
+              AllowInToolBar="false"
+              IsCheckable="true"
+              IsToolTipEnabled="false"
+              MinWidth="100"
+              ShowImage="false"
+              ShowText="true"
+              Text="Selected Objects"
+              UID="9509a5c2-1553-4499-91c3-7b7b0a220cee" />
+          </wpfcust:MaxRibbonMenuButton>
+          <ads:RibbonRowBreak
+            AllowInToolBar="false"
+            IsToolTipEnabled="true"
+            MinWidth="100"
+            ShowText="false"
+            UID="81ee274b-ff06-4d67-a974-6461da0acf54" />
+          <wpfcust:MaxRibbonLabel
+            AllowInToolBar="false"
+            IsToolTipEnabled="false"
+            MinWidth="150"
+            ResizeStyle="NoResize"
+            ShowImage="false"
+            ShowText="true"
+            Text="Animated Transforms:"
+            UID="3907c965-cc4b-4257-8e99-a3c24dea6fa9" />
+          <ads:RibbonSeparator
+            AllowInToolBar="false"
+            IsToolTipEnabled="true"
+            MinWidth="0"
+            ResizeStyle="NoResize"
+            SeparatorStyle="Spacer"
+            ShowText="false"
+            UID="5308eeb0-047b-4bd0-bb89-56745a049b4a" />
+          <wpfcust:MaxRibbonMenuButton
+            IsCheckable="false"
+            IsSplit="false"
+            IsSynchronizedWithCurrentItem="true"
+            IsToolTipEnabled="false"
+            MinWidth="100"
+            ResizeStyle="HideText"
+            ShowImage="false"
+            ShowText="true"
+            UID="59b29aac-ba4b-47bc-beec-31caf1bd6e2b">
+            <wpfcust:MaxRibbonMenuItem
+              AllowInToolBar="false"
+              IsCheckable="true"
+              IsToolTipEnabled="false"
+              MinWidth="100"
+              ShowImage="false"
+              ShowText="true"
+              Text="Current Frame Only"
+              UID="ef4472c0-fa88-4379-af73-842e4b57f752" />
+            <wpfcust:MaxRibbonMenuItem
+              AllowInToolBar="false"
+              IsCheckable="true"
+              IsToolTipEnabled="false"
+              MinWidth="100"
+              ShowImage="false"
+              ShowText="true"
+              Text="Active Time Segment"
+              UID="4097c080-33ee-430d-a47f-ab6e5c06c007" />
+          </wpfcust:MaxRibbonMenuButton>
+        </ads:RibbonRowPanel>
+      </wpfcust:MaxRibbonPanelSource>
+    </wpfcust:MaxModelingRibbonPanel>
+    <wpfcust:MaxModelingRibbonPanel
+      CollapsedPanelImage="pack://application:,,,/AdWindows;component/Resources/Images/MissingImage.ico"
+      FloatingGroup="0"
+      FloatingOrder="0"
+      ResizeStyle="NeverCollapsePanel"
+      UID="a2baf0eb-f693-4079-b75d-e5b8c818aedc">
+      <wpfcust:MaxRibbonPanelSource
+        Title="Tools"
+        UID="6159b143-ad36-4498-9b1c-916917018ef4">
+        <ads:RibbonButton
+          CommandHandler="{wpfcust:ActionItem ActionId=4, TableId=689133272}"
+          IsCheckable="true"
+          IsToolTipEnabled="false"
+          MinWidth="0"
+          Orientation="Vertical"
+          ResizeStyle="HideText"
+          ShowText="true"
+          Size="Large"
+          Text="Messages"
+          UID="57863e62-4371-4d09-b12b-50e0a4c3afd8" />
+        <ads:RibbonSeparator
+          AllowInToolBar="false"
+          IsToolTipEnabled="true"
+          MinWidth="0"
+          ResizeStyle="NoResize"
+          ShowImage="false"
+          ShowText="false"
+          UID="c9e25ff1-0624-464c-bf24-4a033b8b0407" />
+        <ads:RibbonRowPanel
+          AllowInToolBar="false"
+          IsToolTipEnabled="true"
+          MinWidth="100"
+          ShowImage="false"
+          ShowText="false"
+          Size="Large"
+          UID="3a10e004-8897-4ca0-8824-0662af3c5097">
+          <ads:RibbonButton
+            IsCheckable="true"
+            IsToolTipEnabled="false"
+            MinWidth="0"
+            ResizeStyle="HideText"
+            ShowImage="false"
+            ShowText="true"
+            Text="Add Attributes Modifier"
+            UID="cdee8c6d-56a9-4c9b-8776-7e19d8c2358e" />
+        </ads:RibbonRowPanel>
+      </wpfcust:MaxRibbonPanelSource>
+    </wpfcust:MaxModelingRibbonPanel>
+  </wpfcust:MaxRibbonTab>
+</prese_wind:ResourceDictionary>
\ No newline at end of file
diff --git a/Engine/Source/Programs/Enterprise/Datasmith/DatasmithMaxExporter/Private/DatasmithMaxDirectLinkMaxScript.cpp b/Engine/Source/Programs/Enterprise/Datasmith/DatasmithMaxExporter/Private/DatasmithMaxDirectLinkMaxScript.cpp
--- a/Engine/Source/Programs/Enterprise/Datasmith/DatasmithMaxExporter/Private/DatasmithMaxDirectLinkMaxScript.cpp
+++ b/Engine/Source/Programs/Enterprise/Datasmith/DatasmithMaxExporter/Private/DatasmithMaxDirectLinkMaxScript.cpp
@@ -23,6 +23,8 @@ MAX_INCLUDES_START
 	#include "maxscript/foundation/arrays.h"
 	#include "maxscript\macros\define_instantiation_functions.h"
 
+	#include <maxicon.h>
+
 MAX_INCLUDES_END
 
 extern HINSTANCE HInstanceMax;
@@ -250,8 +252,8 @@ Primitive LogInfo_pf(_M("Datasmith_LogInfo"), LogInfo_cf);
 class FDatasmithActions
 {
 public:
-	const ActionTableId ActionTableId = 0x291356d8;
-	const ActionContextId ActionContextId = 0x291356d9;
+	static const ActionTableId ActionTableId = 0x291356d8;
+	static const ActionContextId ActionContextId = 0x291356d9;
 
 	enum EActionIds
 	{
@@ -262,15 +264,75 @@ public:
 		ID_SHOWLOG_ACTION_ID,
 	};
 
-	class FDatasmithActionCallback : public ActionCallback
+	class FDatasmithActionTable : public ActionTable
 	{
+	public:
+		explicit FDatasmithActionTable(MSTR& Name): ActionTable(ActionTableId, ActionContextId, Name)
+		{
+			
+		}
+
+		BOOL IsChecked(int ActionId)
+		{
+			switch (ActionId)
+			{
+			case ID_AUTOSYNC_ACTION_ID: return GetExporter() && GetExporter()->IsAutoSyncEnabled(); // Can't change AutoSync icon but able to set its state to Checked when it's on
+			};
+			return false;
+		}
+
+		TMap<int, TUniquePtr<MaxBmpFileIcon>> IconForAction;
+
+		MaxIcon* GetIcon(int ActionId) override
+		{
+			TUniquePtr<MaxBmpFileIcon>& Icon = IconForAction.FindOrAdd(ActionId);
 
+			if (Icon)
+			{
+				return Icon.Get();
+			}
+
+			switch (ActionId)
+			{
+			case ID_SYNC_ACTION_ID:
+			{
+				Icon.Reset(new MaxBmpFileIcon(_T(":/Datasmith/Icons/DatasmithSyncIcon")));
+				break;
+			}
+			case ID_AUTOSYNC_ACTION_ID:
+			{
+				Icon.Reset(new MaxBmpFileIcon(_T(":/Datasmith/Icons/DatasmithAutoSyncIconOn")));
+				break;
+			}
+			case ID_CONNECTIONS_ACTION_ID:
+			{
+				Icon.Reset(new MaxBmpFileIcon(_T(":/Datasmith/Icons/DatasmithManageConnectionsIcon")));
+				break;
+			}
+			case ID_EXPORT_ACTION_ID:
+			{
+				Icon.Reset(new MaxBmpFileIcon(_T(":/Datasmith/Icons/DatasmithIcon")));
+				break;
+			}
+			case ID_SHOWLOG_ACTION_ID:
+			{
+				Icon.Reset(new MaxBmpFileIcon(_T(":/Datasmith/Icons/DatasmithLogIcon")));
+				break;
+			}
+			}
+			return Icon.Get();
+		}
+	};
+
+
+	class FDatasmithActionCallback : public ActionCallback
+	{
 	public:
 
-		BOOL ExecuteAction (int id) {
-			LogDebug(FString::Printf(TEXT("Action: %d"), id));
+		BOOL ExecuteAction (int ActionId) {
+			LogDebug(FString::Printf(TEXT("Action: %d"), ActionId));
 
-			switch (id)
+			switch (ActionId)
 			{
 			case ID_SYNC_ACTION_ID:
 			{
@@ -318,7 +380,7 @@ public:
 
 	FDatasmithActions()
 		: Name(TEXT("Datasmith"))
-		, Table(ActionTableId, ActionContextId, Name)
+		, Table(Name)
 	{
 		// todo: localization of Name
 
@@ -360,7 +422,7 @@ public:
 
 private:
 	TSTR Name;
-	ActionTable Table;
+	FDatasmithActionTable Table;
 	FDatasmithActionCallback ActionCallback;
 };
 
diff --git a/Engine/Source/Programs/Enterprise/Datasmith/DatasmithMaxExporter/Private/DatasmithMaxExporterWithDirectLink.cpp b/Engine/Source/Programs/Enterprise/Datasmith/DatasmithMaxExporter/Private/DatasmithMaxExporterWithDirectLink.cpp
--- a/Engine/Source/Programs/Enterprise/Datasmith/DatasmithMaxExporter/Private/DatasmithMaxExporterWithDirectLink.cpp
+++ b/Engine/Source/Programs/Enterprise/Datasmith/DatasmithMaxExporter/Private/DatasmithMaxExporterWithDirectLink.cpp
@@ -727,7 +727,7 @@ public:
 		return NodeTracker;
 	}
 
-	void RemoveMaterial(const TSharedPtr<IDatasmithBaseMaterialElement>& DatasmithMaterial)
+	virtual void RemoveMaterial(const TSharedPtr<IDatasmithBaseMaterialElement>& DatasmithMaterial) override
 	{
 		ExportedScene.DatasmithSceneRef->RemoveMaterial(DatasmithMaterial);		
 	}
@@ -785,7 +785,7 @@ public:
 	}
 
 	// Called when mesh element is not needed anymore and should be removed from the scene
-	void ReleaseMeshElement(TSharedPtr<IDatasmithMeshElement> Mesh)
+	virtual void ReleaseMeshElement(TSharedPtr<IDatasmithMeshElement> Mesh) override
 	{
 		ExportedScene.GetDatasmithScene()->RemoveMesh(Mesh);
 	}
@@ -1269,7 +1269,7 @@ public:
 		}
 	}
 
-	void AddMeshElement(TSharedPtr<IDatasmithMeshElement>& DatasmithMeshElement, FDatasmithMesh& DatasmithMesh, FDatasmithMesh* CollisionMesh)
+	virtual void AddMeshElement(TSharedPtr<IDatasmithMeshElement>& DatasmithMeshElement, FDatasmithMesh& DatasmithMesh, FDatasmithMesh* CollisionMesh) override
 	{
 		ExportedScene.GetDatasmithScene()->AddMesh(DatasmithMeshElement);
 
@@ -1306,7 +1306,7 @@ public:
 		return false;
 	}
 
-	void SetupActor(FNodeTracker& NodeTracker)
+	virtual void SetupActor(FNodeTracker& NodeTracker) override
 	{
 		NodeTracker.DatasmithActorElement->SetLabel(NodeTracker.Node->GetName());
 
@@ -1417,8 +1417,8 @@ public:
 
 		return true;
 	}
-	
-	void SetupDatasmithHISMForNode(FNodeTracker& NodeTracker, INode* GeometryNode, const FRenderMeshForConversion& RenderMesh, Mtl* Material, int32 MeshIndex, const TArray<Matrix3>& Transforms)
+
+	virtual void SetupDatasmithHISMForNode(FNodeTracker& NodeTracker, INode* GeometryNode, const FRenderMeshForConversion& RenderMesh, Mtl* Material, int32 MeshIndex, const TArray<Matrix3>& Transforms) override
 	{
 		FString MeshName = FString::FromInt(NodeTracker.Node->GetHandle()) + TEXT("_") + FString::FromInt(MeshIndex);
 
@@ -1522,7 +1522,7 @@ public:
 
 	/******************* Events *****************************/
 
-	void NodeAdded(INode* Node)
+	virtual void NodeAdded(INode* Node) override
 	{
 		// Node sometimes is null. 'Added' NodeEvent might come after node was actually deleted(immediately after creation)
 		// e.g.[mxs]: b = box(); delete b 
@@ -1537,7 +1537,7 @@ public:
 		ParseNode(Node);
 	}
 
-	void NodeDeleted(INode* Node)
+	virtual void NodeDeleted(INode* Node) override
 	{
 		// todo: check for null
 
@@ -1570,7 +1570,7 @@ public:
 		}
 	}
 
-	void NodeTransformChanged(FNodeKey NodeKey)
+	virtual void NodeTransformChanged(FNodeKey NodeKey) override
 	{
 		// todo: invalidate transform only
 
@@ -1590,13 +1590,13 @@ public:
 		}
 	}
 
-	void NodeMaterialAssignmentChanged(FNodeKey NodeKey)
+	virtual void NodeMaterialAssignmentChanged(FNodeKey NodeKey) override
 	{
 		//todo: handle more precisely
 		InvalidateNode(NodeKey);
 	}
 
-	void NodeMaterialGraphModified(FNodeKey NodeKey)
+	virtual void NodeMaterialGraphModified(FNodeKey NodeKey) override
 	{
 		//identify material tree and update all materials
 		//todo: possible to handle this more precisely(only refresh changed materials) - see FMaterialObserver
@@ -1619,7 +1619,7 @@ public:
 		InvalidateNode(NodeKey); // Invalidate node that has this material assigned. This is needed to trigger rebuild - exported geometry might change(e.g. multimaterial changed to slots will change on static mesh)
 	}
 
-	void NodeGeometryChanged(FNodeKey NodeKey)
+	virtual void NodeGeometryChanged(FNodeKey NodeKey) override
 	{
 		// GeometryChanged is executed to handle:
 		// - actual geometry modification(in any way)
@@ -1628,7 +1628,7 @@ public:
 		InvalidateNode(NodeKey);
 	}
 
-	void NodeHideChanged(FNodeKey NodeKey)
+	virtual void NodeHideChanged(FNodeKey NodeKey) override
 	{
 		// todo: invalidate visibility only - note to handle this not enought add/remove 
 		// actor. make sure to invalidate instances(in case geometry usage changed - like hidden node with multimat), materials
@@ -1636,7 +1636,7 @@ public:
 		InvalidateNode(NodeKey);
 	}
 
-	void NodePropertiesChanged(FNodeKey NodeKey)
+	virtual void NodePropertiesChanged(FNodeKey NodeKey) override
 	{
 		// todo: invalidate visibility only - note to handle this not enought add/remove 
 		// actor. make sure to invalidate instances(in case geometry usage changed - like hidden node with multimat), materials
@@ -1690,7 +1690,7 @@ class FExporter: public IExporter
 public:
 	FExporter(): NotificationsHandler(*this), SceneTracker(ExportedScene, NotificationsHandler) {}
 
-	virtual void Shutdown();
+	virtual void Shutdown() override;
 
 	virtual void SetOutputPath(const TCHAR* Path) override
 	{
@@ -1703,7 +1703,7 @@ public:
 		ExportedScene.SetName(Name);
 	}
 
-	virtual void ParseScene() 
+	virtual void ParseScene() override
 	{
 		SceneTracker.ParseScene();
 	}
@@ -1717,13 +1717,13 @@ public:
 		return true;
 	}
 
-	void InitializeDirectLinkForScene()
+	virtual void InitializeDirectLinkForScene() override
 	{
 		DirectLinkImpl.Reset(new FDatasmithDirectLink);
 		DirectLinkImpl->InitializeForScene(ExportedScene.GetDatasmithScene());
 	}
 
-	void UpdateDirectLinkScene()
+	virtual void UpdateDirectLinkScene() override
 	{
 		DirectLinkImpl->UpdateScene(ExportedScene.GetDatasmithScene());
 		StartSceneChangeTracking(); // Always track scene changes if it's synced with DirectLink
@@ -1756,9 +1756,12 @@ public:
 		}
 	}
 
-	bool bAutoSyncEnabled = false;
+	virtual bool IsAutoSyncEnabled() override
+	{
+		return bAutoSyncEnabled;
+	}
 
-	bool ToggleAutoSync()
+	virtual bool ToggleAutoSync() override
 	{
 		if (bAutoSyncEnabled)
 		{
@@ -1774,22 +1777,20 @@ public:
 		return bAutoSyncEnabled;
 	}
 
-	FNotifications NotificationsHandler;
-
 	// Install change notification systems
-	void StartSceneChangeTracking()
+	virtual void StartSceneChangeTracking() override
 	{
 		NotificationsHandler.RegisterForNotifications();
 	}
 
-	bool UpdateScene(bool bQuiet)
+	virtual bool UpdateScene(bool bQuiet) override
 	{
 		SceneTracker.Update(bQuiet);
 
 		return true;
 	}
 
-	void Reset()
+	virtual void Reset() override
 	{
 		ExportedScene.Reset();
 
@@ -1814,7 +1815,7 @@ public:
 		NotificationsHandler.Reset();
 	}
 
-	ISceneTracker& GetSceneTracker()
+	virtual ISceneTracker& GetSceneTracker() override
 	{
 		return SceneTracker;
 	}
@@ -1824,6 +1825,9 @@ public:
 	FString OutputPath;
 
 	FSceneTracker SceneTracker;
+	FNotifications NotificationsHandler;
+
+	bool bAutoSyncEnabled = false;
 };
 
 static TUniquePtr<IExporter> Exporter;
@@ -1862,7 +1866,7 @@ IExporter* GetExporter()
 	return Exporter.Get();	
 }
 
-void FExporter::Shutdown()
+void FExporter::Shutdown() 
 {
 	Exporter.Reset();
 	FDatasmithDirectLink::Shutdown();
diff --git a/Engine/Source/Programs/Enterprise/Datasmith/DatasmithMaxExporter/Public/DatasmithMaxDirectLink.h b/Engine/Source/Programs/Enterprise/Datasmith/DatasmithMaxExporter/Public/DatasmithMaxDirectLink.h
--- a/Engine/Source/Programs/Enterprise/Datasmith/DatasmithMaxExporter/Public/DatasmithMaxDirectLink.h
+++ b/Engine/Source/Programs/Enterprise/Datasmith/DatasmithMaxExporter/Public/DatasmithMaxDirectLink.h
@@ -59,6 +59,7 @@ class IExporter
 	virtual void InitializeDirectLinkForScene() = 0;
 	virtual void UpdateDirectLinkScene() = 0;
 	virtual bool ToggleAutoSync() = 0;
+	virtual bool IsAutoSyncEnabled() = 0;
 };
 
 
