Datasmith: Fix post build of DatasmithFacadeCSharp (python step referencing old python2.7 binaries)

diff --git a/Engine/Source/Programs/Enterprise/Datasmith/DatasmithFacadeCSharp/DatasmithFacadeCSharp.Target.cs b/Engine/Source/Programs/Enterprise/Datasmith/DatasmithFacadeCSharp/DatasmithFacadeCSharp.Target.cs
--- a/Engine/Source/Programs/Enterprise/Datasmith/DatasmithFacadeCSharp/DatasmithFacadeCSharp.Target.cs
+++ b/Engine/Source/Programs/Enterprise/Datasmith/DatasmithFacadeCSharp/DatasmithFacadeCSharp.Target.cs
@@ -80,7 +80,7 @@ public class DatasmithFacadeCSharpTarget : TargetRules
 		PreBuildSteps.Add(SwigCommand);
 
 		// Add copyright headers after generating the new files from swig.
-		string PythonDir = Path.Combine("$(EngineDir)", "Binaries", "ThirdParty", "Python");
+		string PythonDir = Path.Combine("$(EngineDir)", "Binaries", "ThirdParty", "Python3");
 		if (TargetPlatform == UnrealTargetPlatform.Mac)
 		{
 			PreBuildSteps.Add(string.Format("\"{0}\" \"{1}\"", Path.Combine(PythonDir, "Mac", "bin", "python"), Path.Combine(ProjectPath, "FacadeHeaderHelper.py")));
