DatasmithRhinoExpoter - Fixing a bug where DirectLink would not detect a texture modification.

JeanLuc.Corenthin

diff --git a/Engine/Source/Programs/Enterprise/Datasmith/DatasmithRhinoExporter/Private/DatasmithRhinoExportContext.cs b/Engine/Source/Programs/Enterprise/Datasmith/DatasmithRhinoExporter/Private/DatasmithRhinoExportContext.cs
--- a/Engine/Source/Programs/Enterprise/Datasmith/DatasmithRhinoExporter/Private/DatasmithRhinoExportContext.cs
+++ b/Engine/Source/Programs/Enterprise/Datasmith/DatasmithRhinoExporter/Private/DatasmithRhinoExportContext.cs
@@ -574,6 +574,11 @@ namespace DatasmithRhino
 		public Texture RhinoTexture { get { return RhinoCommonObject as Texture; } }
 		public FDatasmithFacadeTexture ExportedTexture { get { return ExportedElement as FDatasmithFacadeTexture; } }
 		public string FilePath { get; private set; }
+		
+		/// <summary>
+		/// Holds the Ids of all the Rhino `Texture` objects using that DatasmithTexture.
+		/// </summary>
+		public HashSet<Guid> TextureIds { get; } = new HashSet<Guid>();
 
 		public DatasmithTextureInfo(Texture InRhinoTexture, string InName, string InFilePath)
 			: base(InRhinoTexture, FDatasmithFacadeElement.GetStringHash(InName), InName, InName)
@@ -1847,7 +1852,7 @@ namespace DatasmithRhino
 					{
 						string TextureHash = DatasmithRhinoUtilities.GetTextureHash(RhinoTexture);
 						TextureHashes.Add(TextureHash);
-						AddTextureHashMapping(TextureHash, RhinoTexture);
+						UpdateTextureHashMapping(TextureHash, RhinoTexture);
 					}
 				}
 
@@ -1862,21 +1867,45 @@ namespace DatasmithRhino
 			return Result;
 		}
 
-		private void AddTextureHashMapping(string TextureHash, Texture RhinoTexture)
+		private void UpdateTextureHashMapping(string TextureHash, Texture RhinoTexture)
 		{
-			if (!TextureIdToTextureHash.ContainsKey(RhinoTexture.Id))
+			Guid TextureId = RhinoTexture.Id;
+
+			if (TextureHashToTextureInfo.TryGetValue(TextureHash, out DatasmithTextureInfo ExistingTextureInfo))
+			{
+				ExistingTextureInfo.TextureIds.Add(TextureId);
+				if (ExistingTextureInfo.DirectLinkStatus == DirectLinkSynchronizationStatus.PendingDeletion)
+				{
+					ExistingTextureInfo.RestorePreviousDirectLinkStatus();
+				}
+			}
+			else
 			{
-				TextureIdToTextureHash.Add(RhinoTexture.Id, TextureHash);
+				if (DatasmithRhinoUtilities.GetRhinoTextureNameAndPath(RhinoTexture, out string TextureName, out string TexturePath))
+				{
+					TextureName = TextureLabelGenerator.GenerateUniqueNameFromBaseName(TextureName);
+					DatasmithTextureInfo TextureInfo = new DatasmithTextureInfo(RhinoTexture, TextureName, TexturePath);
+					TextureInfo.TextureIds.Add(TextureId);
+					TextureHashToTextureInfo.Add(TextureHash, TextureInfo);
+				}
+			}
 
-				if (!TextureHashToTextureInfo.ContainsKey(TextureHash))
+			// Rhino 'Texture' instances are not immutable and the Texture file they point to may be updated.
+			// We need to detect when our DatasmithTextureInfo are no longer used by any Texture in Rhino and flag them for deletion.
+			if (TextureIdToTextureHash.TryGetValue(TextureId, out string PreviousHash))
+			{
+				if (TextureHash != PreviousHash
+					&& TextureHashToTextureInfo.TryGetValue(PreviousHash, out DatasmithTextureInfo TextureInfo))
 				{
-					if (DatasmithRhinoUtilities.GetRhinoTextureNameAndPath(RhinoTexture, out string TextureName, out string TexturePath))
+					TextureInfo.TextureIds.Remove(TextureId);
+					if (TextureInfo.TextureIds.Count == 0)
 					{
-						TextureName = TextureLabelGenerator.GenerateUniqueNameFromBaseName(TextureName);
-						TextureHashToTextureInfo.Add(TextureHash, new DatasmithTextureInfo(RhinoTexture, TextureName, TexturePath));
+						TextureInfo.ApplyDeletedStatus();
 					}
 				}
 			}
+
+			TextureIdToTextureHash[TextureId] = TextureHash;
 		}
 
 		public void UpdateGroups(GroupTableEventType UpdateType, Group UpdatedGroup)
