3ds Max DirectLink: AutoSync fixes

- Full Update when AutoSync was just enabled
- Don't send redundant(with no changes) DirectLink update
- added Datasmith_SetAutoSyncDelay maxscript method for experimenting

diff --git a/Engine/Source/Programs/Enterprise/Datasmith/DatasmithMaxExporter/Private/DatasmithMaxDirectLinkMaxScript.cpp b/Engine/Source/Programs/Enterprise/Datasmith/DatasmithMaxExporter/Private/DatasmithMaxDirectLinkMaxScript.cpp
--- a/Engine/Source/Programs/Enterprise/Datasmith/DatasmithMaxExporter/Private/DatasmithMaxDirectLinkMaxScript.cpp
+++ b/Engine/Source/Programs/Enterprise/Datasmith/DatasmithMaxExporter/Private/DatasmithMaxDirectLinkMaxScript.cpp
@@ -193,6 +193,20 @@ Value* ToggleAutoSync_cf(Value** arg_list, int count)
 }
 Primitive ToggleAutoSync_pf(_M("Datasmith_ToggleAutoSync"), ToggleAutoSync_cf);
 
+Value* SetAutoSyncDelay_cf(Value** arg_list, int count) 
+{
+	check_arg_count(ToggleAutoSync, 1, count);
+
+	if (GetExporter())
+	{
+		GetExporter()->SetAutoSyncDelay(arg_list[0]->to_float());
+		return &true_value;
+	}
+
+	return &false_value;
+}
+Primitive SetAutoSyncDelay_pf(_M("Datasmith_SetAutoSyncDelay"), SetAutoSyncDelay_cf);
+
 Value* OpenDirectlinkUi_cf(Value** arg_list, int count) 
 {
 	check_arg_count(OpenDirectlinkUi, 0, count);
diff --git a/Engine/Source/Programs/Enterprise/Datasmith/DatasmithMaxExporter/Private/DatasmithMaxExporterWithDirectLink.cpp b/Engine/Source/Programs/Enterprise/Datasmith/DatasmithMaxExporter/Private/DatasmithMaxExporterWithDirectLink.cpp
--- a/Engine/Source/Programs/Enterprise/Datasmith/DatasmithMaxExporter/Private/DatasmithMaxExporterWithDirectLink.cpp
+++ b/Engine/Source/Programs/Enterprise/Datasmith/DatasmithMaxExporter/Private/DatasmithMaxExporterWithDirectLink.cpp
@@ -491,8 +491,10 @@ public:
 	// Check every layer and if it's modified invalidate nodes assigned to it
 	// 3ds Max doesn't have events for all Layer changes(e.g. Name seems to be just an UI thing and has no notifications) so
 	// we need to go through all layers every update to see what's changed
-	void UpdateLayers()
+	bool UpdateLayers()
 	{
+		bool bChangeEncountered = false;
+
 		ILayerManager* LayerManager = GetCOREInterface13()->GetLayerManager();
 		int LayerCount = LayerManager->GetLayerCount();
 
@@ -517,6 +519,7 @@ public:
 
 			if (LayerTracker->bIsInvalidated)
 			{
+				bChangeEncountered = true;
 				if (TSet<FNodeTracker*>* NodeTrackersPtr = NodesPerLayer.Find(LayerTracker.Get()))
 				{
 					for(FNodeTracker* NodeTracker:* NodeTrackersPtr)
@@ -527,17 +530,18 @@ public:
 				LayerTracker->bIsInvalidated = false;
 			}
 		}
+		return bChangeEncountered;
 	}
 
 	// Applies all recorded changes to Datasmith scene
-	void Update(bool bQuiet)
+	bool Update(bool bQuiet)
 	{
 		DatasmithMaxLogger::Get().Purge();
 		TimeValue Time = GetCOREInterface()->GetTime();
 		NodesPreparer.Start(Time);
 		__try
 		{
-			UpdateInternal(bQuiet);
+			return UpdateInternal(bQuiet);
 		}
 		__finally // Make sure that finalize called no matter what happened in Vegas
 		{
@@ -545,10 +549,12 @@ public:
 		}
 	}
 
-	void UpdateInternal(bool bQuiet)
+	bool UpdateInternal(bool bQuiet)
 	{
 		FUpdateProgress ProgressManager(!bQuiet, 6); // Will shutdown on end of Update
 
+		bool bChangeEncountered = false;
+
 		if (!bSceneParsed) // Parse whole scene only once
 		{
 			ParseScene();
@@ -556,9 +562,13 @@ public:
 
 		ProgressManager.ProgressStage(TEXT("Refresh layers"));
 		{
-			UpdateLayers();
+			bChangeEncountered = UpdateLayers() && bChangeEncountered;
 		}
 
+		// Changes present only when there are modified layers(changes checked manually), nodes(notified by Max) or materials(notified by Max with all changes in dependencies)
+		bChangeEncountered |= !InvalidatedNodeTrackers.IsEmpty();
+		bChangeEncountered |= !MaterialsCollectionTracker.GetInvalidatedMaterials().IsEmpty();
+
 		ProgressManager.ProgressStage(TEXT("Update node names"));
 		// todo: move to NameChanged and NodeAdded?
 		for (FNodeTracker* NodeTracker : InvalidatedNodeTrackers)
@@ -575,7 +585,7 @@ public:
 		ProgressManager.ProgressStage(TEXT("Refresh collisions")); // Update set of nodes used for collision 
 		{
 			FProgressCounter ProgressCounter(ProgressManager, InvalidatedNodeTrackers.Num());
-			TSet<FNodeTracker*> NodesWithChangedCollisionStatus; // Need to invalidate these nodes to make them renderable or to keep them from renderable dependinding on collision status
+			TSet<FNodeTracker*> NodesWithChangedCollisionStatus; // Need to invalidate these nodes to make them renderable or to keep them from renderable depending on collision status
 			for (FNodeTracker* NodeTracker : InvalidatedNodeTrackers)
 			{
 				ProgressCounter.Next();
@@ -681,6 +691,8 @@ public:
 		//}
 
 		LogDebug("Scene update: done");
+
+		return bChangeEncountered;
 	}
 
 	void ExportAnimations()
@@ -1743,15 +1755,17 @@ public:
 		if (GetLastInputInfo(&LastInputInfo))
 		{
 			DWORD CurrentTime = GetTickCount();
-			DWORD IdlePeriod = GetTickCount() - LastInputInfo.dwTime;
+			int32 IdlePeriod = GetTickCount() - LastInputInfo.dwTime;
 			LogDebug(FString::Printf(TEXT("CurrentTime: %ld, Idle time: %ld, IdlePeriod: %ld"), CurrentTime, LastInputInfo.dwTime, IdlePeriod));
 
-			if (IdlePeriod > 500)
+			if (IdlePeriod > FMath::RoundToInt(AutoSyncDelaySeconds*1000))
 			{
 				// Don't create progress bar for autosync - it steals focus, closes listener and what else
 				// todo: consider creating progress when a big change in scene is detected, e.g. number of nodes?
-				UpdateScene(true);
-				UpdateDirectLinkScene();
+				if (UpdateScene(true)) // Don't sent redundant update if scene change wasn't detected
+				{
+					UpdateDirectLinkScene();
+				}
 			}
 		}
 	}
@@ -1769,7 +1783,12 @@ public:
 		}
 		else
 		{
-			SetTimer(GetCOREInterface()->GetMAXHWnd(), reinterpret_cast<UINT_PTR>(this), 500, AutoSyncTimerProc);
+			// Perform full Sync when AutoSync is first enabled
+			UpdateScene(false);
+			UpdateDirectLinkScene();
+
+			const uint32 AutoSyncCheckIntervalMs = FMath::RoundToInt(AutoSyncDelaySeconds*1000);
+			SetTimer(GetCOREInterface()->GetMAXHWnd(), reinterpret_cast<UINT_PTR>(this), AutoSyncCheckIntervalMs, AutoSyncTimerProc);
 		}
 		bAutoSyncEnabled = !bAutoSyncEnabled;
 
@@ -1777,6 +1796,11 @@ public:
 		return bAutoSyncEnabled;
 	}
 
+	virtual void SetAutoSyncDelay(float Seconds) override
+	{
+		AutoSyncDelaySeconds = Seconds;
+	}
+
 	// Install change notification systems
 	virtual void StartSceneChangeTracking() override
 	{
@@ -1785,9 +1809,7 @@ public:
 
 	virtual bool UpdateScene(bool bQuiet) override
 	{
-		SceneTracker.Update(bQuiet);
-
-		return true;
+		return SceneTracker.Update(bQuiet);
 	}
 
 	virtual void Reset() override
@@ -1828,6 +1850,7 @@ public:
 	FSceneTracker SceneTracker;
 
 	bool bAutoSyncEnabled = false;
+	float AutoSyncDelaySeconds = 0.5f;
 };
 
 static TUniquePtr<IExporter> Exporter;
diff --git a/Engine/Source/Programs/Enterprise/Datasmith/DatasmithMaxExporter/Public/DatasmithMaxDirectLink.h b/Engine/Source/Programs/Enterprise/Datasmith/DatasmithMaxExporter/Public/DatasmithMaxDirectLink.h
--- a/Engine/Source/Programs/Enterprise/Datasmith/DatasmithMaxExporter/Public/DatasmithMaxDirectLink.h
+++ b/Engine/Source/Programs/Enterprise/Datasmith/DatasmithMaxExporter/Public/DatasmithMaxDirectLink.h
@@ -60,6 +60,7 @@ class IExporter
 	virtual void UpdateDirectLinkScene() = 0;
 	virtual bool ToggleAutoSync() = 0;
 	virtual bool IsAutoSyncEnabled() = 0;
+	virtual void SetAutoSyncDelay(float Seconds) = 0;
 };
 
 
