3ds Max DirectLink: install Ribbon with icon images

- Ribbon installed to Extensions so it's loaded on Max startup
- Large icons for Ribbons used, images installed into Max where they are available for the ribbon

diff --git a/Engine/Source/Programs/Enterprise/Datasmith/DatasmithMaxExporter/Extras/3dsMaxScripts/UI/DatasmithUI.xaml b/Engine/Source/Programs/Enterprise/Datasmith/DatasmithMaxExporter/Extras/3dsMaxScripts/UI/DatasmithUI.xaml
--- a/Engine/Source/Programs/Enterprise/Datasmith/DatasmithMaxExporter/Extras/3dsMaxScripts/UI/DatasmithUI.xaml
+++ b/Engine/Source/Programs/Enterprise/Datasmith/DatasmithMaxExporter/Extras/3dsMaxScripts/UI/DatasmithUI.xaml
@@ -16,8 +16,11 @@
         UID="5f41f38d-6458-4e7a-bb7b-1799477ffa1c">
         <ads:RibbonButton
           CommandHandler="{wpfcust:ActionItem ActionId=0, TableId=689133272}"
+          Description="Synchronize with Direct Link"
+          Image="Icons/Datasmith/DatasmithSyncIcon32.png"
           IsCheckable="true"
-          IsToolTipEnabled="false"
+          IsToolTipEnabled="true"
+          LargeImage="Icons/Datasmith/DatasmithSyncIcon32.png"
           MinWidth="0"
           Orientation="Vertical"
           ResizeStyle="HideText"
@@ -25,17 +28,21 @@
           Size="Large"
           Text="Synchronize"
           UID="e12deb32-dc84-4f23-874f-000722c61b11" />
-        <ads:RibbonButton
+        <wpfcust:MaxRibbonToggleButton
+          AlternateImage="Icons/Datasmith/DatasmithAutoSyncIconOff32.png"
           CommandHandler="{wpfcust:ActionItem ActionId=1, TableId=689133272}"
+          Description="Toggle Direct Link Auto Sync"
+          Id="1.689133272"
           IsCheckable="true"
-          IsToolTipEnabled="false"
+          IsToolTipEnabled="true"
+          LargeImage="Icons/Datasmith/DatasmithAutoSyncIconOn32.png"
           MinWidth="0"
           Orientation="Vertical"
           ResizeStyle="HideText"
           ShowText="true"
           Size="Large"
-          Text="Toggle Autos Sync"
-          UID="04cb208c-2fbc-424b-ab75-f206985cc0f4" />
+          Text="Toggle &#xA;Auto Sync"
+          UID="1f535f66-ea8f-41c1-b835-8855215bbed2" />
         <ads:RibbonPanelBreak
           AllowInToolBar="false"
           IsToolTipEnabled="true"
@@ -51,12 +58,14 @@
           UID="533984eb-3982-460e-b97b-d71ec74800f8">
           <ads:RibbonButton
             CommandHandler="{wpfcust:ActionItem ActionId=2, TableId=689133272}"
+            Description="Manage Connections"
+            Image="Icons/Datasmith/DatasmithManageConnectionsIcon32.png"
             IsCheckable="true"
-            IsToolTipEnabled="false"
+            IsToolTipEnabled="true"
+            LargeImage="Icons/Datasmith/DatasmithManageConnectionsIcon32.png"
             MinWidth="0"
             ResizeStyle="HideText"
             ShowText="true"
-            Size="Large"
             Text="Connections"
             UID="e7499676-fcc7-4d7b-9b23-d71a9beea4e9" />
         </ads:RibbonRowPanel>
@@ -72,14 +81,17 @@
         UID="6bff3880-3957-435d-b72d-e49ca2555709">
         <ads:RibbonButton
           CommandHandler="{wpfcust:ActionItem ActionId=3, TableId=689133272}"
+          Description="Export to Datasmith file"
+          Image="Icons/Datasmith/DatasmithIcon32.png"
           IsCheckable="true"
           IsToolTipEnabled="false"
+          LargeImage="Icons/Datasmith/DatasmithIcon32.png"
           MinWidth="0"
           Orientation="Vertical"
           ResizeStyle="HideText"
           ShowText="true"
           Size="Large"
-          Text="Export Datasmith file"
+          Text="Export"
           UID="a1de5acd-fef8-45be-8e01-ec194b9b5191" />
         <ads:RibbonPanelBreak
           AllowInToolBar="false"
@@ -207,8 +219,11 @@
         UID="6159b143-ad36-4498-9b1c-916917018ef4">
         <ads:RibbonButton
           CommandHandler="{wpfcust:ActionItem ActionId=4, TableId=689133272}"
+          Description="Show Messages"
+          Image="Icons/Datasmith/DatasmithLogIcon32.png"
           IsCheckable="true"
           IsToolTipEnabled="false"
+          LargeImage="Icons/Datasmith/DatasmithLogIcon32.png"
           MinWidth="0"
           Orientation="Vertical"
           ResizeStyle="HideText"
