Revit Datasmith Exporter crash on sync or export (Asset is connected to more than one asset.)

diff --git a/Engine/Source/Programs/Enterprise/Datasmith/DatasmithRevitExporter/Private/DatasmithRevitMaterial.cs b/Engine/Source/Programs/Enterprise/Datasmith/DatasmithRevitExporter/Private/DatasmithRevitMaterial.cs
--- a/Engine/Source/Programs/Enterprise/Datasmith/DatasmithRevitExporter/Private/DatasmithRevitMaterial.cs
+++ b/Engine/Source/Programs/Enterprise/Datasmith/DatasmithRevitExporter/Private/DatasmithRevitMaterial.cs
@@ -1783,13 +1783,17 @@ namespace DatasmithRevitExporter
 			string in_propertyName
 		)
         {
-			// DebugLog.AddLine($"Texture Property {in_propertyName}");
-
 			AssetProperty textureProperty = in_asset.FindByName(in_propertyName);
 
 			if (textureProperty != null)
 			{
-				Asset unifiedBitmapAsset = textureProperty.GetSingleConnectedAsset();
+				Asset unifiedBitmapAsset = null;
+
+				try
+				{
+					unifiedBitmapAsset = textureProperty.GetSingleConnectedAsset();
+				}
+				catch {}
 
 				if (unifiedBitmapAsset != null)
 				{
@@ -1799,8 +1803,6 @@ namespace DatasmithRevitExporter
 					{
 						string sourcePath = (sourceProperty as AssetPropertyString).Value;
 
-						// DebugLog.AddLine($"Texture path: {sourcePath}");
-
 						if (!string.IsNullOrEmpty(sourcePath))
 						{
 							if (sourcePath.Contains("|"))
