3ds Max DirectLink: fixes for sync/tracking and cache

- don't sync/start change tracking for new/loaded scene before Sync/AutoSync is pressed
- set Cache Dir on new/reset/load
- disable autosync after reset/load
- scene vendor in udatasmith
- fixed ProgresBar messed stage index
- retired maxscript dialog
- extra AutoSync maxscript functions - Datasmith_IsAutoSyncEnabled, Datasmith_SetAutoSyncIdleDelay

diff --git a/Engine/Source/Programs/Enterprise/Datasmith/DatasmithMaxExporter/Private/DatasmithMaxDirectLinkMaxScript.cpp b/Engine/Source/Programs/Enterprise/Datasmith/DatasmithMaxExporter/Private/DatasmithMaxDirectLinkMaxScript.cpp
--- a/Engine/Source/Programs/Enterprise/Datasmith/DatasmithMaxExporter/Private/DatasmithMaxDirectLinkMaxScript.cpp
+++ b/Engine/Source/Programs/Enterprise/Datasmith/DatasmithMaxExporter/Private/DatasmithMaxDirectLinkMaxScript.cpp
@@ -193,9 +193,22 @@ Value* ToggleAutoSync_cf(Value** arg_list, int count)
 }
 Primitive ToggleAutoSync_pf(_M("Datasmith_ToggleAutoSync"), ToggleAutoSync_cf);
 
+Value* IsAutoSyncEnabled_cf(Value** arg_list, int count) 
+{
+	check_arg_count(IsAutoSyncEnabled, 0, count);
+
+	if (GetExporter())
+	{
+		return bool_result(GetExporter()->IsAutoSyncEnabled());
+	}
+
+	return &false_value;
+}
+Primitive IsAutoSyncEnabled_pf(_M("Datasmith_IsAutoSyncEnabled"), IsAutoSyncEnabled_cf);
+
 Value* SetAutoSyncDelay_cf(Value** arg_list, int count) 
 {
-	check_arg_count(ToggleAutoSync, 1, count);
+	check_arg_count(SetAutoSyncDelay, 1, count);
 
 	if (GetExporter())
 	{
@@ -207,6 +220,20 @@ Value* SetAutoSyncDelay_cf(Value** arg_list, int count)
 }
 Primitive SetAutoSyncDelay_pf(_M("Datasmith_SetAutoSyncDelay"), SetAutoSyncDelay_cf);
 
+Value* SetAutoSyncIdleDelay_cf(Value** arg_list, int count) 
+{
+	check_arg_count(SetAutoSyncIdleDelay, 1, count);
+
+	if (GetExporter())
+	{
+		GetExporter()->SetAutoSyncIdleDelay(arg_list[0]->to_float());
+		return &true_value;
+	}
+
+	return &false_value;
+}
+Primitive SetAutoSyncIdleDelay_pf(_M("Datasmith_SetAutoSyncIdleDelay"), SetAutoSyncIdleDelay_cf);
+
 Value* OpenDirectlinkUi_cf(Value** arg_list, int count) 
 {
 	check_arg_count(OpenDirectlinkUi, 0, count);
diff --git a/Engine/Source/Programs/Enterprise/Datasmith/DatasmithMaxExporter/Private/DatasmithMaxDirectLinkNotifications.cpp b/Engine/Source/Programs/Enterprise/Datasmith/DatasmithMaxExporter/Private/DatasmithMaxDirectLinkNotifications.cpp
--- a/Engine/Source/Programs/Enterprise/Datasmith/DatasmithMaxExporter/Private/DatasmithMaxDirectLinkNotifications.cpp
+++ b/Engine/Source/Programs/Enterprise/Datasmith/DatasmithMaxExporter/Private/DatasmithMaxDirectLinkNotifications.cpp
@@ -400,48 +400,28 @@ FNotifications::FNotifications(IExporter& InExporter)
 	: Exporter(InExporter)
 	, NodeObserver(MakeUnique<FNodeObserver>())
 	, MaterialObserver(MakeUnique<FMaterialObserver>())
-{}
-
-FNotifications::~FNotifications()
 {
-	Reset();
-}
-
-void FNotifications::AddNode(INode* Node)
-{
-	NodeObserver->AddItem(Node);
+	RegisterForSystemNotifications(); // Always follow system events like file reset, new, open and others. But some system notifications also track changes within scene.
 }
 
-void FNotifications::Reset()
+FNotifications::~FNotifications()
 {
-	if (NodeObserver)
-	{
-		NodeObserver->Reset();
-	}
-	if (MaterialObserver)
-	{
-		MaterialObserver->Reset();
-	}
-
 	for (int Code: NotificationCodesRegistered)
 	{
 		UnRegisterNotification(On3dsMaxNotification, this, Code);
 	}
 	NotificationCodesRegistered.Reset();
-	if (NodeEventCallback)
-	{
-		GetISceneEventManager()->UnRegisterCallback(NodeEventCallback->CallbackKey);
-		NodeEventCallback.Reset();
-	}
-	bRegistered = false;
+
+	StopSceneChangeTracking();
 }
 
-void FNotifications::RegisterForNotifications()
+void FNotifications::AddNode(INode* Node)
+{
+	NodeObserver->AddItem(Node);
+}
+
+void FNotifications::RegisterForSystemNotifications()
 {
-	if (bRegistered)
-	{
-		return;
-	}
 	// Build todo: remove strings, for debug/logging
 #pragma warning(push)
 #pragma warning(disable:4995) // disable error on deprecated events, we just assign handlers not firing them
@@ -475,13 +455,35 @@ void FNotifications::RegisterForNotifications()
 		NotificationCodesRegistered.Add(Code);
 		++i;
 	}
+}
+
+void FNotifications::StartSceneChangeTracking()
+{
+	if (bSceneChangeTracking)
+	{
+		return;
+	}
+	bSceneChangeTracking = true;
 
 	NodeEventCallback = MakeUnique<FNodeEventCallback>(Exporter.GetSceneTracker());
 	// Setup Node Event System callback
 	// https://help.autodesk.com/view/3DSMAX/2018/ENU/?guid=__files_GUID_7C91D285_5683_4606_9F7C_B8D3A7CA508B_htm
 	NodeEventCallback->CallbackKey = GetISceneEventManager()->RegisterCallback(NodeEventCallback.Get());
+}
+
+void FNotifications::StopSceneChangeTracking()
+{
+	if (!bSceneChangeTracking)
+	{
+		return;
+	}
+	bSceneChangeTracking = false;
 
-	bRegistered = true;
+	GetISceneEventManager()->UnRegisterCallback(NodeEventCallback->CallbackKey);
+	NodeEventCallback.Reset();
+
+	NodeObserver->Reset();
+	MaterialObserver->Reset();
 }
 
 FString FNotifications::ConvertNotificationCodeToString(int code)
@@ -492,12 +494,12 @@ FString FNotifications::ConvertNotificationCodeToString(int code)
 
 void FNotifications::On3dsMaxNotification(void* param, NotifyInfo* info)
 {
-	FNotifications& NotificationsHandler = *reinterpret_cast<FNotifications*>(param);
+	FNotifications& NotificationsHandler = *static_cast<FNotifications*>(param);
 	IExporter* Exporter = &NotificationsHandler.Exporter;
 
 	switch (info->intcode)
 	{
-		// Skip some events to display(spamming tests)
+	// Skip some events to display(spamming tests)
 	case NOTIFY_VIEWPORT_CHANGE:
 	case NOTIFY_PRE_RENDERER_CHANGE:
 	case NOTIFY_POST_RENDERER_CHANGE:
@@ -514,38 +516,45 @@ void FNotifications::On3dsMaxNotification(void* param, NotifyInfo* info)
 		LOG_DEBUG_HEAVY(FString(TEXT("Notify: ")) + NotificationsHandler.ConvertNotificationCodeToString(info->intcode));
 	};
 
-
-	ISceneTracker& SceneTracker = Exporter->GetSceneTracker();
-	switch (info->intcode)
+	if (NotificationsHandler.bSceneChangeTracking)
 	{
-	case NOTIFY_NODE_POST_MTL:
-		// todo: Event - node got a new material
-		break;
+		ISceneTracker& SceneTracker = Exporter->GetSceneTracker();
+		switch (info->intcode)
+		{
+		case NOTIFY_NODE_POST_MTL:
+			// Event - node got a new material
+			break;
 
-	case NOTIFY_SCENE_ADDED_NODE:
-	{
-		// note: INodeEventCallback::Added/Deleted is not used because there's a test case when it fails:
-		//   When a box is being created(dragging corners using mouse interface) and then cancelled during creation(RMB pressed)
-		//   INodeEventCallback::Deleted event is not fired by Max, although Added was called(along with other change events during creation)
+		case NOTIFY_SCENE_ADDED_NODE:
+		{
+			// note: INodeEventCallback::Added/Deleted is not used because there's a test case when it fails:
+			//   When a box is being created(dragging corners using mouse interface) and then cancelled during creation(RMB pressed)
+			//   INodeEventCallback::Deleted event is not fired by Max, although Added was called(along with other change events during creation)
 
-		INode* Node = reinterpret_cast<INode*>(info->callParam);
+			INode* Node = reinterpret_cast<INode*>(info->callParam);
 
-		LogDebugNode(NotificationsHandler.ConvertNotificationCodeToString(info->intcode), Node);
-		SceneTracker.NodeAdded(Node);
+			LogDebugNode(NotificationsHandler.ConvertNotificationCodeToString(info->intcode), Node);
+			SceneTracker.NodeAdded(Node);
 
-		break;
-	}
+			break;
+		}
 
-	case NOTIFY_SCENE_PRE_DELETED_NODE:
-	{
-		// note: INodeEventCallback::Deleted is not called when object creation was cancelled in the process
+		case NOTIFY_SCENE_PRE_DELETED_NODE:
+		{
+			// note: INodeEventCallback::Deleted is not called when object creation was cancelled in the process
 
-		INode* Node = reinterpret_cast<INode*>(info->callParam);
-		LogDebugNode(NotificationsHandler.ConvertNotificationCodeToString(info->intcode), Node);
+			INode* Node = reinterpret_cast<INode*>(info->callParam);
+			LogDebugNode(NotificationsHandler.ConvertNotificationCodeToString(info->intcode), Node);
 
-		SceneTracker.NodeDeleted(reinterpret_cast<INode*>(info->callParam));
-		break;
+			SceneTracker.NodeDeleted(reinterpret_cast<INode*>(info->callParam));
+			break;
+		}
+		}
 	}
+
+	// Handle events not related to tracking scene changes
+	switch (info->intcode)
+	{
 	// Handle New/Reset events - reset tracking immediately when "Pre events are received - after this point all nodes are invalid, don't wait for "Post" event
 	case NOTIFY_SYSTEM_PRE_NEW:  // Sent when File>New>New All is selected 
 	case NOTIFY_SYSTEM_PRE_RESET:  // Sent when Reset OR File>New>New From Template is selected
diff --git a/Engine/Source/Programs/Enterprise/Datasmith/DatasmithMaxExporter/Private/DatasmithMaxExporterWithDirectLink.cpp b/Engine/Source/Programs/Enterprise/Datasmith/DatasmithMaxExporter/Private/DatasmithMaxExporterWithDirectLink.cpp
--- a/Engine/Source/Programs/Enterprise/Datasmith/DatasmithMaxExporter/Private/DatasmithMaxExporterWithDirectLink.cpp
+++ b/Engine/Source/Programs/Enterprise/Datasmith/DatasmithMaxExporter/Private/DatasmithMaxExporterWithDirectLink.cpp
@@ -84,9 +84,14 @@ public:
 		SceneExporterRef.Reset();
 		SceneExporterRef = MakeShared<FDatasmithSceneExporter>();
 
-		// todo: compute or pass from script
+
+		MSTR Renderer;
+		FString Host;
+		Host = TEXT("Autodesk 3dsmax ") + FString::FromInt(MAX_VERSION_MAJOR) + TEXT(".") + FString::FromInt(MAX_VERSION_MINOR) + TEXT(".") + FString::FromInt(MAX_VERSION_POINT);
+		GetCOREInterface()->GetCurrentRenderer()->GetClassName(Renderer);
+
 		DatasmithSceneRef->SetProductName(TEXT("3dsmax"));
-		DatasmithSceneRef->SetHost(TEXT("3dsmax"));
+		DatasmithSceneRef->SetHost( *( Host + Renderer ) );
 
 		// Set the vendor name of the application used to build the scene.
 		DatasmithSceneRef->SetVendor(TEXT("Autodesk"));
@@ -103,6 +108,10 @@ public:
 		//  ---- PreInitPreStartupScreen
 		//  ----- FTaskGraphInterface::Startup
 		PreExport();
+
+		SetOutputPath(GetDirectlinkCacheDirectory());
+		FString SceneName = FPaths::GetCleanFilename(GetCOREInterface()->GetCurFileName().data());
+		SetName(*SceneName);
 	}
 
 	TSharedRef<IDatasmithScene> GetDatasmithScene()
@@ -205,7 +214,7 @@ public:
 class FUpdateProgress
 {
 	TUniquePtr<FDatasmithMaxProgressManager> ProgressManager;
-	int32 StageIndex;
+	int32 StageIndex = 0;
 	int32 StageCount;
 public:
 	FUpdateProgress(bool bShowProgressBar, int32 InStageCount) : StageCount(InStageCount)
@@ -229,6 +238,7 @@ public:
 
 	void ProgressEvent(float Progress, const TCHAR* Message)
 	{
+		LogDebug(FString::Printf(TEXT("%f %s"), Progress, Message));
 		if (ProgressManager)
 		{
 			ProgressManager->ProgressEvent(Progress, Message);
@@ -1758,7 +1768,7 @@ public:
 			int32 IdlePeriod = GetTickCount() - LastInputInfo.dwTime;
 			LogDebug(FString::Printf(TEXT("CurrentTime: %ld, Idle time: %ld, IdlePeriod: %ld"), CurrentTime, LastInputInfo.dwTime, IdlePeriod));
 
-			if (IdlePeriod > FMath::RoundToInt(AutoSyncDelaySeconds*1000))
+			if (IdlePeriod > FMath::RoundToInt(AutoSyncIdleDelaySeconds*1000))
 			{
 				// Don't create progress bar for autosync - it steals focus, closes listener and what else
 				// todo: consider creating progress when a big change in scene is detected, e.g. number of nodes?
@@ -1801,10 +1811,15 @@ public:
 		AutoSyncDelaySeconds = Seconds;
 	}
 
+	virtual void SetAutoSyncIdleDelay(float Seconds) override
+	{
+		AutoSyncIdleDelaySeconds = Seconds;
+	}
+
 	// Install change notification systems
 	virtual void StartSceneChangeTracking() override
 	{
-		NotificationsHandler.RegisterForNotifications();
+		NotificationsHandler.StartSceneChangeTracking();
 	}
 
 	virtual bool UpdateScene(bool bQuiet) override
@@ -1814,16 +1829,13 @@ public:
 
 	virtual void Reset() override
 	{
-		ExportedScene.Reset();
-
-		// todo: control output path from somewhere else?
-		if (!OutputPath.IsEmpty())
+		NotificationsHandler.StopSceneChangeTracking();
+		if (IsAutoSyncEnabled())
 		{
-			ExportedScene.SetOutputPath(*OutputPath);
+			ToggleAutoSync();
 		}
 
-		FString SceneName = FPaths::GetCleanFilename(GetCOREInterface()->GetCurFileName().data());
-		ExportedScene.SetName(*SceneName);
+		ExportedScene.Reset();
 
 		SceneTracker.Reset();
 
@@ -1834,7 +1846,6 @@ public:
 			DirectLinkImpl->InitializeForScene(ExportedScene.GetDatasmithScene());
 		}
 
-		NotificationsHandler.Reset();
 	}
 
 	virtual ISceneTracker& GetSceneTracker() override
@@ -1850,7 +1861,8 @@ public:
 	FSceneTracker SceneTracker;
 
 	bool bAutoSyncEnabled = false;
-	float AutoSyncDelaySeconds = 0.5f;
+	float AutoSyncDelaySeconds = 0.5f; // AutoSync is attempted periodically using this interval
+	float AutoSyncIdleDelaySeconds = 0.5f; // Time period user should be idle to run AutoSync
 };
 
 static TUniquePtr<IExporter> Exporter;
diff --git a/Engine/Source/Programs/Enterprise/Datasmith/DatasmithMaxExporter/Public/DatasmithMaxDirectLink.h b/Engine/Source/Programs/Enterprise/Datasmith/DatasmithMaxExporter/Public/DatasmithMaxDirectLink.h
--- a/Engine/Source/Programs/Enterprise/Datasmith/DatasmithMaxExporter/Public/DatasmithMaxDirectLink.h
+++ b/Engine/Source/Programs/Enterprise/Datasmith/DatasmithMaxExporter/Public/DatasmithMaxDirectLink.h
@@ -61,6 +61,7 @@ class IExporter
 	virtual bool ToggleAutoSync() = 0;
 	virtual bool IsAutoSyncEnabled() = 0;
 	virtual void SetAutoSyncDelay(float Seconds) = 0;
+	virtual void SetAutoSyncIdleDelay(float Seconds) = 0;
 };
 
 
@@ -212,16 +213,19 @@ public:
 	FNotifications(IExporter& InExporter);
 	~FNotifications();
 
-	void AddNode(INode*);
+	void RegisterForSystemNotifications();
 
-	void Reset();
+	void StartSceneChangeTracking();
+	void StopSceneChangeTracking();
+
+	void AddNode(INode*);
 
-	void RegisterForNotifications();
 	FString ConvertNotificationCodeToString(int code);
 
 	static void On3dsMaxNotification(void* param, NotifyInfo* info);
 
-	bool bRegistered = false; // Whether notification are being processed
+	bool bSceneChangeTracking = false;
+
 	IExporter& Exporter;
 	TMap<int, FString> NotificationCodetoString; // todo: remove, just for debug to output strings for notification codes
 	TArray<int> NotificationCodesRegistered;
