3ds Max DirectLink: fixed initial DL connection

- crash opening Connections and syncing when done on initial Max scene

diff --git a/Engine/Source/Programs/Enterprise/Datasmith/DatasmithMaxExporter/Private/DatasmithMaxExporterWithDirectLink.cpp b/Engine/Source/Programs/Enterprise/Datasmith/DatasmithMaxExporter/Private/DatasmithMaxExporterWithDirectLink.cpp
--- a/Engine/Source/Programs/Enterprise/Datasmith/DatasmithMaxExporter/Private/DatasmithMaxExporterWithDirectLink.cpp
+++ b/Engine/Source/Programs/Enterprise/Datasmith/DatasmithMaxExporter/Private/DatasmithMaxExporterWithDirectLink.cpp
@@ -1710,7 +1710,10 @@ public:
 class FExporter: public IExporter
 {
 public:
-	FExporter(): NotificationsHandler(*this), SceneTracker(ExportedScene, NotificationsHandler) {}
+	FExporter(): NotificationsHandler(*this), SceneTracker(ExportedScene, NotificationsHandler)
+	{
+		Reset();
+	}
 
 	virtual void Shutdown() override;
 
@@ -1839,13 +1842,8 @@ public:
 
 		SceneTracker.Reset();
 
-		if (DirectLinkImpl)
-		{
-			DirectLinkImpl.Reset();
-			DirectLinkImpl.Reset(new FDatasmithDirectLink);
-			DirectLinkImpl->InitializeForScene(ExportedScene.GetDatasmithScene());
-		}
-
+		DirectLinkImpl.Reset();
+		InitializeDirectLinkForScene();
 	}
 
 	virtual ISceneTracker& GetSceneTracker() override
