Solidworks crashes when auto-syncing configurations

diff --git a/Engine/Source/Programs/Enterprise/Datasmith/DatasmithSolidworksExporter/Private/AssemblyDocument.cs b/Engine/Source/Programs/Enterprise/Datasmith/DatasmithSolidworksExporter/Private/AssemblyDocument.cs
--- a/Engine/Source/Programs/Enterprise/Datasmith/DatasmithSolidworksExporter/Private/AssemblyDocument.cs
+++ b/Engine/Source/Programs/Enterprise/Datasmith/DatasmithSolidworksExporter/Private/AssemblyDocument.cs
@@ -322,10 +322,13 @@ namespace DatasmithSolidworks
 			// Export component children
 			object[] Children = (object[])InComponent.GetChildren();
 
-			foreach (object Obj in Children)
+			if (Children != null)
 			{
-				Component2 Child = (Component2)Obj;
-				ExportComponentRecursive(Child, InComponent, ref OutMeshesToExportMap);
+				foreach (object Obj in Children)
+				{
+					Component2 Child = (Component2)Obj;
+					ExportComponentRecursive(Child, InComponent, ref OutMeshesToExportMap);
+				}
 			}
 		}
 
