Revit: Autosync on Datasmith exporter does not work for Revit 2021 (It works on 2022)

diff --git a/Engine/Source/Programs/Enterprise/Datasmith/DatasmithRevitExporter/Private/DatasmithRevitDirectLink.cs b/Engine/Source/Programs/Enterprise/Datasmith/DatasmithRevitExporter/Private/DatasmithRevitDirectLink.cs
--- a/Engine/Source/Programs/Enterprise/Datasmith/DatasmithRevitExporter/Private/DatasmithRevitDirectLink.cs
+++ b/Engine/Source/Programs/Enterprise/Datasmith/DatasmithRevitExporter/Private/DatasmithRevitDirectLink.cs
@@ -302,7 +302,23 @@ namespace DatasmithRevitExporter
 
 			if (bHasChanges || SyncCount == 0)
 			{
-				RevitCommandId CmdId = RevitCommandId.LookupCommandId("3478AAF5-75A6-4E6F-BBF8-FFD791CA1801".ToLower());
+				string CmdGUID = null;
+
+#if REVIT_API_2018
+				CmdGUID = "D38EF9AC-C9B0-4578-8FD2-B4065FEFFABD";
+#elif REVIT_API_2019
+				CmdGUID = "44342F25-7B40-4E5C-A3AA-D94C201D95E8";
+#elif REVIT_API_2020
+				CmdGUID = "0A6A844F-F738-4E7E-A53D-BE8C45CD12A9";
+#elif REVIT_API_2021
+				CmdGUID = "66E62F4C-9F5A-4A3F-94C4-E6DCE838C413";
+#elif REVIT_API_2022
+				CmdGUID = "3478AAF5-75A6-4E6F-BBF8-FFD791CA1801";
+#else
+				#error This version of Revit is not supported yet.
+#endif
+
+				RevitCommandId CmdId = RevitCommandId.LookupCommandId(CmdGUID.ToLower());
 				if (CmdId != null)
 				{
 					UIApp.PostCommand(CmdId);
