[Revit] Texture modifications do not propagate to direct link

diff --git a/Engine/Source/Programs/Enterprise/Datasmith/DatasmithRevitExporter/Private/DatasmithRevitDirectLink.cs b/Engine/Source/Programs/Enterprise/Datasmith/DatasmithRevitExporter/Private/DatasmithRevitDirectLink.cs
--- a/Engine/Source/Programs/Enterprise/Datasmith/DatasmithRevitExporter/Private/DatasmithRevitDirectLink.cs
+++ b/Engine/Source/Programs/Enterprise/Datasmith/DatasmithRevitExporter/Private/DatasmithRevitDirectLink.cs
@@ -371,6 +371,22 @@ namespace DatasmithRevitExporter
 			ModifiedLinkedDocuments.Clear();
 		}
 
+		public bool IsMaterialDirty(Material InMaterial)
+		{
+			if (InMaterial != null)
+			{
+				return RootCache.ModifiedElements.Contains(InMaterial.Id);
+			}
+			return false;
+		}
+
+		public void SetMaterialClean(Material InMaterial)
+		{
+			if (InMaterial != null)
+			{
+				RootCache.ModifiedElements.Remove(InMaterial.Id);
+			}
+		}
 		public void MarkForExport(Element InElement)
 		{
 			if (InElement.GetType() == typeof(RevitLinkInstance))
diff --git a/Engine/Source/Programs/Enterprise/Datasmith/DatasmithRevitExporter/Private/DatasmithRevitDocumentData.cs b/Engine/Source/Programs/Enterprise/Datasmith/DatasmithRevitExporter/Private/DatasmithRevitDocumentData.cs
--- a/Engine/Source/Programs/Enterprise/Datasmith/DatasmithRevitExporter/Private/DatasmithRevitDocumentData.cs
+++ b/Engine/Source/Programs/Enterprise/Datasmith/DatasmithRevitExporter/Private/DatasmithRevitDocumentData.cs
@@ -1436,12 +1436,14 @@ namespace DatasmithRevitExporter
 
 			CurrentMaterialName = FMaterialData.GetMaterialName(InMaterialNode, CurrentMaterial);
 
-			if (!MaterialDataMap.ContainsKey(CurrentMaterialName))
+			if (!MaterialDataMap.ContainsKey(CurrentMaterialName) || (CurrentMaterial != null && DirectLink != null && DirectLink.IsMaterialDirty(CurrentMaterial)))
 			{
 				// Keep track of a new Datasmith master material.
 				MaterialDataMap[CurrentMaterialName] = new FMaterialData(InMaterialNode, CurrentMaterial, InExtraTexturePaths);
 				NewMaterialsMap[CurrentMaterialName] = MaterialDataMap[CurrentMaterialName];
 
+				DirectLink?.SetMaterialClean(CurrentMaterial);
+
 				// A new Datasmith master material was created.
 				return true;
 			}
