SketchUp: fix crash when previously assigned texture is disposed

diff --git a/Engine/Source/Programs/Enterprise/Datasmith/DatasmithSketchUpRubyExporter/Private/DatasmithSketchUpExportContext.cpp b/Engine/Source/Programs/Enterprise/Datasmith/DatasmithSketchUpRubyExporter/Private/DatasmithSketchUpExportContext.cpp
--- a/Engine/Source/Programs/Enterprise/Datasmith/DatasmithSketchUpRubyExporter/Private/DatasmithSketchUpExportContext.cpp
+++ b/Engine/Source/Programs/Enterprise/Datasmith/DatasmithSketchUpRubyExporter/Private/DatasmithSketchUpExportContext.cpp
@@ -98,7 +98,6 @@ void FExportContext::Populate()
 	RootNode->DatasmithActorLabel = TEXT("Model");
 
 	// Parse/convert Model
-	Materials.PopulateFromModel(ModelRef);
 	Scenes.PopulateFromModel(ModelRef);
 	ComponentDefinitions.PopulateFromModel(ModelRef);
 
@@ -497,9 +496,8 @@ FMaterialOccurrence* FMaterialCollection::RegisterInstance(FMaterialIDType Mater
 		NodeOccurrence->MaterialOverride->UnregisterInstance(Context, NodeOccurrence);
 	}
 
-	if (const TSharedPtr<DatasmithSketchUp::FMaterial>* Ptr = Find(MaterialID))
+	if (const TSharedPtr<DatasmithSketchUp::FMaterial> Material = FindOrCreateMaterial(MaterialID))
 	{
-		const TSharedPtr<DatasmithSketchUp::FMaterial>& Material = *Ptr;
 		return &Material->RegisterInstance(NodeOccurrence);
 	}
 
@@ -508,10 +506,8 @@ FMaterialOccurrence* FMaterialCollection::RegisterInstance(FMaterialIDType Mater
 
 FMaterialOccurrence* FMaterialCollection::RegisterGeometry(FMaterialIDType MaterialID, DatasmithSketchUp::FEntitiesGeometry* EntitiesGeometry)
 {
-	if (const TSharedPtr<DatasmithSketchUp::FMaterial>* Ptr = Find(MaterialID))
+	if (const TSharedPtr<DatasmithSketchUp::FMaterial> Material = FindOrCreateMaterial(MaterialID))
 	{
-		const TSharedPtr<DatasmithSketchUp::FMaterial>& Material = *Ptr;
-
 		EntitiesGeometry->MaterialsUsed.Add(Material.Get());
 		return &Material->RegisterGeometry(EntitiesGeometry);
 	}
@@ -556,7 +552,7 @@ TSharedPtr<FMaterial> FMaterialCollection::CreateMaterial(SUMaterialRef Material
 	return Material;
 }
 
-void FMaterialCollection::CreateMaterial(FMaterialIDType MaterialID)
+TSharedPtr<FMaterial> FMaterialCollection::CreateMaterial(FMaterialIDType MaterialID)
 {
 	// Get the number of material definitions in the SketchUp model.
 	size_t SMaterialDefinitionCount = 0;
@@ -573,9 +569,10 @@ void FMaterialCollection::CreateMaterial(FMaterialIDType MaterialID)
 	{
 		if (MaterialID == DatasmithSketchUpUtils::GetMaterialID(SMaterialDefinitionRef))
 		{
-			CreateMaterial(SMaterialDefinitionRef);
+			return CreateMaterial(SMaterialDefinitionRef);
 		}
 	}
+	return nullptr;
 }
 
 void FMaterialCollection::InvalidateMaterial(SUMaterialRef MaterialDefinitionRef)
@@ -618,3 +615,12 @@ bool FMaterialCollection::InvalidateDefaultMaterial()
 	return true;
 }
 
+TSharedPtr<FMaterial> FMaterialCollection::FindOrCreateMaterial(FMaterialIDType MaterialID)
+{
+	if (TSharedPtr<FMaterial>* Ptr = MaterialDefinitionMap.Find(MaterialID))
+	{
+		return *Ptr;
+	}
+	return CreateMaterial(MaterialID);
+}
+
diff --git a/Engine/Source/Programs/Enterprise/Datasmith/DatasmithSketchUpRubyExporter/Private/DatasmithSketchUpExportContext.h b/Engine/Source/Programs/Enterprise/Datasmith/DatasmithSketchUpRubyExporter/Private/DatasmithSketchUpExportContext.h
--- a/Engine/Source/Programs/Enterprise/Datasmith/DatasmithSketchUpRubyExporter/Private/DatasmithSketchUpExportContext.h
+++ b/Engine/Source/Programs/Enterprise/Datasmith/DatasmithSketchUpRubyExporter/Private/DatasmithSketchUpExportContext.h
@@ -180,7 +180,7 @@ namespace DatasmithSketchUp
 		void Update();
 
 		TSharedPtr<FMaterial> CreateMaterial(SUMaterialRef SMaterialDefinitionRef);
-		void CreateMaterial(FMaterialIDType MaterialID);
+		TSharedPtr<FMaterial> CreateMaterial(FMaterialIDType MaterialID);
 		void InvalidateMaterial(SUMaterialRef SMaterialDefinitionRef);
 		bool InvalidateMaterial(FMaterialIDType MateriadId);
 		bool RemoveMaterial(FEntityIDType EntityId);
@@ -188,10 +188,7 @@ namespace DatasmithSketchUp
 		bool InvalidateDefaultMaterial();
 
 
-		TSharedPtr<FMaterial>* Find(FMaterialIDType MaterialID)
-		{
-			return MaterialDefinitionMap.Find(MaterialID);
-		}
+		TSharedPtr<FMaterial> FindOrCreateMaterial(FMaterialIDType MaterialID);
 
 		// Tell that this materials is assigned on the node
 		FMaterialOccurrence* RegisterInstance(FMaterialIDType MaterialID, FNodeOccurence* NodeOccurrence);
diff --git a/Engine/Source/Programs/Enterprise/Datasmith/DatasmithSketchUpRubyExporter/Private/DatasmithSketchUpExporter.cpp b/Engine/Source/Programs/Enterprise/Datasmith/DatasmithSketchUpRubyExporter/Private/DatasmithSketchUpExporter.cpp
--- a/Engine/Source/Programs/Enterprise/Datasmith/DatasmithSketchUpRubyExporter/Private/DatasmithSketchUpExporter.cpp
+++ b/Engine/Source/Programs/Enterprise/Datasmith/DatasmithSketchUpRubyExporter/Private/DatasmithSketchUpExporter.cpp
@@ -11,6 +11,7 @@
 #include "DatasmithSketchUpString.h"
 #include "DatasmithSketchUpSummary.h"
 
+
 #include "DatasmithSketchUpExportContext.h"
 
 #include "DatasmithDirectLink.h"
@@ -504,7 +505,7 @@ public:
 	
 	bool OnMaterialAdded(DatasmithSketchUp::FEntityIDType EntityId)
 	{
-		Context.Materials.CreateMaterial(EntityId);
+		// Not handling material additon here - materials are created when needed by geometry/components
 		return true;
 	}
 
diff --git a/Engine/Source/Programs/Enterprise/Datasmith/DatasmithSketchUpRubyExporter/Private/DatasmithSketchUpMaterial.cpp b/Engine/Source/Programs/Enterprise/Datasmith/DatasmithSketchUpRubyExporter/Private/DatasmithSketchUpMaterial.cpp
--- a/Engine/Source/Programs/Enterprise/Datasmith/DatasmithSketchUpRubyExporter/Private/DatasmithSketchUpMaterial.cpp
+++ b/Engine/Source/Programs/Enterprise/Datasmith/DatasmithSketchUpRubyExporter/Private/DatasmithSketchUpMaterial.cpp
@@ -292,8 +292,28 @@ void FMaterial::Remove(FExportContext& Context)
 {
 
 	MaterialDirectlyAppliedToMeshes.RemoveDatasmithElement(Context);
+
+	// Expecting that SketchUp removes material from geometry/instances prior to removing it from the scene, but in case it doesn't...
+	if (!ensure(MaterialDirectlyAppliedToMeshes.MeshesMaterialDirectlyAppliedTo.IsEmpty()))
+	{
+		for (FEntitiesGeometry* Geom: MaterialDirectlyAppliedToMeshes.MeshesMaterialDirectlyAppliedTo.Array())
+		{
+			UnregisterGeometry(Context, Geom);
+		}
+	}
+
 	MaterialInheritedByNodes.RemoveDatasmithElement(Context);
 
+	// Expecting that SketchUp removes material from geometry/instances prior to removing it from the scene, but in case it doesn't...
+	if (!ensure(MaterialInheritedByNodes.NodesMaterialInheritedBy.IsEmpty()))
+	{
+		for (FNodeOccurence* Node: MaterialDirectlyAppliedToMeshes.NodesMaterialInheritedBy.Array())
+		{
+			UnregisterInstance(Context, Node);
+		}
+	}
+
+
 	Context.Textures.UnregisterMaterial(this);
 }
 
@@ -500,7 +520,8 @@ void FMaterialOccurrence::UnregisterInstance(FExportContext& Context, FNodeOccur
 
 void FMaterialCollection::Update()
 {
-	// Update usage of textures by materials before updating textures(to only update used textures)
+	// Update usage of materials and textures by materials before updating textures(to only update used textures)
+	TArray<FMaterialIDType> UnusedMaterials;
 	for (TPair<FMaterialIDType, TSharedPtr<DatasmithSketchUp::FMaterial>> IdAndMaterial : MaterialDefinitionMap)
 	{
 		FMaterial& Material = *IdAndMaterial.Value;
@@ -508,6 +529,16 @@ void FMaterialCollection::Update()
 		{
 			Material.UpdateTexturesUsage(Context);
 		}
+		else
+		{
+			UnusedMaterials.Add(IdAndMaterial.Key);
+		}
+	}
+
+	// Remove unused materials(SketchUp material removal notification is unusable but we need to remove unused materials)
+	for (FMaterialIDType MaterialId: UnusedMaterials)
+	{
+		RemoveMaterial(MaterialId);
 	}
 
 	Context.Textures.Update();
