3ds Max DirectLink: export options

- options are affecting only Export To File, not DirectLink Sync

- also fixed options exclusive for Selected/Visible and Animation/NoAnim

diff --git a/Engine/Source/Programs/Enterprise/Datasmith/DatasmithMaxExporter/Extras/3dsMaxScripts/UI/DatasmithUI.xaml b/Engine/Source/Programs/Enterprise/Datasmith/DatasmithMaxExporter/Extras/3dsMaxScripts/UI/DatasmithUI.xaml
--- a/Engine/Source/Programs/Enterprise/Datasmith/DatasmithMaxExporter/Extras/3dsMaxScripts/UI/DatasmithUI.xaml
+++ b/Engine/Source/Programs/Enterprise/Datasmith/DatasmithMaxExporter/Extras/3dsMaxScripts/UI/DatasmithUI.xaml
@@ -124,35 +124,50 @@
             SeparatorStyle="Spacer"
             ShowText="false"
             UID="918543c4-50dc-4203-bbdf-fe5fa4bd10fe" />
-          <wpfcust:MaxRibbonMenuButton
+
+          <wpfcust:MaxRibbonSplitButton
             IsCheckable="false"
             IsSplit="false"
             IsSynchronizedWithCurrentItem="true"
-            IsToolTipEnabled="false"
-            MinWidth="150"
+            IsToolTipEnabled="true"
+            LargeImage="pack://application:,,,/AdWindows;component/Resources/Images/MissingImage.ico"
+            ListImageSize="Standard"
+            MaxscriptIndexCallback="Datasmith_ExportOptions_IncludeIndex()"
+            MinWidth="0"
             ResizeStyle="HideText"
-            ShowImage="false"
             ShowText="true"
+            Text=""
             UID="f97d81dd-03f8-4a14-b718-06733fb224f2">
-            <wpfcust:MaxRibbonMenuItem
+            <ads:RibbonCommandItem
               AllowInToolBar="false"
+              CommandHandler="{wpfcust:ActionItem MacroName=Datasmith_ExportOptions_IncludeAllVisibleObjects, MacroCategory=Datasmith}"
+              Id="mOBJPaintMotionTypeRandom"
               IsCheckable="true"
-              IsToolTipEnabled="false"
+              IsToolTipEnabled="true"
+              LargeImage="pack://application:,,,/AdWindows;component/Resources/Images/MissingImage.ico"
               MinWidth="100"
-              ShowImage="false"
-              ShowText="true"
+              Name="Visible Objects"
+              ShowText="false"
               Text="Visible Objects"
               UID="6ef627ea-f1fb-4b7a-9235-6f302d5ac4de" />
-            <wpfcust:MaxRibbonMenuItem
+            <ads:RibbonCommandItem
               AllowInToolBar="false"
+              CommandHandler="{wpfcust:ActionItem MacroName=Datasmith_ExportOptions_IncludeSelectedOnly, MacroCategory=Datasmith}"
+              Id="mOBJPaintMotionTypeConsecutive"
               IsCheckable="true"
-              IsToolTipEnabled="false"
+              IsToolTipEnabled="true"
+              LargeImage="pack://application:,,,/AdWindows;component/Resources/Images/MissingImage.ico"
               MinWidth="100"
-              ShowImage="false"
-              ShowText="true"
+              Name="Selected Objects"
+              ShowText="false"
               Text="Selected Objects"
               UID="9509a5c2-1553-4499-91c3-7b7b0a220cee" />
-          </wpfcust:MaxRibbonMenuButton>
+          </wpfcust:MaxRibbonSplitButton>
+
+
+
+
+
           <ads:RibbonRowBreak
             AllowInToolBar="false"
             IsToolTipEnabled="true"
@@ -176,35 +191,48 @@
             SeparatorStyle="Spacer"
             ShowText="false"
             UID="5308eeb0-047b-4bd0-bb89-56745a049b4a" />
-          <wpfcust:MaxRibbonMenuButton
+
+          <wpfcust:MaxRibbonSplitButton
+            Image="Icons/ObjectPaint/PaintRandom_16.ico"
             IsCheckable="false"
             IsSplit="false"
             IsSynchronizedWithCurrentItem="true"
-            IsToolTipEnabled="false"
-            MinWidth="100"
+            IsToolTipEnabled="true"
+            LargeImage="pack://application:,,,/AdWindows;component/Resources/Images/MissingImage.ico"
+            ListImageSize="Standard"
+            MaxscriptIndexCallback="Datasmith_ExportOptions_AnimatedTransformsIndex()"
+            MinWidth="0"
             ResizeStyle="HideText"
-            ShowImage="false"
             ShowText="true"
+            Text=""
             UID="59b29aac-ba4b-47bc-beec-31caf1bd6e2b">
-            <wpfcust:MaxRibbonMenuItem
+            <ads:RibbonCommandItem
               AllowInToolBar="false"
+              CommandHandler="{wpfcust:ActionItem MacroName=Datasmith_ExportOptions_AnimatedTransformsCurrentFrameOnly, MacroCategory=Datasmith}"
+              Id="mOBJPaintMotionTypeRandom"
               IsCheckable="true"
-              IsToolTipEnabled="false"
+              IsToolTipEnabled="true"
+              LargeImage="pack://application:,,,/AdWindows;component/Resources/Images/MissingImage.ico"
               MinWidth="100"
-              ShowImage="false"
-              ShowText="true"
+              Name="Current Frame Only"
+              ShowText="false"
               Text="Current Frame Only"
               UID="ef4472c0-fa88-4379-af73-842e4b57f752" />
-            <wpfcust:MaxRibbonMenuItem
+            <ads:RibbonCommandItem
               AllowInToolBar="false"
+              CommandHandler="{wpfcust:ActionItem MacroName=Datasmith_ExportOptions_AnimatedTransformsActiveSegment, MacroCategory=Datasmith}"
+              Id="mOBJPaintMotionTypeConsecutive"
               IsCheckable="true"
-              IsToolTipEnabled="false"
+              IsToolTipEnabled="true"
+              LargeImage="pack://application:,,,/AdWindows;component/Resources/Images/MissingImage.ico"
               MinWidth="100"
-              ShowImage="false"
-              ShowText="true"
+              Name="Active Time Segment"
+              ShowText="false"
               Text="Active Time Segment"
               UID="4097c080-33ee-430d-a47f-ab6e5c06c007" />
-          </wpfcust:MaxRibbonMenuButton>
+          </wpfcust:MaxRibbonSplitButton>
+
+
         </ads:RibbonRowPanel>
       </wpfcust:MaxRibbonPanelSource>
     </wpfcust:MaxModelingRibbonPanel>
diff --git a/Engine/Source/Programs/Enterprise/Datasmith/DatasmithMaxExporter/Plugin/PackageContents.xml b/Engine/Source/Programs/Enterprise/Datasmith/DatasmithMaxExporter/Plugin/PackageContents.xml
--- a/Engine/Source/Programs/Enterprise/Datasmith/DatasmithMaxExporter/Plugin/PackageContents.xml
+++ b/Engine/Source/Programs/Enterprise/Datasmith/DatasmithMaxExporter/Plugin/PackageContents.xml
@@ -53,8 +53,14 @@
         <ComponentEntry ModuleName="./UI/pre-start-up.ms" />
     </Components>
 
+    <Components Description="macroscripts parts">
+        <RuntimeRequirements OS="Win64" Platform="3ds Max" SeriesMin="2017" SeriesMax="2022" />
+        <ComponentEntry ModuleName="./UI/DatasmithMacros.mcr" />
+    </Components>
+
     <Components Description="post-start-up scripts parts">
         <RuntimeRequirements OS="Win64" Platform="3ds Max" SeriesMin="2017" SeriesMax="2022" />
         <ComponentEntry ModuleName="./UI/DatasmithDialog.ms" />
     </Components>
+
 </ApplicationPackage>
\ No newline at end of file
diff --git a/Engine/Source/Programs/Enterprise/Datasmith/DatasmithMaxExporter/Private/DatasmithMaxDirectLinkMaxScript.cpp b/Engine/Source/Programs/Enterprise/Datasmith/DatasmithMaxExporter/Private/DatasmithMaxDirectLinkMaxScript.cpp
--- a/Engine/Source/Programs/Enterprise/Datasmith/DatasmithMaxExporter/Private/DatasmithMaxDirectLinkMaxScript.cpp
+++ b/Engine/Source/Programs/Enterprise/Datasmith/DatasmithMaxExporter/Private/DatasmithMaxDirectLinkMaxScript.cpp
@@ -52,7 +52,6 @@ Value* OnLoad_cf(Value **arg_list, int count)
 	{
 		return &false_value;
 	}
-	GetExporter()->ParseScene();
 
 	return bool_result(true);
 }
@@ -242,6 +241,30 @@ Value* OpenDirectlinkUi_cf(Value** arg_list, int count)
 }
 Primitive OpenDirectlinkUi_pf(_M("Datasmith_OpenDirectlinkUi"), OpenDirectlinkUi_cf);
 
+#define DefinePersistentExportOption(name) \
+Value* GetExportOption_##name##_cf(Value** arg_list, int count) \
+{ \
+	check_arg_count(GetExportOption_##name, 0, count); \
+ \
+	return bool_result(GetPersistentExportOptions().Get##name##()); \
+} \
+Primitive GetExportOption_##name##_pf(_M("Datasmith_GetExportOption_" #name), GetExportOption_##name##_cf);\
+\
+Value* SetExportOption_##name##_cf(Value** arg_list, int count) \
+{ \
+	check_arg_count(SetExportOption_##name, 1, count); \
+	Value* pValue= arg_list[0]; \
+	GetPersistentExportOptions().Set##name##(pValue->to_bool()); \
+	return &true_value; \
+} \
+Primitive SetExportOption_##name##_pf(_M("Datasmith_SetExportOption_" #name), SetExportOption_##name##_cf);\
+
+//////////////////////////////////////////
+DefinePersistentExportOption(SelectedOnly)
+DefinePersistentExportOption(AnimatedTransforms)
+//////////////////////////////////////////
+
+#undef DefinePersistentExportOption
 
 Value* GetDirectlinkCacheDirectory_cf(Value** arg_list, int count)
 {
@@ -425,31 +448,18 @@ public:
 	{
 		// todo: localization of Name
 
-		static ActionDescription ActionsDescriptions[] = { 
-			ID_SYNC_ACTION_ID, // ID
-			IDS_SYNC_DESC, // Description
-			IDS_SYNC_NAME, // Name
-			IDS_DATASMITH_CATEGORY, // Category name
-
-			ID_AUTOSYNC_ACTION_ID, // ID
-			IDS_AUTOSYNC_DESC, // Description
-			IDS_AUTOSYNC_NAME, // Name
-			IDS_DATASMITH_CATEGORY, // Category name
-
-			ID_CONNECTIONS_ACTION_ID, // ID
-			IDS_CONNECTIONS_DESC, // Description
-			IDS_CONNECTIONS_NAME, // Name
-			IDS_DATASMITH_CATEGORY, // Category name
-
-			ID_EXPORT_ACTION_ID, // ID
-			IDS_EXPORT_DESC, // Description
-			IDS_EXPORT_NAME, // Name
-			IDS_DATASMITH_CATEGORY, // Category name
-
-			ID_SHOWLOG_ACTION_ID, // ID
-			IDS_SHOWLOG_DESC, // Description
-			IDS_SHOWLOG_NAME, // Name
-			IDS_DATASMITH_CATEGORY, // Category name
+#define DATASMITH_ACTION(name) \
+			ID_##name##_ACTION_ID, \
+			IDS_##name##_DESC, \
+			IDS_##name##_NAME, \
+			IDS_DATASMITH_CATEGORY 
+
+		static ActionDescription ActionsDescriptions[] = {
+			DATASMITH_ACTION(SYNC),
+			DATASMITH_ACTION(AUTOSYNC),
+			DATASMITH_ACTION(CONNECTIONS),
+			DATASMITH_ACTION(EXPORT),
+			DATASMITH_ACTION(SHOWLOG),
 		};
 
 		Table.BuildActionTable(nullptr, sizeof(ActionsDescriptions) / sizeof(ActionsDescriptions[0]), ActionsDescriptions, HInstanceMax);
diff --git a/Engine/Source/Programs/Enterprise/Datasmith/DatasmithMaxExporter/Private/DatasmithMaxExporterWithDirectLink.cpp b/Engine/Source/Programs/Enterprise/Datasmith/DatasmithMaxExporter/Private/DatasmithMaxExporterWithDirectLink.cpp
--- a/Engine/Source/Programs/Enterprise/Datasmith/DatasmithMaxExporter/Private/DatasmithMaxExporterWithDirectLink.cpp
+++ b/Engine/Source/Programs/Enterprise/Datasmith/DatasmithMaxExporter/Private/DatasmithMaxExporterWithDirectLink.cpp
@@ -384,12 +384,49 @@ public:
 
 
 
+struct FExportOptions
+{
+	// Default options for DirectLink change-tracking
+	bool bSelectedOnly = false;
+	bool bAnimatedTransforms = false;
+};
+
+// Global export options, stored in preferences
+class FPersistentExportOptions: public IPersistentExportOptions
+{
+public:
+
+	virtual void SetSelectedOnly(bool bValue) override
+	{
+		Options.bSelectedOnly = bValue;
+	}
+	virtual bool GetSelectedOnly() override
+	{
+		return Options.bSelectedOnly;
+	}
+
+	virtual void SetAnimatedTransforms(bool bValue) override
+	{
+		Options.bAnimatedTransforms = bValue;
+	}
+
+	virtual bool GetAnimatedTransforms() override
+	{
+		return Options.bAnimatedTransforms;
+	}
+
+	FExportOptions Options;
+};
 
 // Holds states of entities for syncronization and handles change events
 class FSceneTracker: public ISceneTracker
 {
 public:
-	FSceneTracker(FDatasmith3dsMaxScene& InExportedScene, FNotifications& InNotificationsHandler) : ExportedScene(InExportedScene), NotificationsHandler(InNotificationsHandler), MaterialsCollectionTracker(*this) {}
+	FSceneTracker(const FExportOptions& InOptions, FDatasmith3dsMaxScene& InExportedScene, FNotifications& InNotificationsHandler)
+		: Options(InOptions)
+		, ExportedScene(InExportedScene)
+		, NotificationsHandler(InNotificationsHandler)
+		, MaterialsCollectionTracker(*this) {}
 
 	bool ParseScene()
 	{
@@ -613,7 +650,6 @@ public:
 				ProgressCounter.Next();
 				UpdateNode(*NodeTracker);
 			}
-			InvalidatedNodeTrackers.Reset();
 		}
 		
 		ProgressManager.ProgressStage(TEXT("Process invalidated instances"));
@@ -623,10 +659,18 @@ public:
 			{
 				ProgressCounter.Next();
 				UpdateInstances(*Instances);
+				InvalidatedNodeTrackers.Append(Instances->NodeTrackers);
 			}
 			InvalidatedInstances.Reset();
 		}
 
+		ProgressManager.ProgressStage(TEXT("Reparent Datasmith Actors"));
+		for (FNodeTracker* NodeTracker : InvalidatedNodeTrackers)
+		{
+			AttachNodeToDatasmithScene(*NodeTracker);
+		}
+		InvalidatedNodeTrackers.Reset();
+
 		TSet<Mtl*> ActualMaterialToUpdate;
 		TSet<Texmap*> ActualTexmapsToUpdate;
 
@@ -989,6 +1033,11 @@ public:
 			return;
 		}
 
+		if (Options.bSelectedOnly && !NodeTracker.Node->Selected())
+		{
+			return;
+		}
+
 		ObjectState ObjState = NodeTracker.Node->EvalWorldState(0);
 		Object* Obj = ObjState.obj;
 
@@ -1089,39 +1138,50 @@ public:
 		NodeDatasmithMetadata.Add(&NodeTracker, MetadataElement);
 	}
 
-	void AttachNodeToDatasmithScene(FNodeTracker& NodeTracker)
+	// Find first ansestor node which has DatasmithActor created for it
+	FNodeTracker* GetParentNodeTracker(FNodeTracker& NodeTracker)
 	{
-		// Add to parent
 		FNodeKey ParentNodeKey = NodeEventNamespace::GetKeyByNode(NodeTracker.Node->GetParentNode());
-		if (FNodeTrackerHandle* ParentNodeTrackerHandle = NodeTrackers.Find(ParentNodeKey))
+		FNodeTrackerHandle* ParentNodeTrackerHandle = NodeTrackers.Find(ParentNodeKey);
+		if (!ParentNodeTrackerHandle)
 		{
-			// Add to parent Datasmith actor if it has been updated already, if not parent will add it
-			if (!IsNodeInvalidated(*ParentNodeTrackerHandle))
+			return nullptr; // Max node not known, assume it doesn't exist
+		}
+		
+		return ParentNodeTrackerHandle->GetNodeTracker();
+	}
+
+	// Not all nodes result in creation of DatasmithActor for them(e.g. skipped as invisible), find first ascestor that has it
+	FNodeTracker* GetAncestorNodeTrackerWithDatasmithActor(FNodeTracker& InNodeTracker)
+	{
+		FNodeTracker* NodeTracker = &InNodeTracker;
+		while (FNodeTracker* ParentNodeTracker = GetParentNodeTracker(*NodeTracker))
+		{
+			if (ParentNodeTracker->DatasmithActorElement)
 			{
-				FNodeTracker* ParentNodeTracker = ParentNodeTrackerHandle->GetNodeTracker();
-				ParentNodeTracker->DatasmithActorElement->AddChild(NodeTracker.DatasmithActorElement);
+				return ParentNodeTracker;
 			}
+			NodeTracker = ParentNodeTracker;
 		}
-		else
+		return nullptr;
+	}
+
+	void AttachNodeToDatasmithScene(FNodeTracker& NodeTracker)
+	{
+		if (!NodeTracker.DatasmithActorElement)
 		{
-			// If there's no parent node registered assume it's at root
-			ExportedScene.GetDatasmithScene()->AddActor(NodeTracker.DatasmithActorElement);
+			return;
 		}
 
-		// Add attach datasmith actors of child nodes
-		int32 ChildNum = NodeTracker.Node->NumberOfChildren();
-		for (int32 ChildIndex = 0; ChildIndex < ChildNum; ++ChildIndex)
+		if (FNodeTracker* ParentNodeTracker = GetAncestorNodeTrackerWithDatasmithActor(NodeTracker))
 		{
-
-			FNodeKey ChildNodeKey = NodeEventNamespace::GetKeyByNode(NodeTracker.Node->GetChildNode(ChildIndex));
-			if (FNodeTrackerHandle* ChildNodeTrackerHandle = NodeTrackers.Find(ChildNodeKey))
-			{
-				// Add child datasmith actor if child is updated, if not child will add itself(it will be updated further in the queue)
-				if (!IsNodeInvalidated(*ChildNodeTrackerHandle))
-				{
-					NodeTracker.DatasmithActorElement->AddChild(ChildNodeTrackerHandle->GetNodeTracker()->DatasmithActorElement);
-				}
-			}
+			ParentNodeTracker->DatasmithActorElement->AddChild(NodeTracker.DatasmithActorElement);
+		}
+		else
+		{
+			// If there's no ancestor node with DatasmithActor assume node it at root
+			// (node's parent might be node that was skipped - e.g. it was hidden in Max or not selected when exporting only selected objects)
+			ExportedScene.GetDatasmithScene()->AddActor(NodeTracker.DatasmithActorElement);
 		}
 	}
 
@@ -1247,7 +1307,6 @@ public:
 		NodeTracker.DatasmithActorElement = DatasmithActorElement;
 		NodeTracker.DatasmithMeshActor = DatasmithMeshActor;
 
-		AttachNodeToDatasmithScene(NodeTracker);
 		UpdateNodeMetadata(NodeTracker);
 		TagsConverter.ConvertNodeTags(NodeTracker);
 		if (NodeTracker.Layer)
@@ -1332,7 +1391,6 @@ public:
 	{
 		NodeTracker.DatasmithActorElement->SetLabel(NodeTracker.Node->GetName());
 
-		AttachNodeToDatasmithScene(NodeTracker);
 		UpdateNodeMetadata(NodeTracker);
 		TagsConverter.ConvertNodeTags(NodeTracker);
 		if (NodeTracker.Layer)
@@ -1668,6 +1726,7 @@ public:
 	
 	///////////////////////////////////////////////
 
+	const FExportOptions& Options;
 	FDatasmith3dsMaxScene& ExportedScene;
 	FNotifications& NotificationsHandler;
 
@@ -1706,11 +1765,10 @@ public:
 	FTagsConverter TagsConverter; // Converts max node information to Datasmith tags
 };
 
-
 class FExporter: public IExporter
 {
 public:
-	FExporter(): NotificationsHandler(*this), SceneTracker(ExportedScene, NotificationsHandler)
+	FExporter(FExportOptions& InOptions): Options(InOptions), NotificationsHandler(*this), SceneTracker(Options, ExportedScene, NotificationsHandler)
 	{
 		Reset();
 	}
@@ -1737,7 +1795,10 @@ public:
 	bool Export(bool bQuiet)
 	{
 		SceneTracker.Update(bQuiet);
-		SceneTracker.ExportAnimations();
+		if (Options.bAnimatedTransforms)
+		{
+			SceneTracker.ExportAnimations();
+		}
 		ExportedScene.GetSceneExporter().Export(ExportedScene.GetDatasmithScene(), false);
 		return true;
 	}
@@ -1851,6 +1912,8 @@ public:
 		return SceneTracker;
 	}
 
+	FExportOptions Options;
+
 	FDatasmith3dsMaxScene ExportedScene;
 	TUniquePtr<FDatasmithDirectLink> DirectLinkImpl;
 	FString OutputPath;
@@ -1863,7 +1926,9 @@ public:
 	float AutoSyncIdleDelaySeconds = 0.5f; // Time period user should be idle to run AutoSync
 };
 
-static TUniquePtr<IExporter> Exporter;
+FPersistentExportOptions PersistentExportOptions;
+
+TUniquePtr<IExporter> Exporter;
 
 bool CreateExporter(bool bEnableUI, const TCHAR* EnginePath)
 {
@@ -1883,7 +1948,8 @@ bool CreateExporter(bool bEnableUI, const TCHAR* EnginePath)
 		return false;
 	}
 
-	Exporter = MakeUnique<FExporter>();
+	static FExportOptions ExporterOptions; // Default options
+	Exporter = MakeUnique<FExporter>(ExporterOptions);
 	return true;
 }
 
@@ -1899,6 +1965,11 @@ IExporter* GetExporter()
 	return Exporter.Get();	
 }
 
+IPersistentExportOptions& GetPersistentExportOptions()
+{
+	return PersistentExportOptions;
+}
+
 void FExporter::Shutdown() 
 {
 	Exporter.Reset();
@@ -1908,7 +1979,7 @@ void FExporter::Shutdown()
 
 bool Export(const TCHAR* Name, const TCHAR* OutputPath, bool bQuiet)
 {
-	FExporter TempExporter;
+	FExporter TempExporter(PersistentExportOptions.Options);
 	TempExporter.ExportedScene.SetName(Name);
 	TempExporter.SetOutputPath(OutputPath);
 
diff --git a/Engine/Source/Programs/Enterprise/Datasmith/DatasmithMaxExporter/Public/DatasmithMaxDirectLink.h b/Engine/Source/Programs/Enterprise/Datasmith/DatasmithMaxExporter/Public/DatasmithMaxDirectLink.h
--- a/Engine/Source/Programs/Enterprise/Datasmith/DatasmithMaxExporter/Public/DatasmithMaxDirectLink.h
+++ b/Engine/Source/Programs/Enterprise/Datasmith/DatasmithMaxExporter/Public/DatasmithMaxDirectLink.h
@@ -34,6 +34,21 @@ namespace DatasmithMaxDirectLink
 class ISceneTracker;
 class FLayerTracker;
 
+
+// Global export options, stored in preferences
+class IPersistentExportOptions
+{
+public:
+	 // Whether to export all visible or selected objects only
+	virtual void SetSelectedOnly(bool) = 0;
+	virtual bool GetSelectedOnly() = 0;
+
+
+	// Whether to export animation or not
+	virtual void SetAnimatedTransforms(bool) = 0;
+	virtual bool GetAnimatedTransforms() = 0;
+};
+
 //---- Main class for export/change tracking
 class IExporter
 {
@@ -64,11 +79,12 @@ class IExporter
 	virtual void SetAutoSyncIdleDelay(float Seconds) = 0;
 };
 
-
 bool CreateExporter(bool bEnableUI, const TCHAR* EnginePath); // Create exporter with ability for DirectLink change tracking
 IExporter* GetExporter();
 void ShutdownExporter();
-bool Export(const TCHAR* Name, const TCHAR* OutputPath, bool bQuiet); 
+bool Export(const TCHAR* Name, const TCHAR* OutputPath, bool bQuiet);
+
+IPersistentExportOptions& GetPersistentExportOptions();
 
 // Identifies Max node to track its changes
 class FNodeTracker: FNoncopyable
