3ds Max DirectLink: ribbon panel icons and Attributes Modifier

- Datasmith icons are displayed when ribbon in 'collapsed'(minimized) state
- Add Datasmith Attributes modifier works

diff --git a/Engine/Source/Programs/Enterprise/Datasmith/DatasmithMaxExporter/Extras/3dsMaxScripts/UI/DatasmithUI.xaml b/Engine/Source/Programs/Enterprise/Datasmith/DatasmithMaxExporter/Extras/3dsMaxScripts/UI/DatasmithUI.xaml
--- a/Engine/Source/Programs/Enterprise/Datasmith/DatasmithMaxExporter/Extras/3dsMaxScripts/UI/DatasmithUI.xaml
+++ b/Engine/Source/Programs/Enterprise/Datasmith/DatasmithMaxExporter/Extras/3dsMaxScripts/UI/DatasmithUI.xaml
@@ -7,7 +7,7 @@
     UID="9d925117-9b31-4445-a189-ccca79934d90"
     x:Key="Datasmith">
     <wpfcust:MaxModelingRibbonPanel
-      CollapsedPanelImage="pack://application:,,,/AdWindows;component/Resources/Images/MissingImage.ico"
+      CollapsedPanelImage="Icons/Datasmith/DatasmithSyncIcon32.png"
       FloatingGroup="0"
       FloatingOrder="0"
       UID="0b8165dd-f894-42a6-9bfc-64e5b189fa9d">
@@ -72,7 +72,7 @@
       </wpfcust:MaxRibbonPanelSource>
     </wpfcust:MaxModelingRibbonPanel>
     <wpfcust:MaxModelingRibbonPanel
-      CollapsedPanelImage="pack://application:,,,/AdWindows;component/Resources/Images/MissingImage.ico"
+      CollapsedPanelImage="Icons/Datasmith/DatasmithIcon32.png"
       FloatingGroup="0"
       FloatingOrder="0"
       UID="0113ae1a-ac1c-4b19-bd49-e69ea811778d">
@@ -124,50 +124,45 @@
             SeparatorStyle="Spacer"
             ShowText="false"
             UID="918543c4-50dc-4203-bbdf-fe5fa4bd10fe" />
-
           <wpfcust:MaxRibbonSplitButton
+            Image="pack://application:,,,/AdWindows;component/Resources/Images/MissingImage.ico"
             IsCheckable="false"
             IsSplit="false"
             IsSynchronizedWithCurrentItem="true"
             IsToolTipEnabled="true"
-            LargeImage="pack://application:,,,/AdWindows;component/Resources/Images/MissingImage.ico"
             ListImageSize="Standard"
             MaxscriptIndexCallback="Datasmith_ExportOptions_IncludeIndex()"
             MinWidth="0"
             ResizeStyle="HideText"
+            ShowImage="false"
             ShowText="true"
-            Text=""
             UID="f97d81dd-03f8-4a14-b718-06733fb224f2">
             <ads:RibbonCommandItem
               AllowInToolBar="false"
               CommandHandler="{wpfcust:ActionItem MacroName=Datasmith_ExportOptions_IncludeAllVisibleObjects, MacroCategory=Datasmith}"
-              Id="mOBJPaintMotionTypeRandom"
+              Image="pack://application:,,,/AdWindows;component/Resources/Images/MissingImage.ico"
               IsCheckable="true"
               IsToolTipEnabled="true"
-              LargeImage="pack://application:,,,/AdWindows;component/Resources/Images/MissingImage.ico"
               MinWidth="100"
               Name="Visible Objects"
+              ShowImage="false"
               ShowText="false"
               Text="Visible Objects"
               UID="6ef627ea-f1fb-4b7a-9235-6f302d5ac4de" />
             <ads:RibbonCommandItem
               AllowInToolBar="false"
               CommandHandler="{wpfcust:ActionItem MacroName=Datasmith_ExportOptions_IncludeSelectedOnly, MacroCategory=Datasmith}"
-              Id="mOBJPaintMotionTypeConsecutive"
+              Image="pack://application:,,,/AdWindows;component/Resources/Images/MissingImage.ico"
               IsCheckable="true"
               IsToolTipEnabled="true"
               LargeImage="pack://application:,,,/AdWindows;component/Resources/Images/MissingImage.ico"
               MinWidth="100"
               Name="Selected Objects"
+              ShowImage="false"
               ShowText="false"
               Text="Selected Objects"
               UID="9509a5c2-1553-4499-91c3-7b7b0a220cee" />
           </wpfcust:MaxRibbonSplitButton>
-
-
-
-
-
           <ads:RibbonRowBreak
             AllowInToolBar="false"
             IsToolTipEnabled="true"
@@ -191,9 +186,8 @@
             SeparatorStyle="Spacer"
             ShowText="false"
             UID="5308eeb0-047b-4bd0-bb89-56745a049b4a" />
-
           <wpfcust:MaxRibbonSplitButton
-            Image="Icons/ObjectPaint/PaintRandom_16.ico"
+            Image="pack://application:,,,/AdWindows;component/Resources/Images/MissingImage.ico"
             IsCheckable="false"
             IsSplit="false"
             IsSynchronizedWithCurrentItem="true"
@@ -203,41 +197,41 @@
             MaxscriptIndexCallback="Datasmith_ExportOptions_AnimatedTransformsIndex()"
             MinWidth="0"
             ResizeStyle="HideText"
+            ShowImage="false"
             ShowText="true"
-            Text=""
             UID="59b29aac-ba4b-47bc-beec-31caf1bd6e2b">
             <ads:RibbonCommandItem
               AllowInToolBar="false"
               CommandHandler="{wpfcust:ActionItem MacroName=Datasmith_ExportOptions_AnimatedTransformsCurrentFrameOnly, MacroCategory=Datasmith}"
-              Id="mOBJPaintMotionTypeRandom"
+              Image="pack://application:,,,/AdWindows;component/Resources/Images/MissingImage.ico"
               IsCheckable="true"
               IsToolTipEnabled="true"
               LargeImage="pack://application:,,,/AdWindows;component/Resources/Images/MissingImage.ico"
               MinWidth="100"
               Name="Current Frame Only"
+              ShowImage="false"
               ShowText="false"
               Text="Current Frame Only"
               UID="ef4472c0-fa88-4379-af73-842e4b57f752" />
             <ads:RibbonCommandItem
               AllowInToolBar="false"
               CommandHandler="{wpfcust:ActionItem MacroName=Datasmith_ExportOptions_AnimatedTransformsActiveSegment, MacroCategory=Datasmith}"
-              Id="mOBJPaintMotionTypeConsecutive"
+              Image="pack://application:,,,/AdWindows;component/Resources/Images/MissingImage.ico"
               IsCheckable="true"
               IsToolTipEnabled="true"
               LargeImage="pack://application:,,,/AdWindows;component/Resources/Images/MissingImage.ico"
               MinWidth="100"
               Name="Active Time Segment"
+              ShowImage="false"
               ShowText="false"
               Text="Active Time Segment"
               UID="4097c080-33ee-430d-a47f-ab6e5c06c007" />
           </wpfcust:MaxRibbonSplitButton>
-
-
         </ads:RibbonRowPanel>
       </wpfcust:MaxRibbonPanelSource>
     </wpfcust:MaxModelingRibbonPanel>
     <wpfcust:MaxModelingRibbonPanel
-      CollapsedPanelImage="pack://application:,,,/AdWindows;component/Resources/Images/MissingImage.ico"
+      CollapsedPanelImage="Icons/Datasmith/DatasmithSettingsIcon32.png"
       FloatingGroup="0"
       FloatingOrder="0"
       ResizeStyle="NeverCollapsePanel"
@@ -275,15 +269,19 @@
           ShowText="false"
           Size="Large"
           UID="3a10e004-8897-4ca0-8824-0662af3c5097">
-          <ads:RibbonButton
+          <wpfcust:MaxRibbonToggleButton
+            CommandHandler="{wpfcust:ActionItem MacroName=Datasmith_AddModifier, MacroCategory=Datasmith}"
+            Id="Datasmith_AddModifier.Datasmith"
             IsCheckable="true"
-            IsToolTipEnabled="false"
+            IsToolTipEnabled="true"
+            LargeImage="Icons/MissingImage_32.png"
             MinWidth="0"
+            OriginalImage="Icons/MissingImage_16.png"
             ResizeStyle="HideText"
             ShowImage="false"
             ShowText="true"
-            Text="Add Attributes Modifier"
-            UID="cdee8c6d-56a9-4c9b-8776-7e19d8c2358e" />
+            Text="Add Datasmith Attributes Modifier"
+            UID="09bcd514-f707-4848-9bf4-32e73fe0d1c2" />
         </ads:RibbonRowPanel>
       </wpfcust:MaxRibbonPanelSource>
     </wpfcust:MaxModelingRibbonPanel>
