3ds Max DirectLink:  export options stored in an external file

-  'ini' file located in standard Max user 'plugcfg' folder e.g. for 3ds Max 2022 it's %localappdata%\Autodesk\3dsMax\2022 - 64bit\ENU\en-US\plugcfg\UnrealDatasmithMax.ini

- loaded on startup and saved every time an option is changed

diff --git a/Engine/Source/Programs/Enterprise/Datasmith/DatasmithMaxExporter/Private/DatasmithMaxExporterWithDirectLink.cpp b/Engine/Source/Programs/Enterprise/Datasmith/DatasmithMaxExporter/Private/DatasmithMaxExporterWithDirectLink.cpp
--- a/Engine/Source/Programs/Enterprise/Datasmith/DatasmithMaxExporter/Private/DatasmithMaxExporterWithDirectLink.cpp
+++ b/Engine/Source/Programs/Enterprise/Datasmith/DatasmithMaxExporter/Private/DatasmithMaxExporterWithDirectLink.cpp
@@ -396,9 +396,21 @@ class FPersistentExportOptions: public IPersistentExportOptions
 {
 public:
 
+	void Load()
+	{
+		if (bLoaded)
+		{
+			return;
+		}
+		GetBool(TEXT("SelectedOnly"), Options.bSelectedOnly);
+		GetBool(TEXT("AnimatedTransforms"), Options.bAnimatedTransforms);
+		bLoaded = true;
+	}
+
 	virtual void SetSelectedOnly(bool bValue) override
 	{
 		Options.bSelectedOnly = bValue;
+		SetBool(TEXT("SelectedOnly"), bValue);
 	}
 	virtual bool GetSelectedOnly() override
 	{
@@ -408,6 +420,7 @@ public:
 	virtual void SetAnimatedTransforms(bool bValue) override
 	{
 		Options.bAnimatedTransforms = bValue;
+		SetBool(TEXT("AnimatedTransforms"), bValue);
 	}
 
 	virtual bool GetAnimatedTransforms() override
@@ -415,7 +428,36 @@ public:
 		return Options.bAnimatedTransforms;
 	}
 
+	void GetBool(const TCHAR* Name, bool& bValue)
+	{
+		if (!GConfig)
+		{
+			return;
+		}
+		FString ConfigPath = GetConfigPath();
+		GConfig->GetBool(TEXT("Export"), Name, bValue, ConfigPath);
+	}
+
+	void SetBool(const TCHAR* Name, bool bValue)
+	{
+		if (!GConfig)
+		{
+			return;
+		}
+		FString ConfigPath = GetConfigPath();
+		GConfig->SetBool(TEXT("Export"), Name, bValue, ConfigPath);
+		GConfig->Flush(false, ConfigPath);
+	}
+
+	FString GetConfigPath()
+	{
+		FString PlugCfgPath = GetCOREInterface()->GetDir(APP_PLUGCFG_DIR);
+		return FPaths::Combine(PlugCfgPath, TEXT("UnrealDatasmithMax.ini"));
+	}
+
+
 	FExportOptions Options;
+	bool bLoaded = false;
 };
 
 // Holds states of entities for syncronization and handles change events
@@ -1950,6 +1992,9 @@ bool CreateExporter(bool bEnableUI, const TCHAR* EnginePath)
 
 	static FExportOptions ExporterOptions; // Default options
 	Exporter = MakeUnique<FExporter>(ExporterOptions);
+
+	PersistentExportOptions.Load(); // Access GConfig only after FDatasmithExporterManager::Initialize finishes, which ensures that Unreal game thread was initialized(GConfig is created there)
+
 	return true;
 }
 
