3ds Max DirectLink: Messages Dialog, crash on exit fix

- Messages dialog first iteration - modeless, realtime update

diff --git a/Engine/Source/Programs/Enterprise/Datasmith/DatasmithMaxExporter/Private/DatasmithMaxDirectLinkMaxScript.cpp b/Engine/Source/Programs/Enterprise/Datasmith/DatasmithMaxExporter/Private/DatasmithMaxDirectLinkMaxScript.cpp
--- a/Engine/Source/Programs/Enterprise/Datasmith/DatasmithMaxExporter/Private/DatasmithMaxDirectLinkMaxScript.cpp
+++ b/Engine/Source/Programs/Enterprise/Datasmith/DatasmithMaxExporter/Private/DatasmithMaxDirectLinkMaxScript.cpp
@@ -3,15 +3,16 @@
 #ifdef NEW_DIRECTLINK_PLUGIN
 
 #include "DatasmithMaxDirectLink.h"
+#include "DatasmithMaxLogger.h"
 
-#include "DatasmithExporterManager.h"
-
-
-#include "Logging/LogMacros.h"
+#include "Async/Async.h"
+#include "Containers/Queue.h"
 
 #include "Resources/Windows/resource.h"
 
 #include "Windows/AllowWindowsPlatformTypes.h"
+
+
 MAX_INCLUDES_START
 	#include "impexp.h"
 	#include "max.h"
@@ -24,9 +25,10 @@ MAX_INCLUDES_START
 	#include "maxscript\macros\define_instantiation_functions.h"
 
 	#include <maxicon.h>
-
 MAX_INCLUDES_END
 
+#include "windows.h"
+
 extern HINSTANCE HInstanceMax;
 
 
@@ -35,6 +37,8 @@ namespace DatasmithMaxDirectLink
 
 /************************************* MaxScript exports *********************************/
 
+void InitScripts();
+
 Value* OnLoad_cf(Value**, int);
 Primitive OnLoad_pf(_M("Datasmith_OnLoad"), OnLoad_cf);
 
@@ -48,6 +52,8 @@ Value* OnLoad_cf(Value **arg_list, int count)
 
 	const TCHAR* EnginePathUnreal = (const TCHAR*)pEnginePath->to_string();
 
+	InitScripts();
+
 	if (CreateExporter(bEnableUI, EnginePathUnreal))
 	{
 		return &false_value;
@@ -312,6 +318,126 @@ Value* LogInfo_cf(Value** arg_list, int count)
 }
 Primitive LogInfo_pf(_M("Datasmith_LogInfo"), LogInfo_cf);
 
+
+class FMessagesDialog
+{
+public:
+
+	~FMessagesDialog()
+	{
+		ensure(false);
+		if (Thread.IsValid())
+		{
+			ThreadEvent->Trigger();
+			Thread.Get();
+		}
+	}
+
+	static INT_PTR CALLBACK DlgProc(HWND hDlg, UINT iMsg, WPARAM wParam, LPARAM lParam)
+	{
+		switch (iMsg)
+		{
+			case WM_INITDIALOG:
+			{
+				return TRUE;
+			}
+
+			case WM_CLOSE:
+			{
+				ShowWindow(hDlg, SW_HIDE);
+				return FALSE;
+			}
+		};
+		return FALSE;
+	}
+
+	void PumpWindowMessages()
+	{
+		MSG Message;
+		// standard Windows message handling
+		while(PeekMessage(&Message, NULL, 0, 0, PM_REMOVE))
+		{
+			if (!IsDialogMessage(DialogHwnd, &Message))
+			{
+				TranslateMessage(&Message);
+				DispatchMessage(&Message);
+			}
+		}
+	}
+
+	void ProcessLogMessages()
+	{
+		if (Messages.IsEmpty())
+		{
+			return;
+		}
+
+		FString Message;
+		while(Messages.Dequeue(Message))
+		{
+			SendDlgItemMessage(DialogHwnd, IDC_ERROR_MSG_LIST, LB_ADDSTRING, NULL, reinterpret_cast<LPARAM>(*Message));
+		}
+
+		// Scroll message list to bottom
+		LRESULT ItemCount = SendDlgItemMessage(DialogHwnd, IDC_ERROR_MSG_LIST, LB_GETCOUNT,NULL,NULL);
+		SendDlgItemMessage(DialogHwnd, IDC_ERROR_MSG_LIST, LB_SETTOPINDEX,ItemCount-1,NULL);
+	}
+
+	void Toggle()
+	{
+		if (!bDialogCreated)
+		{
+			Thread = Async(EAsyncExecution::Thread,
+				[&, this]
+				{
+					DialogHwnd = CreateDialogParam(HInstanceMax, MAKEINTRESOURCE(IDD_ERROR_MSGS), GetCOREInterface()->GetMAXHWnd(), DlgProc, reinterpret_cast<LPARAM>(this));
+					while (true)
+					{
+						PumpWindowMessages();
+
+						ProcessLogMessages();
+
+						if (ThreadEvent->Wait(FTimespan::FromMilliseconds(10)))
+						{
+							DestroyWindow(DialogHwnd);
+							break;
+						}
+					}
+				}
+			);
+			bDialogCreated = true;
+		}
+		else
+		{
+			ShowWindow(DialogHwnd, IsWindowVisible(DialogHwnd) ? SW_HIDE : SW_SHOW);
+		}
+	}
+
+	void AddWarning(const TCHAR* Message)
+	{
+		Messages.Enqueue(Message);
+	}
+
+	bool bDialogCreated = false;
+	HWND DialogHwnd = NULL;
+	FEventRef ThreadEvent;
+	TFuture<void> Thread;
+
+	TQueue<FString> Messages;
+};
+
+TUniquePtr<FMessagesDialog> MessagesDialog;
+
+void LogWarningDialog(const FString& Msg)
+{
+	LogWarningDialog(*Msg);
+}
+
+void LogWarningDialog(const TCHAR* Msg)
+{
+	MessagesDialog->AddWarning(Msg);
+}
+
 // Setup ActionTable with Datasmith commands exposed as actions
 class FDatasmithActions
 {
@@ -388,12 +514,12 @@ public:
 		}
 	};
 
-
 	class FDatasmithActionCallback : public ActionCallback
 	{
 	public:
 
-		BOOL ExecuteAction (int ActionId) {
+		BOOL ExecuteAction (int ActionId)
+		{
 			LogDebug(FString::Printf(TEXT("Action: %d"), ActionId));
 
 			switch (ActionId)
@@ -435,6 +561,7 @@ public:
 			}
 			case ID_SHOWLOG_ACTION_ID:
 			{
+				MessagesDialog->Toggle();
 				return true;
 			}
 			}
@@ -444,7 +571,6 @@ public:
 
 	FDatasmithActions()
 		: Name(TEXT("Datasmith"))
-		, Table(Name)
 	{
 		// todo: localization of Name
 
@@ -462,24 +588,36 @@ public:
 			DATASMITH_ACTION(SHOWLOG),
 		};
 
-		Table.BuildActionTable(nullptr, sizeof(ActionsDescriptions) / sizeof(ActionsDescriptions[0]), ActionsDescriptions, HInstanceMax);
+		FDatasmithActionTable* Table = new FDatasmithActionTable(Name); // Table, registered with RegisterActionTable will be deallocated by Max
+
+		Table->BuildActionTable(nullptr, sizeof(ActionsDescriptions) / sizeof(ActionsDescriptions[0]), ActionsDescriptions, HInstanceMax);
 		GetCOREInterface()->GetActionManager()->RegisterActionContext(ActionContextId, Name.data());
 
 		// Register table - this needs to be called explicitly when action table is not returned to Max with ClassDesc's GetActionTable method
-		GetCOREInterface()->GetActionManager()->RegisterActionTable(&Table);
+		GetCOREInterface()->GetActionManager()->RegisterActionTable(Table); 
 
 		GetCOREInterface()->GetActionManager()->ActivateActionTable(&ActionCallback, ActionTableId);
 	}
 
 private:
 	TSTR Name;
-	FDatasmithActionTable Table;
 	FDatasmithActionCallback ActionCallback;
 };
 
 
 TUniquePtr<FDatasmithActions> Actions;
 
+void InitScripts()
+{
+	MessagesDialog.Reset(new FMessagesDialog);
+}
+
+void ShutdownScripts()
+{
+	MessagesDialog.Reset();
+	Actions.Reset();
+}
+
 Value* SetupActions_cf(Value** arg_list, int count)
 {
 	if (!Actions)
@@ -492,6 +630,19 @@ Value* SetupActions_cf(Value** arg_list, int count)
 Primitive SetupActions_pf(_M("Datasmith_SetupActions"), SetupActions_cf);
 
 
+Value* AddWarning_cf(Value**, int);
+Primitive AddWarning_pf(_M("Datasmith_AddWarning"), AddWarning_cf);
+
+Value* AddWarning_cf(Value** arg_list, int count)
+{
+	check_arg_count(AddWarning, 1, count);
+
+	MessagesDialog->AddWarning(arg_list[0]->to_string());
+
+	return bool_result(true);
+}
+
+
 }
 
 #include "Windows/HideWindowsPlatformTypes.h"
diff --git a/Engine/Source/Programs/Enterprise/Datasmith/DatasmithMaxExporter/Private/DatasmithMaxDirectLinkMesh.cpp b/Engine/Source/Programs/Enterprise/Datasmith/DatasmithMaxExporter/Private/DatasmithMaxDirectLinkMesh.cpp
--- a/Engine/Source/Programs/Enterprise/Datasmith/DatasmithMaxExporter/Private/DatasmithMaxDirectLinkMesh.cpp
+++ b/Engine/Source/Programs/Enterprise/Datasmith/DatasmithMaxExporter/Private/DatasmithMaxDirectLinkMesh.cpp
@@ -427,15 +427,23 @@ bool ConvertMaxMeshToDatasmith(ISceneTracker& Scene, TSharedPtr<IDatasmithMeshEl
 	if (DatasmithAttributes && DatasmithAttributes->GetExportMode() == EStaticMeshExportMode::BoundingBox)
 	{
 		FDatasmithMesh DatasmithMesh;
-		FillDatasmithMeshFromBoundingBox(DatasmithMesh, Node, FTransform::Identity);
-		Scene.AddMeshElement(DatasmithMeshElement, DatasmithMesh, nullptr);
-		return true;
+		if (FillDatasmithMeshFromBoundingBox(DatasmithMesh, Node, FTransform::Identity))
+		{
+			Scene.AddMeshElement(DatasmithMeshElement, DatasmithMesh, nullptr);
+			return true;
+		}
+		else
+		{
+			LogWarningDialog(FString(TEXT("Invalid object: ")) + Node->GetName());
+			return false;
+		}
 	}
 
 	FDatasmithMesh DatasmithMesh;
 	TMap<int32, int32> UVChannelsMap;
 	if (!CreateDatasmithMeshFromMaxMesh(DatasmithMesh, Node, MeshName, RenderMesh, SupportedChannels, UVChannelsMap))
 	{
+		LogWarningDialog(FString(TEXT("Invalid object: ")) + Node->GetName());
 		return false;
 	}
 
@@ -453,7 +461,7 @@ bool ConvertMaxMeshToDatasmith(ISceneTracker& Scene, TSharedPtr<IDatasmithMeshEl
 		}
 		else if (SelectedLightmapUVChannel != DefaultValue)
 		{
-			DatasmithMaxLogger::Get().AddGeneralError(*FString::Printf(TEXT("%s won't use the channel %i for its lightmap because it's not supported by the mesh. A new channel will be generated.")
+			LogWarningDialog(*FString::Printf(TEXT("%s won't use the channel %i for its lightmap because it's not supported by the mesh. A new channel will be generated.")
 				, static_cast<const TCHAR*>(Node->GetName())
 				, SelectedLightmapUVChannel));
 		}
diff --git a/Engine/Source/Programs/Enterprise/Datasmith/DatasmithMaxExporter/Private/DatasmithMaxDirectLinkNotifications.cpp b/Engine/Source/Programs/Enterprise/Datasmith/DatasmithMaxExporter/Private/DatasmithMaxDirectLinkNotifications.cpp
--- a/Engine/Source/Programs/Enterprise/Datasmith/DatasmithMaxExporter/Private/DatasmithMaxDirectLinkNotifications.cpp
+++ b/Engine/Source/Programs/Enterprise/Datasmith/DatasmithMaxExporter/Private/DatasmithMaxDirectLinkNotifications.cpp
@@ -508,10 +508,8 @@ void FNotifications::On3dsMaxNotification(void* param, NotifyInfo* info)
 	case NOTIFY_MTL_REFDELETED:
 		break;
 
-		// This one crashes when calling LogInfo
-	case NOTIFY_PLUGINS_PRE_SHUTDOWN:
-		Exporter->Shutdown();
-		break;
+	case NOTIFY_PLUGINS_PRE_SHUTDOWN: // This one crashes when calling LogInfo
+		return;
 	default:
 		LOG_DEBUG_HEAVY(FString(TEXT("Notify: ")) + NotificationsHandler.ConvertNotificationCodeToString(info->intcode));
 	};
diff --git a/Engine/Source/Programs/Enterprise/Datasmith/DatasmithMaxExporter/Private/DatasmithMaxExporterWithDirectLink.cpp b/Engine/Source/Programs/Enterprise/Datasmith/DatasmithMaxExporter/Private/DatasmithMaxExporterWithDirectLink.cpp
--- a/Engine/Source/Programs/Enterprise/Datasmith/DatasmithMaxExporter/Private/DatasmithMaxExporterWithDirectLink.cpp
+++ b/Engine/Source/Programs/Enterprise/Datasmith/DatasmithMaxExporter/Private/DatasmithMaxExporterWithDirectLink.cpp
@@ -1519,7 +1519,7 @@ public:
 				}
 				else
 				{
-					DatasmithMaxLogger::Get().AddUnsupportedLight(NodeTracker.Node);
+					LogWarningDialog(TEXT("Unsupported light: ") + DatasmithMaxLogger::Get().GetLightDescription(NodeTracker.Node));
 				}
 				return false;
 			}
@@ -2000,6 +2000,7 @@ bool CreateExporter(bool bEnableUI, const TCHAR* EnginePath)
 
 void ShutdownExporter()
 {
+	ShutdownScripts();
 	Exporter.Reset();
 	FDatasmithDirectLink::Shutdown();
 	FDatasmithExporterManager::Shutdown();
diff --git a/Engine/Source/Programs/Enterprise/Datasmith/DatasmithMaxExporter/Private/DatasmithMaxLogger.cpp b/Engine/Source/Programs/Enterprise/Datasmith/DatasmithMaxExporter/Private/DatasmithMaxLogger.cpp
--- a/Engine/Source/Programs/Enterprise/Datasmith/DatasmithMaxExporter/Private/DatasmithMaxLogger.cpp
+++ b/Engine/Source/Programs/Enterprise/Datasmith/DatasmithMaxExporter/Private/DatasmithMaxLogger.cpp
@@ -222,6 +222,15 @@ void DatasmithMaxLogger::AddObjectList(TArray< INode* > ObjectList, HWND Handle,
 	}
 }
 
+FString DatasmithMaxLogger::GetLightDescription(INode* LightNode)
+{
+	ObjectState State = LightNode->EvalWorldState(0);
+	LightObject *Light = (LightObject*)State.obj;
+	MSTR Classname;
+	Light->GetClassName(Classname);
+	return FString::Printf(TEXT("\"%s\" of type %s (0x%08x-0x%08x)"), LightNode->GetName(), Classname.ToBSTR(), Light->ClassID().PartA(), Light->ClassID().PartB());
+}
+
 void DatasmithMaxLogger::Show(HWND Handle)
 {
 	ShowMessage = TEXT("");
@@ -343,11 +352,7 @@ void DatasmithMaxLogger::Show(HWND Handle)
 
 		for (int i = 0; i < UnsupportedLight.Num(); i++)
 		{
-			ObjectState State = UnsupportedLight[i]->EvalWorldState(0);
-			LightObject *Light = (LightObject*)State.obj;
-			MSTR Classname;
-			Light->GetClassName(Classname);
-			FString Msg = FString::Printf(TEXT("\"%s\" of type %s (0x%08x-0x%08x)"), UnsupportedLight[i]->GetName(), Classname.ToBSTR(), Light->ClassID().PartA(), Light->ClassID().PartB());
+			FString Msg = GetLightDescription(UnsupportedLight[i]);
 			AddItem(*Msg, Handle, ShowMessage);
 		}
 		AddItem(TEXT("\n\n"), Handle, ShowMessage);
diff --git a/Engine/Source/Programs/Enterprise/Datasmith/DatasmithMaxExporter/Public/DatasmithMaxDirectLink.h b/Engine/Source/Programs/Enterprise/Datasmith/DatasmithMaxExporter/Public/DatasmithMaxDirectLink.h
--- a/Engine/Source/Programs/Enterprise/Datasmith/DatasmithMaxExporter/Public/DatasmithMaxDirectLink.h
+++ b/Engine/Source/Programs/Enterprise/Datasmith/DatasmithMaxExporter/Public/DatasmithMaxDirectLink.h
@@ -83,6 +83,8 @@ bool CreateExporter(bool bEnableUI, const TCHAR* EnginePath); // Create exporter
 IExporter* GetExporter();
 void ShutdownExporter();
 bool Export(const TCHAR* Name, const TCHAR* OutputPath, bool bQuiet);
+void ShutdownScripts();
+
 
 IPersistentExportOptions& GetPersistentExportOptions();
 
@@ -419,6 +421,9 @@ void LogDebugNode(const FString& Name, class INode* Node);
 void LogNodeEvent(const MCHAR* Name, INodeEventCallback::NodeKeyTab& nodes);
 void LogFlush();
 
+void LogWarningDialog(const TCHAR* Msg);
+void LogWarningDialog(const FString& Msg);
+
 }
 
 
diff --git a/Engine/Source/Programs/Enterprise/Datasmith/DatasmithMaxExporter/Public/DatasmithMaxLogger.h b/Engine/Source/Programs/Enterprise/Datasmith/DatasmithMaxExporter/Public/DatasmithMaxLogger.h
--- a/Engine/Source/Programs/Enterprise/Datasmith/DatasmithMaxExporter/Public/DatasmithMaxLogger.h
+++ b/Engine/Source/Programs/Enterprise/Datasmith/DatasmithMaxExporter/Public/DatasmithMaxLogger.h
@@ -39,6 +39,8 @@ public:
 	void Show(HWND hDlg);
 	bool CopyToClipBoard();
 
+	FString GetLightDescription(INode* LightNode);
+
 private:
 	void AddItem(const TCHAR* Msg, HWND Handle, FString& FullMsg);
 	void AddObjectList(TArray< INode* > ObjectList, HWND Handle, const TCHAR* Header, const TCHAR* Description = nullptr);
