Navisworks - added separate warning dialog when trying to export an empty state

diff --git a/Engine/Source/Programs/Enterprise/Datasmith/DatasmithNavisworksExporter/Private/DatasmithNavisworksPlugin.cs b/Engine/Source/Programs/Enterprise/Datasmith/DatasmithNavisworksExporter/Private/DatasmithNavisworksPlugin.cs
--- a/Engine/Source/Programs/Enterprise/Datasmith/DatasmithNavisworksExporter/Private/DatasmithNavisworksPlugin.cs
+++ b/Engine/Source/Programs/Enterprise/Datasmith/DatasmithNavisworksExporter/Private/DatasmithNavisworksPlugin.cs
@@ -667,6 +667,12 @@ namespace DatasmithNavisworks
 					return -1;
 				}
 
+				if (ActiveDocument.IsClear)
+				{
+					DisplayWarning("Opened document is required for Datasmith Export");
+					return -1;
+				}
+
 				ExportParameters ExportParameters = ParseParameters(Parameters, out string FilePath);
 
 				// todo: just check for FilePath existence/null/empty?
