Fix build - Datasmith SketchUp Exporter using removed FStreamInfo member.

diff --git a/Engine/Source/Programs/Enterprise/Datasmith/DatasmithSketchUpRubyExporter/Private/DatasmithSketchUpExporter.cpp b/Engine/Source/Programs/Enterprise/Datasmith/DatasmithSketchUpRubyExporter/Private/DatasmithSketchUpExporter.cpp
--- a/Engine/Source/Programs/Enterprise/Datasmith/DatasmithSketchUpRubyExporter/Private/DatasmithSketchUpExporter.cpp
+++ b/Engine/Source/Programs/Enterprise/Datasmith/DatasmithSketchUpRubyExporter/Private/DatasmithSketchUpExporter.cpp
@@ -183,7 +183,7 @@ public:
 
 			for (const FRawInfo::FStreamInfo& StreamInfo : RawInfo.StreamsInfo)
 			{
-				if (!StreamInfo.bIsActive)
+				if (StreamInfo.ConnectionState != EStreamConnectionState::Active)
 				{
 					continue;
 				}
