Fix DatasmithSketchupExporter missing libraries

JeanLuc.Corenthin

diff --git a/Engine/Source/Programs/Enterprise/Datasmith/DatasmithSketchUpRubyExporter/Build/build_mac.sh b/Engine/Source/Programs/Enterprise/Datasmith/DatasmithSketchUpRubyExporter/Build/build_mac.sh
--- a/Engine/Source/Programs/Enterprise/Datasmith/DatasmithSketchUpRubyExporter/Build/build_mac.sh
+++ b/Engine/Source/Programs/Enterprise/Datasmith/DatasmithSketchUpRubyExporter/Build/build_mac.sh
@@ -24,6 +24,8 @@ rm -rf "$IntermediatePath"
 mkdir -p "$IntermediatePath"
 
 dylibLibFreeImage=libfreeimage-3.18.0.dylib
+dylibLibTbb=libtbb.dylib
+dylibLibTbbMalloc=libtbbmalloc.dylib
 
 # Copy and fixup dylibs
 Dylibs="$IntermediatePath/Dylibs"
@@ -32,6 +34,8 @@ mkdir -p "$Dylibs"
 
 cp "$EnginePath/Binaries/ThirdParty/FreeImage/Mac/$dylibLibFreeImage" "$Dylibs"
 cp "$EnginePath/Binaries/Mac/DatasmithSDK/DatasmithSDK.dylib" "$Dylibs"
+cp "$EnginePath/Binaries/ThirdParty/Intel/TBB/Mac/$dylibLibTbb" "$Dylibs"
+cp "$EnginePath/Binaries/ThirdParty/Intel/TBB/Mac/$dylibLibTbbMalloc" "$Dylibs"
 
 chmod 777 "$Dylibs/$dylibLibFreeImage"
 
