DatasmithRhinoExporter - Adding missing TBB dependencies to Mac package.

JeanLuc.Corenthin

diff --git a/Engine/Source/Programs/Enterprise/Datasmith/DatasmithRhinoExporter/Installer/CreateMacPackages.sh b/Engine/Source/Programs/Enterprise/Datasmith/DatasmithRhinoExporter/Installer/CreateMacPackages.sh
--- a/Engine/Source/Programs/Enterprise/Datasmith/DatasmithRhinoExporter/Installer/CreateMacPackages.sh
+++ b/Engine/Source/Programs/Enterprise/Datasmith/DatasmithRhinoExporter/Installer/CreateMacPackages.sh
@@ -5,6 +5,7 @@ PLUGIN_VERSION=$2
 SCRIPT_DIR=$(builtin cd $(dirname $0); pwd)
 PLUGIN_BASE_DIR="$(dirname $SCRIPT_DIR)"
 ENGINE_DIR="$PLUGIN_BASE_DIR/../../../../.."
+TBB_BINARY_DIR="$ENGINE_DIR/Binaries/ThirdParty/Intel/TBB/Mac"
 PLUGIN_BINARY_DIR="$ENGINE_DIR/Binaries/Mac/Rhino/$PLUGIN_VERSION"
 DATASMITH_FACADE_BINARY_DIR="$ENGINE_DIR/Binaries/Mac/DatasmithFacadeCSharp"
 PLUGIN_RHP_PATH="$PLUGIN_BINARY_DIR/$PLUGIN_NAME.rhp"
@@ -17,6 +18,7 @@ rm -f "$PLUGIN_MACRHI_PATH"
 echo "Creating .rhp package"
 mkdir "$PLUGIN_RHP_PATH"
 cp "$DATASMITH_FACADE_BINARY_DIR/"*".dylib" "$PLUGIN_RHP_PATH"
+cp "$TBB_BINARY_DIR/"*".dylib" "$PLUGIN_RHP_PATH"
 cp "$PLUGIN_BINARY_DIR/$PLUGIN_NAME.dll" "$PLUGIN_RHP_PATH/$PLUGIN_NAME.rhp"
 cp "$PLUGIN_BASE_DIR/Config/DatasmithRhino.rhp.config" "$PLUGIN_RHP_PATH/$PLUGIN_NAME.rhp.config"
 
