Fix DatasmithArchiCAD exporter build.

diff --git a/Engine/Source/Programs/Enterprise/Datasmith/DatasmithARCHICADExporter/Private/ConnectionWindow.cpp b/Engine/Source/Programs/Enterprise/Datasmith/DatasmithARCHICADExporter/Private/ConnectionWindow.cpp
--- a/Engine/Source/Programs/Enterprise/Datasmith/DatasmithARCHICADExporter/Private/ConnectionWindow.cpp
+++ b/Engine/Source/Programs/Enterprise/Datasmith/DatasmithARCHICADExporter/Private/ConnectionWindow.cpp
@@ -73,7 +73,7 @@ inline bool operator==(const DirectLink::FRawInfo::FStreamInfo& InStreamInfo1,
 {
 	return InStreamInfo1.StreamId == InStreamInfo2.StreamId && InStreamInfo1.Source == InStreamInfo2.Source &&
 		   InStreamInfo1.Destination == InStreamInfo2.Destination &&
-		   InStreamInfo1.bIsActive == InStreamInfo2.bIsActive &&
+		   InStreamInfo1.ConnectionState == InStreamInfo2.ConnectionState &&
 		   InStreamInfo1.CommunicationStatus == InStreamInfo2.CommunicationStatus;
 }
 
@@ -279,7 +279,7 @@ class FConnectionDialog : public DG::Palette,
 			FString ThisText;
 			for (const DirectLink::FRawInfo::FStreamInfo& StreamInfo : CurrentStatus.StreamsInfo)
 			{
-				if (StreamInfo.bIsActive)
+				if (StreamInfo.ConnectionState == DirectLink::EStreamConnectionState::Active)
 				{
 					const DirectLink::FRawInfo::FDataPointId** SourceDataPointId = SourcesMap.Find(StreamInfo.Source);
 					const FDestination* DestinationDataPointId = DestinationsMap.Find(StreamInfo.Destination);
