DatasmithArchiCADExporter Mac - Adding missing runtime thirdparty libraries

JeanLuc.Corenthin

diff --git a/Engine/Source/Programs/Enterprise/Datasmith/DatasmithARCHICADExporter/Build/Config/setupdylibs.sh b/Engine/Source/Programs/Enterprise/Datasmith/DatasmithARCHICADExporter/Build/Config/setupdylibs.sh
--- a/Engine/Source/Programs/Enterprise/Datasmith/DatasmithARCHICADExporter/Build/Config/setupdylibs.sh
+++ b/Engine/Source/Programs/Enterprise/Datasmith/DatasmithARCHICADExporter/Build/Config/setupdylibs.sh
@@ -38,18 +38,24 @@ OurDylibFolder=$projectPath/Dylibs
 mkdir -p "$OurDylibFolder"
 
 dylibLibFreeImage=libfreeimage-3.18.0.dylib
+dylibtbb=libtbb.dylib
+dylibtbbmalloc=libtbbmalloc.dylib
 
-if [[ "$EnginePath/Binaries/ThirdParty/FreeImage/Mac/$dylibLibFreeImage" -nt "$OurDylibFolder/$dylibLibFreeImage" ]]; then
-	if [ -f "$OurDylibFolder/$dylibLibFreeImage" ]; then
-		unlink "$OurDylibFolder/$dylibLibFreeImage"
+SetUpThirdPartyDll() {
+	DylibName=$1
+	DylibPath=$2
+	if [[ "$DylibPath" -nt "$OurDylibFolder/$DylibName" ]]; then
+		if [ -f "$OurDylibFolder/$DylibName" ]; then
+			unlink "$OurDylibFolder/$DylibName"
+		fi
 	fi
-fi
-if [ ! -f "$OurDylibFolder/$dylibLibFreeImage" ]; then
-    echo "Copy $dylibLibFreeImage"
-    cp "$EnginePath/Binaries/ThirdParty/FreeImage/Mac/$dylibLibFreeImage" "$OurDylibFolder"
-	chmod +w "$OurDylibFolder/$dylibLibFreeImage"
-    install_name_tool -id @loader_path/$dylibLibFreeImage "$OurDylibFolder/$dylibLibFreeImage"
-fi
+	if [ ! -f "$OurDylibFolder/$DylibName" ]; then
+		echo "Copy $DylibName"
+		cp "$DylibPath" "$OurDylibFolder"
+		chmod +w "$OurDylibFolder/$DylibName"
+		install_name_tool -id @loader_path/$DylibName "$OurDylibFolder/$DylibName"
+	fi
+}
 
 SetUpDll() {
 	DylibName=$1
@@ -72,5 +78,9 @@ SetUpDll() {
 	fi
 }
 
+SetUpThirdPartyDll $dylibLibFreeImage "$EnginePath/Binaries/ThirdParty/FreeImage/Mac/$dylibLibFreeImage"
+SetUpThirdPartyDll $dylibtbb "$EnginePath/Binaries/ThirdParty/Intel/TBB/Mac/$dylibtbb"
+SetUpThirdPartyDll $dylibtbbmalloc "$EnginePath/Binaries/ThirdParty/Intel/TBB/Mac/$dylibtbbmalloc"
+
 SetUpDll DatasmithUE4ArchiCAD.dylib
 SetUpDll DatasmithUE4ArchiCAD-Mac-Debug.dylib
