DatasmithArchiCADExporter - Fix cache path textbox not updated after changing the cache directory.

JeanLuc.Corenthin

diff --git a/Engine/Source/Programs/Enterprise/Datasmith/DatasmithARCHICADExporter/Private/ConnectionWindow.cpp b/Engine/Source/Programs/Enterprise/Datasmith/DatasmithARCHICADExporter/Private/ConnectionWindow.cpp
--- a/Engine/Source/Programs/Enterprise/Datasmith/DatasmithARCHICADExporter/Private/ConnectionWindow.cpp
+++ b/Engine/Source/Programs/Enterprise/Datasmith/DatasmithARCHICADExporter/Private/ConnectionWindow.cpp
@@ -360,6 +360,7 @@ class FConnectionDialog : public DG::Palette,
 				{
 					UE_AC_TraceF("FConnectionDialog::ButtonClicked - New cache folder is \"%s\"\n", CachePath.ToUtf8());
 					FSyncDatabase::SetCachePath(CachePath);
+					UpdateCacheDirectoryText();
 				}
 				else
 				{
@@ -373,6 +374,11 @@ class FConnectionDialog : public DG::Palette,
 			UE_AC::ShellOpenDocument(FSyncDatabase::GetCachePath().ToUtf8());
 		}
 	}
+	
+	void UpdateCacheDirectoryText()
+	{
+		CacheFolderText.SetText(GetGSName(kName_CacheDirectory) + FSyncDatabase::GetCachePath());
+	}
 };
 
 FConnectionDialog::FConnectionDialog()
@@ -430,7 +436,7 @@ FConnectionDialog::FConnectionDialog()
 	Attach(*this);
 	AttachToAllItems(*this);
 
-	CacheFolderText.SetText(GetGSName(kName_CacheDirectory) + FSyncDatabase::GetCachePath());
+	UpdateCacheDirectoryText();
 
 	bool SendForInactiveApp = false;
 	EnableIdleEvent(SendForInactiveApp);
