Datasmith Revit: Export User Defined Collection of Metadata

diff --git a/Engine/Source/Programs/Enterprise/Datasmith/DatasmithRevitExporter/DatasmithRevit2018/DatasmithRevit2018.csproj b/Engine/Source/Programs/Enterprise/Datasmith/DatasmithRevitExporter/DatasmithRevit2018/DatasmithRevit2018.csproj
--- a/Engine/Source/Programs/Enterprise/Datasmith/DatasmithRevitExporter/DatasmithRevit2018/DatasmithRevit2018.csproj
+++ b/Engine/Source/Programs/Enterprise/Datasmith/DatasmithRevitExporter/DatasmithRevit2018/DatasmithRevit2018.csproj
@@ -87,8 +87,11 @@
     <Compile Include="..\Private\DatasmithRevitMaterial.cs">
       <Link>DatasmithRevitExporter\DatasmithRevitMaterial.cs</Link>
     </Compile>
-    <Compile Include="..\Private\DatasmithRevitMetadataManager.cs">
-      <Link>DatasmithRevitExporter\DatasmithRevitMetdataManager.cs</Link>
+    <Compile Include="..\Private\DatasmithRevitSettingsManager.cs">
+      <Link>DatasmithRevitExporter\DatasmithRevitSettingsManager.cs</Link>
+    </Compile>
+    <Compile Include="..\Private\DatasmithRevitUtils.cs">
+      <Link>DatasmithRevitExporter\DatasmithRevitUtils.cs</Link>
     </Compile>
     <Compile Include="..\Private\DatasmithRevitSettingsDialog.cs">
       <Link>DatasmithRevitExporter\DatasmithRevitSettingsDialog.cs</Link>
diff --git a/Engine/Source/Programs/Enterprise/Datasmith/DatasmithRevitExporter/DatasmithRevit2019/DatasmithRevit2019.csproj b/Engine/Source/Programs/Enterprise/Datasmith/DatasmithRevitExporter/DatasmithRevit2019/DatasmithRevit2019.csproj
--- a/Engine/Source/Programs/Enterprise/Datasmith/DatasmithRevitExporter/DatasmithRevit2019/DatasmithRevit2019.csproj
+++ b/Engine/Source/Programs/Enterprise/Datasmith/DatasmithRevitExporter/DatasmithRevit2019/DatasmithRevit2019.csproj
@@ -87,8 +87,11 @@
     <Compile Include="..\Private\DatasmithRevitMaterial.cs">
       <Link>DatasmithRevitExporter\DatasmithRevitMaterial.cs</Link>
     </Compile>
-    <Compile Include="..\Private\DatasmithRevitMetadataManager.cs">
-      <Link>DatasmithRevitExporter\DatasmithRevitMetdataManager.cs</Link>
+    <Compile Include="..\Private\DatasmithRevitSettingsManager.cs">
+      <Link>DatasmithRevitExporter\DatasmithRevitSettingsManager.cs</Link>
+    </Compile>
+    <Compile Include="..\Private\DatasmithRevitUtils.cs">
+      <Link>DatasmithRevitExporter\DatasmithRevitUtils.cs</Link>
     </Compile>
     <Compile Include="..\Private\DatasmithRevitSettingsDialog.cs">
       <Link>DatasmithRevitExporter\DatasmithRevitSettingsDialog.cs</Link>
diff --git a/Engine/Source/Programs/Enterprise/Datasmith/DatasmithRevitExporter/DatasmithRevit2020/DatasmithRevit2020.csproj b/Engine/Source/Programs/Enterprise/Datasmith/DatasmithRevitExporter/DatasmithRevit2020/DatasmithRevit2020.csproj
--- a/Engine/Source/Programs/Enterprise/Datasmith/DatasmithRevitExporter/DatasmithRevit2020/DatasmithRevit2020.csproj
+++ b/Engine/Source/Programs/Enterprise/Datasmith/DatasmithRevitExporter/DatasmithRevit2020/DatasmithRevit2020.csproj
@@ -87,8 +87,11 @@
     <Compile Include="..\Private\DatasmithRevitMaterial.cs">
       <Link>DatasmithRevitExporter\DatasmithRevitMaterial.cs</Link>
     </Compile>
-    <Compile Include="..\Private\DatasmithRevitMetadataManager.cs">
-      <Link>DatasmithRevitExporter\DatasmithRevitMetdataManager.cs</Link>
+    <Compile Include="..\Private\DatasmithRevitSettingsManager.cs">
+      <Link>DatasmithRevitExporter\DatasmithRevitSettingsManager.cs</Link>
+    </Compile>
+    <Compile Include="..\Private\DatasmithRevitUtils.cs">
+      <Link>DatasmithRevitExporter\DatasmithRevitUtils.cs</Link>
     </Compile>
     <Compile Include="..\Private\DatasmithRevitSettingsDialog.cs">
       <Link>DatasmithRevitExporter\DatasmithRevitSettingsDialog.cs</Link>
diff --git a/Engine/Source/Programs/Enterprise/Datasmith/DatasmithRevitExporter/DatasmithRevit2021/DatasmithRevit2021.csproj b/Engine/Source/Programs/Enterprise/Datasmith/DatasmithRevitExporter/DatasmithRevit2021/DatasmithRevit2021.csproj
--- a/Engine/Source/Programs/Enterprise/Datasmith/DatasmithRevitExporter/DatasmithRevit2021/DatasmithRevit2021.csproj
+++ b/Engine/Source/Programs/Enterprise/Datasmith/DatasmithRevitExporter/DatasmithRevit2021/DatasmithRevit2021.csproj
@@ -87,8 +87,11 @@
     <Compile Include="..\Private\DatasmithRevitMaterial.cs">
       <Link>DatasmithRevitExporter\DatasmithRevitMaterial.cs</Link>
     </Compile>
-    <Compile Include="..\Private\DatasmithRevitMetadataManager.cs">
-      <Link>DatasmithRevitExporter\DatasmithRevitMetdataManager.cs</Link>
+    <Compile Include="..\Private\DatasmithRevitSettingsManager.cs">
+      <Link>DatasmithRevitExporter\DatasmithRevitSettingsManager.cs</Link>
+    </Compile>
+    <Compile Include="..\Private\DatasmithRevitUtils.cs">
+      <Link>DatasmithRevitExporter\DatasmithRevitUtils.cs</Link>
     </Compile>
     <Compile Include="..\Private\DatasmithRevitSettingsDialog.cs">
       <Link>DatasmithRevitExporter\DatasmithRevitSettingsDialog.cs</Link>
diff --git a/Engine/Source/Programs/Enterprise/Datasmith/DatasmithRevitExporter/DatasmithRevit2022/DatasmithRevit2022.csproj b/Engine/Source/Programs/Enterprise/Datasmith/DatasmithRevitExporter/DatasmithRevit2022/DatasmithRevit2022.csproj
--- a/Engine/Source/Programs/Enterprise/Datasmith/DatasmithRevitExporter/DatasmithRevit2022/DatasmithRevit2022.csproj
+++ b/Engine/Source/Programs/Enterprise/Datasmith/DatasmithRevitExporter/DatasmithRevit2022/DatasmithRevit2022.csproj
@@ -87,8 +87,11 @@
     <Compile Include="..\Private\DatasmithRevitMaterial.cs">
       <Link>DatasmithRevitExporter\DatasmithRevitMaterial.cs</Link>
     </Compile>
-    <Compile Include="..\Private\DatasmithRevitMetadataManager.cs">
-      <Link>DatasmithRevitExporter\DatasmithRevitMetdataManager.cs</Link>
+    <Compile Include="..\Private\DatasmithRevitSettingsManager.cs">
+      <Link>DatasmithRevitExporter\DatasmithRevitSettingsManager.cs</Link>
+    </Compile>
+    <Compile Include="..\Private\DatasmithRevitUtils.cs">
+      <Link>DatasmithRevitExporter\DatasmithRevitUtils.cs</Link>
     </Compile>
     <Compile Include="..\Private\DatasmithRevitSettingsDialog.cs">
       <Link>DatasmithRevitExporter\DatasmithRevitSettingsDialog.cs</Link>
diff --git a/Engine/Source/Programs/Enterprise/Datasmith/DatasmithRevitExporter/Private/DatasmithDynamoNode.cs b/Engine/Source/Programs/Enterprise/Datasmith/DatasmithRevitExporter/Private/DatasmithDynamoNode.cs
--- a/Engine/Source/Programs/Enterprise/Datasmith/DatasmithRevitExporter/Private/DatasmithDynamoNode.cs
+++ b/Engine/Source/Programs/Enterprise/Datasmith/DatasmithRevitExporter/Private/DatasmithDynamoNode.cs
@@ -75,7 +75,7 @@ namespace DSNode
 				null);
 
 			// Clamp tesselation parameter to a valid range.
-			ExportContext.LevelOfTessellation = Math.Min(Math.Max(InTesselation, -1), 15);
+			ExportContext.LevelOfTessellation = Math.Min(Math.Max(InTesselation, 1), 15);
 
 			using (CustomExporter Exporter = new CustomExporter(InDocument, ExportContext))
 			{
diff --git a/Engine/Source/Programs/Enterprise/Datasmith/DatasmithRevitExporter/Private/DatasmithRevitApplication.cs b/Engine/Source/Programs/Enterprise/Datasmith/DatasmithRevitExporter/Private/DatasmithRevitApplication.cs
--- a/Engine/Source/Programs/Enterprise/Datasmith/DatasmithRevitExporter/Private/DatasmithRevitApplication.cs
+++ b/Engine/Source/Programs/Enterprise/Datasmith/DatasmithRevitExporter/Private/DatasmithRevitApplication.cs
@@ -77,9 +77,9 @@ namespace DatasmithRevitExporter
 			SyncPushButton = DirectLinkRibbonPanel.AddItem(SyncButtonData) as PushButton;
 			AutoSyncPushButton = DirectLinkRibbonPanel.AddItem(AutoSyncButtonData) as PushButton;
 			PushButton ManageConnectionsButton = DirectLinkRibbonPanel.AddItem(ManageConnectionsButtonData) as PushButton;
-			PushButton SettingsButton = DirectLinkRibbonPanel.AddItem(SettingsButtonData) as PushButton;
 			PushButton ExportPushButton = FileExportRibbonPanel.AddItem(ExportButtonData) as PushButton;
 			PushButton ShowLogButton = DatasmithRibbonPanel.AddItem(LogButtonData) as PushButton;
+			PushButton SettingsButton = DatasmithRibbonPanel.AddItem(SettingsButtonData) as PushButton;
 
 			string DatasmithIconBase = Path.Combine(Path.GetDirectoryName(AssemblyPath), "DatasmithIcon");
 			ExportPushButton.Image = new BitmapImage(new Uri(DatasmithIconBase + "16.png"));
@@ -147,7 +147,7 @@ namespace DatasmithRevitExporter
 
 			Debug.Assert(bDirectLinkInitOk);
 
-			FMetadataManager.Init(InApplication);
+			FSettingsManager.Init(InApplication);
 
 			return Result.Succeeded;
 		}
@@ -183,7 +183,7 @@ namespace DatasmithRevitExporter
 			{
 				ExportMessagesDialog.Close();
 			}
-			FMetadataManager.Destroy(InApplication);
+			FSettingsManager.Destroy(InApplication);
 			FDatasmithFacadeDirectLink.Shutdown();
 			return Result.Succeeded;
 		}
diff --git a/Engine/Source/Programs/Enterprise/Datasmith/DatasmithRevitExporter/Private/DatasmithRevitCommand.cs b/Engine/Source/Programs/Enterprise/Datasmith/DatasmithRevitExporter/Private/DatasmithRevitCommand.cs
--- a/Engine/Source/Programs/Enterprise/Datasmith/DatasmithRevitExporter/Private/DatasmithRevitCommand.cs
+++ b/Engine/Source/Programs/Enterprise/Datasmith/DatasmithRevitExporter/Private/DatasmithRevitCommand.cs
@@ -178,14 +178,6 @@ namespace DatasmithRevitExporter
 
 			// Retrieve the Unreal Datasmith export options.
 			DatasmithRevitExportOptions ExportOptions = new DatasmithRevitExportOptions(Doc);
-			if ((System.Windows.Forms.Control.ModifierKeys & Keys.Control) == Keys.Control)
-			{
-				if (ExportOptions.ShowDialog() != DialogResult.OK)
-				{
-					return Result.Cancelled;
-				}
-				ExportActiveViewOnly = false;
-			}
 
 			// Generate file path for each view.
 			Dictionary<ElementId, string> FilePaths = new Dictionary<ElementId, string>();
diff --git a/Engine/Source/Programs/Enterprise/Datasmith/DatasmithRevitExporter/Private/DatasmithRevitDirectLink.cs b/Engine/Source/Programs/Enterprise/Datasmith/DatasmithRevitExporter/Private/DatasmithRevitDirectLink.cs
--- a/Engine/Source/Programs/Enterprise/Datasmith/DatasmithRevitExporter/Private/DatasmithRevitDirectLink.cs
+++ b/Engine/Source/Programs/Enterprise/Datasmith/DatasmithRevitExporter/Private/DatasmithRevitDirectLink.cs
@@ -724,14 +724,9 @@ namespace DatasmithRevitExporter
 			bHasChanges = false;
 			bSyncInProgress = false;
 
-			// Control metadata export via env var REVIT_DIRECTLINK_WITH_METADATA.
-			// We are not interested of its value, just if it was set.
-			if (null != Environment.GetEnvironmentVariable("REVIT_DIRECTLINK_WITH_METADATA"))
-			{
-				Debug.Assert(MetadataTask == null); // We cannot have metadata export running at this point (must be stopped in OnBeginExport)
-				MetadataCancelToken = new CancellationTokenSource();
-				MetadataTask = Task.Run(() => ExportMetadata());
-			}
+			Debug.Assert(MetadataTask == null); // We cannot have metadata export running at this point (must be stopped in OnBeginExport)
+			MetadataCancelToken = new CancellationTokenSource();
+			MetadataTask = Task.Run(() => ExportMetadata());
 		}
 
 		void ExportMetadata()
@@ -783,7 +778,7 @@ namespace DatasmithRevitExporter
 					ElementData.ElementMetaData.SetLabel(Actor.GetLabel());
 					ElementData.ElementMetaData.SetAssociatedElement(Actor);
 
-					FMetadataManager.AddActorMetadata(RevitElement, ElementData.ElementMetaData, true);
+					FUtils.AddActorMetadata(RevitElement, ElementData.ElementMetaData);
 
 					++CurrentBatchSize;
 
diff --git a/Engine/Source/Programs/Enterprise/Datasmith/DatasmithRevitExporter/Private/DatasmithRevitDocumentData.cs b/Engine/Source/Programs/Enterprise/Datasmith/DatasmithRevitExporter/Private/DatasmithRevitDocumentData.cs
--- a/Engine/Source/Programs/Enterprise/Datasmith/DatasmithRevitExporter/Private/DatasmithRevitDocumentData.cs
+++ b/Engine/Source/Programs/Enterprise/Datasmith/DatasmithRevitExporter/Private/DatasmithRevitDocumentData.cs
@@ -1040,12 +1040,12 @@ namespace DatasmithRevitExporter
 				}
 
 				// Add Revit element metadata to the Datasmith actor.
-				FMetadataManager.AddActorMetadata(CurrentElement, "Element*", ElementMetaData, false);
+				FUtils.AddActorMetadata(CurrentElement, "Element*", ElementMetaData);
 
 				if (BaseElementType != null)
 				{
 					// Add Revit element type metadata to the Datasmith actor.
-					FMetadataManager.AddActorMetadata(BaseElementType, "Type*", ElementMetaData, false);
+					FUtils.AddActorMetadata(BaseElementType, "Type*", ElementMetaData);
 				}
 
 				return ElementMetaData;
@@ -1874,7 +1874,7 @@ namespace DatasmithRevitExporter
 
 					if (!bSkipMetadataExport)
 					{
-						FMetadataManager.AddActorMetadata(BasePointLocation, MetadataPrefix, BasePointMetaData, false);
+						FUtils.AddActorMetadata(BasePointLocation, MetadataPrefix, BasePointMetaData);
 					}
 
 					if (BasePointElement == null)
diff --git a/Engine/Source/Programs/Enterprise/Datasmith/DatasmithRevitExporter/Private/DatasmithRevitExportContext.cs b/Engine/Source/Programs/Enterprise/Datasmith/DatasmithRevitExporter/Private/DatasmithRevitExportContext.cs
--- a/Engine/Source/Programs/Enterprise/Datasmith/DatasmithRevitExporter/Private/DatasmithRevitExportContext.cs
+++ b/Engine/Source/Programs/Enterprise/Datasmith/DatasmithRevitExporter/Private/DatasmithRevitExportContext.cs
@@ -97,7 +97,7 @@ namespace DatasmithRevitExporter
 
 			// Get the Unreal Datasmith export options.
 			DebugLog = InExportOptions.GetWriteLogFile() ? new FDatasmithFacadeLog() : null;
-			LevelOfTessellation = InExportOptions.GetLevelOfTessellation();
+			LevelOfTessellation = FSettingsManager.CurrentSettings?.LevelOfTesselation ?? 8;
 		}
 
 		//========================================================================================================================
diff --git a/Engine/Source/Programs/Enterprise/Datasmith/DatasmithRevitExporter/Private/DatasmithRevitExportOptions.cs b/Engine/Source/Programs/Enterprise/Datasmith/DatasmithRevitExporter/Private/DatasmithRevitExportOptions.cs
--- a/Engine/Source/Programs/Enterprise/Datasmith/DatasmithRevitExporter/Private/DatasmithRevitExportOptions.cs
+++ b/Engine/Source/Programs/Enterprise/Datasmith/DatasmithRevitExporter/Private/DatasmithRevitExportOptions.cs
@@ -11,7 +11,6 @@ namespace DatasmithRevitExporter
 	public class DatasmithRevitExportOptions : Form
 	{
 		private CheckBox		WriteLogFile;
-		private NumericUpDown	LevelOfTessellation;
 		private CheckedListBox	ViewsList;
 
 		public List<Autodesk.Revit.DB.View3D> Selected3DViews { get; private set; } = new List<Autodesk.Revit.DB.View3D>();
@@ -45,7 +44,7 @@ namespace DatasmithRevitExporter
 
 			Label LevelOfTessellationLabel = new Label();
 			LevelOfTessellationLabel.Name = "LevelOfTessellationLabel";
-			LevelOfTessellationLabel.Text = DatasmithRevitResources.Strings.ExportOptionsDialog_LevelOfTesselation;
+			LevelOfTessellationLabel.Text = DatasmithRevitResources.Strings.SettingsDialog_LevelOfTesselation;
 			LevelOfTessellationLabel.Anchor = AnchorStyles.Left;
 			LevelOfTessellationLabel.AutoSize = true;
 			LevelOfTessellationLabel.Location = new Point(3, 26);
@@ -53,20 +52,6 @@ namespace DatasmithRevitExporter
 			LevelOfTessellationLabel.TabIndex = 2;
 			LevelOfTessellationLabel.TextAlign = ContentAlignment.MiddleLeft;
 
-			LevelOfTessellation = new NumericUpDown();
-			LevelOfTessellation.Name = "LevelOfTessellation";
-			LevelOfTessellation.Minimum = -1;
-			LevelOfTessellation.Maximum = 15;
-			LevelOfTessellation.Value = 8;
-			LevelOfTessellation.Anchor = AnchorStyles.Left;
-			LevelOfTessellation.AutoSize = true;
-			LevelOfTessellation.Location = new Point(113, 23);
-			LevelOfTessellation.Size = new Size(35, 20);
-			LevelOfTessellation.TabIndex = 3;
-			LevelOfTessellation.TextAlign = HorizontalAlignment.Right;
-
-			OptionToolTip.SetToolTip(LevelOfTessellation, DatasmithRevitResources.Strings.ExportOptionsDialog_LevelOfTesselationTooltip);
-
 			Button OKButton = new Button();
 			OKButton.Name = "OKButton";
 			OKButton.Text = DatasmithRevitResources.Strings.ExportOptionsDialog_ButtonOK;
@@ -126,7 +111,6 @@ namespace DatasmithRevitExporter
 			DialogLayout.Controls.Add(WriteLogFileLabel,		0, 2);
 			DialogLayout.Controls.Add(WriteLogFile,				1, 2);
 			DialogLayout.Controls.Add(LevelOfTessellationLabel, 0, 3);
-			DialogLayout.Controls.Add(LevelOfTessellation,		1, 3);
 			DialogLayout.Controls.Add(OKButton,					1, 4);
 
 			DialogLayout.SetColumnSpan(ViewsList, 2);
@@ -162,10 +146,5 @@ namespace DatasmithRevitExporter
 		{
 			return WriteLogFile.Checked;
 		}
-
-		public int GetLevelOfTessellation()
-		{
-			return Decimal.ToInt32(LevelOfTessellation.Value);
-		}
 	}
 }
diff --git a/Engine/Source/Programs/Enterprise/Datasmith/DatasmithRevitExporter/Private/DatasmithRevitMetadataManager.cs b/Engine/Source/Programs/Enterprise/Datasmith/DatasmithRevitExporter/Private/DatasmithRevitMetadataManager.cs
deleted file mode 100644
diff --git a/Engine/Source/Programs/Enterprise/Datasmith/DatasmithRevitExporter/Private/DatasmithRevitSettingsDialog.cs b/Engine/Source/Programs/Enterprise/Datasmith/DatasmithRevitExporter/Private/DatasmithRevitSettingsDialog.cs
--- a/Engine/Source/Programs/Enterprise/Datasmith/DatasmithRevitExporter/Private/DatasmithRevitSettingsDialog.cs
+++ b/Engine/Source/Programs/Enterprise/Datasmith/DatasmithRevitExporter/Private/DatasmithRevitSettingsDialog.cs
@@ -3,7 +3,6 @@
 using System;
 using System.Collections.Generic;
 using System.Drawing;
-using System.Text;
 using System.Windows.Forms;
 
 namespace DatasmithRevitExporter
@@ -87,53 +86,59 @@ namespace DatasmithRevitExporter
 
 	public class DatasmithRevitSettingsDialog : Form
 	{
-		private CheckedListBox GroupsList;
-		private TextBox ParamNameTextBox;
-		private FMetadataSettings MetadataSettings;
+		private ListBox GroupsList;
+		private NumericUpDown LevelOfTessellation;
+		private FSettings Settings;
 		private Autodesk.Revit.DB.Document Document;
 		private SortedList<string, int> AddedBuiltinParamGroups = new SortedList<string, int>();
 		private List<string> AddedParamNames = new List<string>();
 
 		public DatasmithRevitSettingsDialog(Autodesk.Revit.DB.Document InDocument)
 		{
+			string FormatTooltip(string InText)
+			{
+				string Tooltip = InText;
+				return Tooltip.Replace("\\n", Environment.NewLine);
+			}
+
 			Document = InDocument;
 			FormClosing += OnClosing;
 
-			MetadataSettings = FMetadataManager.CurrentSettings;
+			Settings = FSettingsManager.CurrentSettings;
 
 			ToolTip OptionToolTip = new ToolTip();
 			OptionToolTip.AutoPopDelay = 10000; // milliseconds
 
-			Label DescriptionLabel = new Label();
-			DescriptionLabel.Name = "MatchParameterName";
-			DescriptionLabel.Text = DatasmithRevitResources.Strings.SettingsDialog_LabelMetadataDescription;
-			DescriptionLabel.Anchor = AnchorStyles.Left;
-			DescriptionLabel.AutoSize = false;
-			DescriptionLabel.TabIndex = 0;
-			DescriptionLabel.Font = new Font(FontFamily.GenericSansSerif, 8f);
-			DescriptionLabel.BorderStyle = BorderStyle.FixedSingle;
-			DescriptionLabel.BackColor = Color.Wheat;
-			DescriptionLabel.TextAlign = ContentAlignment.MiddleLeft;
-			DescriptionLabel.Margin = new Padding(0, 10, 0, 20);
-			DescriptionLabel.Width = 385;
-			DescriptionLabel.Height = 55;
-
-			Label ParamNameLabel = new Label();
-			ParamNameLabel.Name = "MatchParameterNames";
-			ParamNameLabel.Text = DatasmithRevitResources.Strings.SettingsDialog_LabelMatchParamNames;
-			ParamNameLabel.Anchor = AnchorStyles.Left;
-			ParamNameLabel.AutoSize = true;
-			ParamNameLabel.TabIndex = 0;
-			ParamNameLabel.TextAlign = ContentAlignment.MiddleLeft;
-			ParamNameLabel.Margin = new Padding(0, 0, 0, 10);
-			OptionToolTip.SetToolTip(ParamNameLabel, "Paramter names to match, separated by new lines (can be partial names).\nCase insensitive.");
-
-			ParamNameTextBox = new TextBox();
-			ParamNameTextBox.Width = 385;
-			ParamNameTextBox.Height = 150;
-			ParamNameTextBox.TabStop = false;
-			ParamNameTextBox.Multiline = true;
-			ParamNameTextBox.Margin = new Padding(0, 0, 0, 20);
+			FlowLayoutPanel LevelOfTesselationPanel = new FlowLayoutPanel();
+			LevelOfTesselationPanel.FlowDirection = FlowDirection.LeftToRight;
+			LevelOfTesselationPanel.WrapContents = false;
+			LevelOfTesselationPanel.AutoSize = true;
+			LevelOfTesselationPanel.Margin = new Padding(0, 0, 0, 30);
+			{
+				Label LevelOfTessellationLabel = new Label();
+				LevelOfTessellationLabel.Name = "LevelOfTessellationLabel";
+				LevelOfTessellationLabel.Text = DatasmithRevitResources.Strings.SettingsDialog_LevelOfTesselation;
+				LevelOfTessellationLabel.Anchor = AnchorStyles.Left;
+				LevelOfTessellationLabel.AutoSize = true;
+				LevelOfTessellationLabel.Margin = new Padding(0, 0, 3, 0);
+				LevelOfTessellationLabel.TabIndex = 2;
+				LevelOfTessellationLabel.TextAlign = ContentAlignment.MiddleLeft;
+
+				LevelOfTessellation = new NumericUpDown();
+				LevelOfTessellation.Name = "LevelOfTessellation";
+				LevelOfTessellation.Minimum = 1;
+				LevelOfTessellation.Maximum = 15;
+				LevelOfTessellation.Value = 8;
+				LevelOfTessellation.Anchor = AnchorStyles.Left;
+				LevelOfTessellation.AutoSize = true;
+				LevelOfTessellation.TabIndex = 3;
+				LevelOfTessellation.TextAlign = HorizontalAlignment.Right;
+
+				OptionToolTip.SetToolTip(LevelOfTessellationLabel, FormatTooltip(DatasmithRevitResources.Strings.SettingsDialog_LevelOfTesselationTooltip));
+
+				LevelOfTesselationPanel.Controls.Add(LevelOfTessellationLabel);
+				LevelOfTesselationPanel.Controls.Add(LevelOfTessellation);
+			}
 
 			Label GroupsLabel = new Label();
 			GroupsLabel.Name = "GroupsLabel";
@@ -141,62 +146,91 @@ namespace DatasmithRevitExporter
 			GroupsLabel.Anchor = AnchorStyles.Left;
 			GroupsLabel.AutoSize = true;
 
-			GroupsList = new CheckedListBox();
-			GroupsList.Size = new Size(380, 200);
-
-			Button AddGroupButton = new Button();
-			AddGroupButton.Name = "Add";
-			AddGroupButton.Text = DatasmithRevitResources.Strings.SettingsDialog_ButtonAddGroups;
-			AddGroupButton.AutoSize = true;
-			AddGroupButton.Anchor = AnchorStyles.Left;
-			AddGroupButton.UseVisualStyleBackColor = true;
-			AddGroupButton.Click += (s, e) =>
+			GroupsList = new ListBox();
+			GroupsList.Size = new Size(450, 250);
+			GroupsList.SelectionMode = SelectionMode.MultiExtended;
+
+			OptionToolTip.SetToolTip(GroupsLabel, FormatTooltip(DatasmithRevitResources.Strings.SettingsDialog_MetadataFilterTooltip));
+	
+			FlowLayoutPanel ButtonsPanel = new FlowLayoutPanel();
+			ButtonsPanel.FlowDirection = FlowDirection.TopDown;
+			ButtonsPanel.WrapContents = false;
+			ButtonsPanel.AutoSize = true;
+			ButtonsPanel.AutoSizeMode = AutoSizeMode.GrowAndShrink;
+			ButtonsPanel.Padding = new Padding(0);
 			{
-				DatasmithRevitParamGroupsDialog Dlg = new DatasmithRevitParamGroupsDialog();
-				if (DialogResult.OK == Dlg.ShowDialog())
+				Button AddGroupButton = new Button();
+				AddGroupButton.Name = "Add";
+				AddGroupButton.Text = DatasmithRevitResources.Strings.SettingsDialog_ButtonAddGroups;
+				AddGroupButton.Margin = new Padding(0);
+				AddGroupButton.AutoSizeMode = AutoSizeMode.GrowAndShrink;
+				AddGroupButton.AutoSize = true;
+				AddGroupButton.Anchor = AnchorStyles.Top | AnchorStyles.Left | AnchorStyles.Right;
+				AddGroupButton.UseVisualStyleBackColor = true;
+				AddGroupButton.Click += (s, e) =>
 				{
-					bool bAddedNewGroups = false;
-					foreach (int Index in Dlg.GroupsList.CheckedIndices)
+					DatasmithRevitParamGroupsDialog Dlg = new DatasmithRevitParamGroupsDialog();
+					if (DialogResult.OK == Dlg.ShowDialog())
 					{
-						string GroupLabel = DatasmithRevitParamGroupsDialog.AllBuiltinParamGroups.Keys[Index];
-						if (!AddedBuiltinParamGroups.ContainsKey(GroupLabel))
+						bool bAddedNewGroups = false;
+						foreach (int Index in Dlg.GroupsList.CheckedIndices)
 						{
-							int GroupId;
-							if (DatasmithRevitParamGroupsDialog.AllBuiltinParamGroups.TryGetValue(GroupLabel, out GroupId))
+							string GroupLabel = DatasmithRevitParamGroupsDialog.AllBuiltinParamGroups.Keys[Index];
+							if (!AddedBuiltinParamGroups.ContainsKey(GroupLabel))
 							{
-								AddedBuiltinParamGroups.Add(GroupLabel, GroupId);
-								bAddedNewGroups = true;
+								int GroupId;
+								if (DatasmithRevitParamGroupsDialog.AllBuiltinParamGroups.TryGetValue(GroupLabel, out GroupId))
+								{
+									AddedBuiltinParamGroups.Add(GroupLabel, GroupId);
+									bAddedNewGroups = true;
+								}
 							}
 						}
-					}
 
-					if (bAddedNewGroups)
-					{
-						ReloadAddedGroupsList();
+						if (bAddedNewGroups)
+						{
+							ReloadAddedGroupsList();
+						}
 					}
-				}
-			};
-
-			Button RemoveGroupButton = new Button();
-			RemoveGroupButton.Name = "Remove";
-			RemoveGroupButton.Text = DatasmithRevitResources.Strings.SettingsDialog_ButtonRemoveGroups;
-			RemoveGroupButton.AutoSize = true;
-			RemoveGroupButton.Anchor = AnchorStyles.Right;
-			RemoveGroupButton.UseVisualStyleBackColor = true;
-			RemoveGroupButton.Click += (s, e) =>
-			{
-				if (GroupsList.CheckedIndices.Count > 0)
+				};
+
+				Button RemoveGroupButton = new Button();
+				RemoveGroupButton.Name = "Remove";
+				RemoveGroupButton.Text = DatasmithRevitResources.Strings.SettingsDialog_ButtonRemoveGroups;
+				RemoveGroupButton.Margin = new Padding(0, 10, 0, 0);
+				RemoveGroupButton.Anchor = AnchorStyles.Top | AnchorStyles.Left | AnchorStyles.Right;
+				RemoveGroupButton.AutoSizeMode = AutoSizeMode.GrowAndShrink;
+				RemoveGroupButton.AutoSize = true;
+				RemoveGroupButton.UseVisualStyleBackColor = true;
+				RemoveGroupButton.Click += (s, e) =>
 				{
-					int[] IndicesToRemove = new int[GroupsList.CheckedIndices.Count];
-					GroupsList.CheckedIndices.CopyTo(IndicesToRemove, 0);
-					Array.Sort(IndicesToRemove, (A, B) => B.CompareTo(A));
-
-					foreach (int Index in IndicesToRemove)
+					if (GroupsList.SelectedIndices.Count > 0)
 					{
-						GroupsList.Items.RemoveAt(Index);
-						AddedBuiltinParamGroups.RemoveAt(Index);
+						int[] IndicesToRemove = new int[GroupsList.SelectedIndices.Count];
+						GroupsList.SelectedIndices.CopyTo(IndicesToRemove, 0);
+						Array.Sort(IndicesToRemove, (A, B) => B.CompareTo(A));
+
+						foreach (int Index in IndicesToRemove)
+						{
+							GroupsList.Items.RemoveAt(Index);
+							AddedBuiltinParamGroups.RemoveAt(Index);
+						}
 					}
-				}
+				};
+
+				ButtonsPanel.Controls.Add(AddGroupButton);
+				ButtonsPanel.Controls.Add(RemoveGroupButton);
+			}
+
+			Button CloseButton = new Button();
+			CloseButton.Name = "Close";
+			CloseButton.Text = DatasmithRevitResources.Strings.SettingsDialog_ButtonClose;
+			CloseButton.Anchor = AnchorStyles.Right;
+			CloseButton.UseVisualStyleBackColor = true;
+			CloseButton.Margin = new Padding(0, 7, 0, 0);
+			CloseButton.Click += (s, e) =>
+			{
+				Close();
 			};
 
 			TableLayoutPanel DialogLayout = new TableLayoutPanel();
@@ -204,63 +238,62 @@ namespace DatasmithRevitExporter
 			DialogLayout.ColumnCount = 2;
 			DialogLayout.ColumnStyles.Add(new ColumnStyle());
 			DialogLayout.ColumnStyles.Add(new ColumnStyle());
-			DialogLayout.RowCount = 6;
-			DialogLayout.RowStyles.Add(new RowStyle());
-			DialogLayout.RowStyles.Add(new RowStyle());
-			DialogLayout.RowStyles.Add(new RowStyle());
-			DialogLayout.RowStyles.Add(new RowStyle());
-			DialogLayout.RowStyles.Add(new RowStyle());
+			DialogLayout.RowCount = 3;
+			DialogLayout.RowStyles.Add(new RowStyle(SizeType.AutoSize));
+			DialogLayout.RowStyles.Add(new RowStyle(SizeType.AutoSize));
+			DialogLayout.RowStyles.Add(new RowStyle(SizeType.AutoSize));
 			DialogLayout.AutoSize = true;
 			DialogLayout.AutoSizeMode = AutoSizeMode.GrowAndShrink;
 			DialogLayout.Dock = DockStyle.Fill;
-			DialogLayout.Location = new Point(10, 10);
 			DialogLayout.TabIndex = 0;
+		
+			DialogLayout.Controls.Add(LevelOfTesselationPanel, 0, 0);
+			DialogLayout.Controls.Add(GroupsLabel, 0, 1);
+			DialogLayout.Controls.Add(GroupsList, 0, 2);
+			DialogLayout.Controls.Add(ButtonsPanel, 1, 2);
 
-			DialogLayout.Controls.Add(DescriptionLabel, 0,0);
-			DialogLayout.Controls.Add(ParamNameLabel, 0, 1);
-			DialogLayout.Controls.Add(ParamNameTextBox, 0, 2);
-			DialogLayout.Controls.Add(GroupsLabel, 0, 3);
-			DialogLayout.Controls.Add(GroupsList, 0, 4);
-			DialogLayout.Controls.Add(AddGroupButton, 0, 5);
-			DialogLayout.Controls.Add(RemoveGroupButton, 1, 5);
-
-			DialogLayout.SetColumnSpan(ParamNameTextBox, 2);
-			DialogLayout.SetColumnSpan(GroupsList, 2);
-			DialogLayout.SetColumnSpan(DescriptionLabel, 2);
+			DialogLayout.SetColumnSpan(LevelOfTesselationPanel, 2);
+			DialogLayout.SetColumnSpan(GroupsLabel, 2);
 
 			Name = "MetadataExportFilter";
 			Text = DatasmithRevitResources.Strings.SettingsDialog_DialogTitle;
-			AutoScaleDimensions = new SizeF(6F, 13F);
 			AutoScaleMode = AutoScaleMode.Font;
 			AutoSize = true;
-			AutoSizeMode = AutoSizeMode.GrowAndShrink;
-			ClientSize = new Size(309, 114);
 			FormBorderStyle = FormBorderStyle.FixedDialog;
 			MaximizeBox = false;
 			MinimizeBox = false;
-			Padding = new Padding(10);
+			//Padding = new Padding(7);
 			SizeGripStyle = SizeGripStyle.Hide;
 			StartPosition = FormStartPosition.CenterParent;
 
-			Controls.Add(DialogLayout);
+			Panel ParentPanel = new Panel();
+			ParentPanel.BorderStyle = BorderStyle.FixedSingle;
+			ParentPanel.AutoSize = true;
+			ParentPanel.Margin = new Padding(0);
+			ParentPanel.Padding = new Padding(7);
+			ParentPanel.Controls.Add(DialogLayout);
+			//ParentPanel.Location = new Point(7, 7);
+
+			FlowLayoutPanel TopLayout = new FlowLayoutPanel();
+			TopLayout.FlowDirection = FlowDirection.TopDown;
+			TopLayout.Padding = new Padding(7);
+			TopLayout.AutoSize = true;
+			TopLayout.AutoSizeMode = AutoSizeMode.GrowAndShrink;
+
+			TopLayout.Controls.Add(ParentPanel);
+			TopLayout.Controls.Add(CloseButton);
 
-			if (MetadataSettings == null)
+			Controls.Add(TopLayout);
+
+			if (Settings == null)
 			{
-				MetadataSettings = new FMetadataSettings();
+				Settings = new FSettings();
 			}
 			else
 			{
-				StringBuilder SB = new StringBuilder();
-				foreach (string ParamNameFilter in MetadataSettings.ParamNamesFilter)
-				{
-					if (ParamNameFilter.Length > 0)
-					{
-						SB.AppendLine(ParamNameFilter);
-					}
-				}
-				ParamNameTextBox.Text = SB.ToString();
+				LevelOfTessellation.Value = Settings.LevelOfTesselation;
 
-				foreach (int Group in MetadataSettings.ParamGroupsFilter)
+				foreach (int Group in Settings.MetadataParamGroupsFilter)
 				{
 					string GroupLabel = Autodesk.Revit.DB.LabelUtils.GetLabelFor((Autodesk.Revit.DB.BuiltInParameterGroup)Group);
 					AddedBuiltinParamGroups.Add(GroupLabel, Group);
@@ -282,9 +315,9 @@ namespace DatasmithRevitExporter
 
 		private void OnClosing(object InSender, FormClosingEventArgs InArgs)
 		{
-			MetadataSettings.ParamNamesFilter = ParamNameTextBox.Text.Split(new[] { Environment.NewLine }, StringSplitOptions.RemoveEmptyEntries);
-			MetadataSettings.ParamGroupsFilter = AddedBuiltinParamGroups.Values;
-			FMetadataManager.WriteSettings(Document, MetadataSettings);
+			Settings.LevelOfTesselation = Decimal.ToInt32(LevelOfTessellation.Value);
+			Settings.MetadataParamGroupsFilter = AddedBuiltinParamGroups.Values;
+			FSettingsManager.WriteSettings(Document, Settings);
 		}
 	}
 }
diff --git a/Engine/Source/Programs/Enterprise/Datasmith/DatasmithRevitExporter/Private/DatasmithRevitSettingsManager.cs b/Engine/Source/Programs/Enterprise/Datasmith/DatasmithRevitExporter/Private/DatasmithRevitSettingsManager.cs
new file mode 100644
--- /dev/null
+++ b/Engine/Source/Programs/Enterprise/Datasmith/DatasmithRevitExporter/Private/DatasmithRevitSettingsManager.cs
@@ -0,0 +1,269 @@
+// Copyright Epic Games, Inc. All Rights Reserved.
+
+using System;
+using System.Collections.Generic;
+using Autodesk.Revit.DB;
+using Autodesk.Revit.DB.Events;
+using Autodesk.Revit.DB.ExtensibleStorage;
+using Autodesk.Revit.UI;
+
+namespace DatasmithRevitExporter
+{
+	public class FSettings
+	{
+		private IList<string> _MetadataParamNamesFilter = new List<string>();
+		private IList<int> _MetadataParamGroupsFilter = new List<int>();
+
+		// Derived data for fast lookup
+		private HashSet<string> MetadataParamNamesSet;
+		private HashSet<int> MetadataParamGroupsSet;
+
+		public int LevelOfTesselation { get; set; } = 8;
+
+		public IList<string> MetadataParamNamesFilter 
+		{
+			get => _MetadataParamNamesFilter;
+			set
+			{
+				_MetadataParamNamesFilter = value;
+				MetadataParamNamesSet = new HashSet<string>(_MetadataParamNamesFilter);
+			}
+		}
+
+		public IList<int> MetadataParamGroupsFilter 
+		{
+			get => _MetadataParamGroupsFilter;
+			set
+			{
+				_MetadataParamGroupsFilter = value;
+				MetadataParamGroupsSet = new HashSet<int>(_MetadataParamGroupsFilter);
+			}
+		}
+
+		// Returns true if the parameter passes the filter, false otherwise.
+		public bool MatchParameterByMetadata(Parameter InParam)
+		{
+			if (MetadataParamNamesSet != null && MetadataParamNamesSet.Count > 0)
+			{
+				bool bNameMatch = false;
+				foreach (string NameFilter in MetadataParamNamesSet)
+				{
+					// Check if NameFilter is contained in the current param name
+					if (InParam.Definition.Name.IndexOf(NameFilter, StringComparison.OrdinalIgnoreCase) >= 0)
+					{
+						bNameMatch = true;
+						break;
+					}
+				}
+
+				if (!bNameMatch)
+				{
+					return false;
+				}
+			}
+
+			// Name filter passes (or empty), so check group filter.
+			bool bGroupMatch = true;
+			if (MetadataParamGroupsSet != null && MetadataParamGroupsSet.Count > 0)
+			{
+				bGroupMatch = MetadataParamGroupsSet.Contains((int)InParam.Definition.ParameterGroup);
+			}
+		
+			return bGroupMatch;
+		}
+	}
+
+	public class FSettingsManager
+	{
+		static class FSettingsSchema
+		{
+			readonly static Guid SchemaGuid = new Guid("{EAFD9DF3-1E44-4B16-B9DF-E7DD404729A3}");
+			
+			public readonly static string MetadataParamNamesField = "MetadataParamNames";
+			public readonly static string MetadataParamGroupsField = "MetadataParamGroups";
+			public readonly static string LevelOfTesselationField = "LevelOfTesselation";
+
+			public static Schema GetSchema()
+			{
+				Schema Schema = Schema.Lookup(SchemaGuid);
+				if (Schema != null)
+				{
+					return Schema;
+				}
+
+				SchemaBuilder SchemaBuilder = new SchemaBuilder(SchemaGuid);
+				SchemaBuilder.SetSchemaName("DatasmithRevitExportSettings");
+
+				SchemaBuilder.AddArrayField(MetadataParamNamesField, typeof(string));
+				SchemaBuilder.AddArrayField(MetadataParamGroupsField, typeof(int));
+				SchemaBuilder.AddSimpleField(LevelOfTesselationField, typeof(int));
+
+				return SchemaBuilder.Finish();
+			}
+		}
+
+		static class DataStorageUniqueIdSchema
+		{
+			static readonly Guid SchemaGuid = new Guid("{2E363F82-4589-43BE-8927-C596AA02396F}");
+
+			public static Schema GetSchema()
+			{
+				Schema Schema = Schema.Lookup(SchemaGuid);
+				if (Schema != null)
+				{
+					return Schema;
+				}
+
+				SchemaBuilder SchemaBuilder = new SchemaBuilder(SchemaGuid);
+				SchemaBuilder.SetSchemaName("DataStorageUniqueId");
+				SchemaBuilder.AddSimpleField("Id", typeof(Guid));
+
+				return SchemaBuilder.Finish();
+			}
+		}
+
+		static EventHandler<DocumentOpenedEventArgs> DocumentOpenedHandler;
+		static EventHandler<DocumentClosingEventArgs> DocumentClosingHandler;
+		static readonly Guid SettingDataStorageId = new Guid("{6CD2BB42-4A04-4F12-8E7F-5C6723C33791}");
+
+		public static FSettings CurrentSettings = null;
+
+		public static void Init(UIControlledApplication InApplication)
+		{
+			DocumentOpenedHandler = new EventHandler<DocumentOpenedEventArgs>(OnDocumentOpened);
+			DocumentClosingHandler = new EventHandler<DocumentClosingEventArgs>(OnDocumentClosing);
+			InApplication.ControlledApplication.DocumentOpened += DocumentOpenedHandler;
+			InApplication.ControlledApplication.DocumentClosing += DocumentClosingHandler;
+		}
+
+		public static void Destroy(UIControlledApplication InApplication)
+		{
+			InApplication.ControlledApplication.DocumentClosing -= DocumentClosingHandler;
+			InApplication.ControlledApplication.DocumentOpened -= DocumentOpenedHandler;
+			DocumentClosingHandler = null;
+			DocumentOpenedHandler = null;
+		}
+
+		private static FSettings ReadSettings(Document Doc)
+		{
+			FSettings Settings = new FSettings();
+
+			Entity SettingsEntity = GetSettingsEntity(Doc);
+
+			if (SettingsEntity == null || !SettingsEntity.IsValid())
+			{
+				// No settings in document, create defaults
+				Settings.LevelOfTesselation = 8;
+				Settings.MetadataParamGroupsFilter.Add((int)Autodesk.Revit.DB.BuiltInParameterGroup.PG_GEOMETRY);
+				Settings.MetadataParamGroupsFilter.Add((int)Autodesk.Revit.DB.BuiltInParameterGroup.PG_IDENTITY_DATA);
+				Settings.MetadataParamGroupsFilter.Add((int)Autodesk.Revit.DB.BuiltInParameterGroup.PG_MATERIALS);
+				Settings.MetadataParamGroupsFilter.Add((int)Autodesk.Revit.DB.BuiltInParameterGroup.PG_PHASING);
+				Settings.MetadataParamGroupsFilter.Add((int)Autodesk.Revit.DB.BuiltInParameterGroup.PG_STRUCTURAL);
+
+				WriteSettings(Doc, Settings);
+			}
+			else
+			{
+				Settings.MetadataParamNamesFilter = SettingsEntity.Get<IList<string>>(FSettingsSchema.MetadataParamNamesField);
+				Settings.MetadataParamGroupsFilter = SettingsEntity.Get<IList<int>>(FSettingsSchema.MetadataParamGroupsField);
+				Settings.LevelOfTesselation = SettingsEntity.Get<int>(FSettingsSchema.LevelOfTesselationField);
+			}
+
+			return Settings;
+		}
+
+		public static void WriteSettings(Document Doc, FSettings Settings)
+		{
+			using (Transaction T = new Transaction(Doc, "Write settings"))
+			{
+				try
+				{
+					T.Start();
+
+					DataStorage SettingDataStorage = GetSettingsDataStorage(Doc);
+					if (SettingDataStorage == null)
+					{
+						SettingDataStorage = DataStorage.Create(Doc);
+					}
+
+					Entity SettingsEntity = new Entity(FSettingsSchema.GetSchema());
+					SettingsEntity.Set(FSettingsSchema.MetadataParamNamesField, Settings.MetadataParamNamesFilter);
+					SettingsEntity.Set(FSettingsSchema.MetadataParamGroupsField, Settings.MetadataParamGroupsFilter);
+					SettingsEntity.Set(FSettingsSchema.LevelOfTesselationField, Settings.LevelOfTesselation);
+
+					// Identify settings data storage
+					Entity IdEntity = new Entity(DataStorageUniqueIdSchema.GetSchema());
+					IdEntity.Set("Id", SettingDataStorageId);
+
+					SettingDataStorage.SetEntity(IdEntity);
+					SettingDataStorage.SetEntity(SettingsEntity);
+
+					CurrentSettings = Settings;
+
+					T.Commit();
+				}
+				catch { }
+			}
+		}
+
+		static void OnDocumentOpened(object sender, DocumentOpenedEventArgs e)
+		{
+			CurrentSettings = ReadSettings(e.Document);
+		}
+
+		static void OnDocumentClosing(object sender, DocumentClosingEventArgs e)
+		{
+			CurrentSettings = null;
+		}
+
+		private static DataStorage GetSettingsDataStorage(Document Doc)
+		{
+			// Retrieve all data storages from project
+			FilteredElementCollector Collector = new FilteredElementCollector(Doc);
+			var DataStorages = Collector.OfClass(typeof(DataStorage));
+
+			// Find setting data storage
+			foreach (DataStorage DataStorage in DataStorages)
+			{
+				Entity SettingIdEntity = DataStorage.GetEntity(DataStorageUniqueIdSchema.GetSchema());
+
+				if (!SettingIdEntity.IsValid())
+				{
+					continue;
+				}
+
+				Guid Id = SettingIdEntity.Get<Guid>("Id");
+
+				if (!Id.Equals(SettingDataStorageId))
+				{
+					continue;
+				}
+
+				return DataStorage;
+			}
+			return null;
+		}
+
+		private static Entity GetSettingsEntity(Document Doc)
+		{
+			FilteredElementCollector Collector = new FilteredElementCollector(Doc);
+			var DataStorages = Collector.OfClass(typeof(DataStorage));
+
+			// Find setting data storage
+			foreach (DataStorage DataStorage in DataStorages)
+			{
+				Entity SettingEntity = DataStorage.GetEntity(FSettingsSchema.GetSchema());
+
+				// If a DataStorage contains setting entity, we found it
+				if (!SettingEntity.IsValid())
+				{
+					continue;
+				}
+
+				return SettingEntity;
+			}
+
+			return null;
+		}
+	}
+}
diff --git a/Engine/Source/Programs/Enterprise/Datasmith/DatasmithRevitExporter/Private/DatasmithRevitUtils.cs b/Engine/Source/Programs/Enterprise/Datasmith/DatasmithRevitExporter/Private/DatasmithRevitUtils.cs
new file mode 100644
--- /dev/null
+++ b/Engine/Source/Programs/Enterprise/Datasmith/DatasmithRevitExporter/Private/DatasmithRevitUtils.cs
@@ -0,0 +1,107 @@
+// Copyright Epic Games, Inc. All Rights Reserved.
+
+using System.Collections.Generic;
+using Autodesk.Revit.DB;
+
+namespace DatasmithRevitExporter
+{
+	public static class FUtils
+	{
+		private static string GetCategoryName(Element InElement)
+		{
+			ElementType Type = GetElementType(InElement);
+			return Type?.Category?.Name ?? InElement.Category?.Name;
+		}
+
+		private static ElementType GetElementType(Element InElement)
+		{
+			return InElement.Document.GetElement(InElement.GetTypeId()) as ElementType;
+		}
+
+		public static void AddActorMetadata(Element InElement, FDatasmithFacadeMetaData ActorMetadata)
+		{
+			// Add the Revit element category name metadata to the Datasmith actor.
+			string CategoryName = GetCategoryName(InElement);
+			if (!string.IsNullOrEmpty(CategoryName))
+			{
+				ActorMetadata.AddPropertyString("Element*Category", CategoryName);
+			}
+
+			// Add the Revit element family name metadata to the Datasmith actor.
+			ElementType ElemType = GetElementType(InElement);
+			string FamilyName = ElemType?.FamilyName;
+			if (!string.IsNullOrEmpty(FamilyName))
+			{
+				ActorMetadata.AddPropertyString("Element*Family", FamilyName);
+			}
+
+			// Add the Revit element type name metadata to the Datasmith actor.
+			string TypeName = ElemType?.Name;
+			if (!string.IsNullOrEmpty(TypeName))
+			{
+				ActorMetadata.AddPropertyString("Element*Type", TypeName);
+			}
+
+			// Add Revit element metadata to the Datasmith actor.
+			AddActorMetadata(InElement, "Element*", ActorMetadata);
+
+			if (ElemType != null)
+			{
+				// Add Revit element type metadata to the Datasmith actor.
+				AddActorMetadata(ElemType, "Type*", ActorMetadata);
+			}
+		}
+
+		public static void AddActorMetadata(
+			Element InSourceElement,
+			string InMetadataPrefix,
+			FDatasmithFacadeMetaData ElementMetaData
+		)
+		{
+			FSettings Settings = FSettingsManager.CurrentSettings;
+
+			IList<Parameter> Parameters = InSourceElement.GetOrderedParameters();
+
+			if (Parameters != null)
+			{
+				foreach (Parameter Parameter in Parameters)
+				{
+					if (Parameter.HasValue)
+					{
+						if (Settings != null && !Settings.MatchParameterByMetadata(Parameter))
+						{
+							continue; // Skip export of this param
+						}
+
+						string ParameterValue = Parameter.AsValueString();
+
+						if (string.IsNullOrEmpty(ParameterValue))
+						{
+							switch (Parameter.StorageType)
+							{
+								case StorageType.Integer:
+								ParameterValue = Parameter.AsInteger().ToString();
+								break;
+								case StorageType.Double:
+								ParameterValue = Parameter.AsDouble().ToString();
+								break;
+								case StorageType.String:
+								ParameterValue = Parameter.AsString();
+								break;
+								case StorageType.ElementId:
+								ParameterValue = Parameter.AsElementId().ToString();
+								break;
+							}
+						}
+
+						if (!string.IsNullOrEmpty(ParameterValue))
+						{
+							string MetadataKey = InMetadataPrefix + Parameter.Definition.Name;
+							ElementMetaData.AddPropertyString(MetadataKey, ParameterValue);
+						}
+					}
+				}
+			}
+		}
+	}
+}
