Datasmith Solidworks: Additional fix for "Lights exported from solidwork should be converted to StdEmissive materials in UE"

Include file that was left out in original submit

diff --git a/Engine/Source/Programs/Enterprise/Datasmith/DatasmithFacadeCSharp/Public/FDatasmithFacadeMasterMaterial.cs b/Engine/Source/Programs/Enterprise/Datasmith/DatasmithFacadeCSharp/Public/FDatasmithFacadeMasterMaterial.cs
--- a/Engine/Source/Programs/Enterprise/Datasmith/DatasmithFacadeCSharp/Public/FDatasmithFacadeMasterMaterial.cs
+++ b/Engine/Source/Programs/Enterprise/Datasmith/DatasmithFacadeCSharp/Public/FDatasmithFacadeMasterMaterial.cs
@@ -123,6 +123,7 @@ public class FDatasmithFacadeMasterMaterial : FDatasmithFacadeBaseMaterial {
     Opaque,
     Transparent,
     ClearCoat,
+    Emissive,
     Custom,
     CutOut,
     Count
