DatasmithSketchupExporter Mac - Fixing plugin not being able to be uninstalled because of read-only files in the package.

diff --git a/Engine/Source/Programs/Enterprise/Datasmith/DatasmithSketchUpRubyExporter/Build/build_mac.sh b/Engine/Source/Programs/Enterprise/Datasmith/DatasmithSketchUpRubyExporter/Build/build_mac.sh
--- a/Engine/Source/Programs/Enterprise/Datasmith/DatasmithSketchUpRubyExporter/Build/build_mac.sh
+++ b/Engine/Source/Programs/Enterprise/Datasmith/DatasmithSketchUpRubyExporter/Build/build_mac.sh
@@ -38,6 +38,8 @@ cp "$EnginePath/Binaries/ThirdParty/Intel/TBB/Mac/$dylibLibTbb" "$Dylibs"
 cp "$EnginePath/Binaries/ThirdParty/Intel/TBB/Mac/$dylibLibTbbMalloc" "$Dylibs"
 
 chmod 777 "$Dylibs/$dylibLibFreeImage"
+chmod 777 "$Dylibs/$dylibLibTbb"
+chmod 777 "$Dylibs/$dylibLibTbbMalloc"
 
 install_name_tool -id @loader_path/Dylibs/DatasmithSDK.dylib "$Dylibs/DatasmithSDK.dylib"
 install_name_tool -id @loader_path/Dylibs/$dylibLibFreeImage "$Dylibs/$dylibLibFreeImage"
@@ -61,6 +63,8 @@ BuildSketchUpPlugin() {
     # Copy plugin files as they should b ein the plugin folder:
     # ruby code
     cp -r "$PluginSrcPath/Plugin" "$BuildDir"
+	# ruby code needs to be writable for the plugin to be "un-installable"
+	chmod -R 777 "$BuildDir"
     # support libs
     cp -r "$Dylibs" "$BuildDir/Plugin/UnrealDatasmithSketchUp"
     # resources
