3ds Max DirectLink: installed Datasmith Attributes to use by the plugin

- installed and loaded pre-startup(so it can be used by macroscripts)
- plugin was previously dependent on the Datasmith Attributes code loaded by the old Datasmith exporter

diff --git a/Engine/Source/Programs/Enterprise/Datasmith/DatasmithMaxExporter/Plugin/PackageContents.xml b/Engine/Source/Programs/Enterprise/Datasmith/DatasmithMaxExporter/Plugin/PackageContents.xml
--- a/Engine/Source/Programs/Enterprise/Datasmith/DatasmithMaxExporter/Plugin/PackageContents.xml
+++ b/Engine/Source/Programs/Enterprise/Datasmith/DatasmithMaxExporter/Plugin/PackageContents.xml
@@ -51,6 +51,7 @@
     <Components Description="pre-start-up scripts parts">
         <RuntimeRequirements OS="Win64" Platform="3ds Max" SeriesMin="2017" SeriesMax="2022" />
         <ComponentEntry ModuleName="./UI/pre-start-up.ms" />
+        <ComponentEntry ModuleName="./UI/UnrealStudio_Datasmith_Attributes_Holder_Modifier.ms" />
     </Components>
 
     <Components Description="macroscripts parts">
