3ds Max: Fixed node deletion with AutoSync

- handling deletion of nodes moved to Update, not only fixing autosync but also ensuring that this potentially heavy work(in case of many removed nodes) doesn't happen when user is editing the scene

diff --git a/Engine/Source/Programs/Enterprise/Datasmith/DatasmithMaxExporter/Private/DatasmithMaxExporterWithDirectLink.cpp b/Engine/Source/Programs/Enterprise/Datasmith/DatasmithMaxExporter/Private/DatasmithMaxExporterWithDirectLink.cpp
--- a/Engine/Source/Programs/Enterprise/Datasmith/DatasmithMaxExporter/Private/DatasmithMaxExporterWithDirectLink.cpp
+++ b/Engine/Source/Programs/Enterprise/Datasmith/DatasmithMaxExporter/Private/DatasmithMaxExporterWithDirectLink.cpp
@@ -151,7 +151,7 @@ public:
 class FNodeTrackerHandle
 {
 public:
-	explicit FNodeTrackerHandle(INode* InNode) : Impl(MakeShared<FNodeTracker>(InNode)) {}
+	explicit FNodeTrackerHandle(FNodeKey InNodeKey, INode* InNode) : Impl(MakeShared<FNodeTracker>(InNodeKey, InNode)) {}
 
 	FNodeTracker* GetNodeTracker() const
 	{
@@ -548,16 +548,26 @@ public:
 //			}
 		}
 
-
 		FNodeKey NodeKey = NodeEventNamespace::GetKeyByNode(Node);
 
-		FNodeTrackerHandle& NodeTracker = AddNode(NodeKey, Node);
-
-		// Parse children
-		int32 ChildNum = Node->NumberOfChildren();
-		for (int32 ChildIndex = 0; ChildIndex < ChildNum; ++ChildIndex)
+		if (FNodeTrackerHandle* NodeTrackerHandle = NodeTrackers.Find(NodeKey))
+		{
+			// Node being added might already be tracked(e.g. if it was deleted before but Update wasn't called to SceneTracker yet)
+			FNodeTracker* NodeTracker = NodeTrackerHandle->GetNodeTracker();
+			ensure(NodeTracker->bDeleted);
+			NodeTracker->bDeleted = false;
+			InvalidateNode(*NodeTracker);
+		}
+		else
 		{
-			ParseNode(Node->GetChildNode(ChildIndex));
+			FNodeTrackerHandle& NodeTracker = AddNode(NodeKey, Node);
+
+			// Parse children
+			int32 ChildNum = Node->NumberOfChildren();
+			for (int32 ChildIndex = 0; ChildIndex < ChildNum; ++ChildIndex)
+			{
+				ParseNode(Node->GetChildNode(ChildIndex));
+			}
 		}
 	}
 
@@ -658,6 +668,24 @@ public:
 		bChangeEncountered |= !InvalidatedNodeTrackers.IsEmpty();
 		bChangeEncountered |= !MaterialsCollectionTracker.GetInvalidatedMaterials().IsEmpty();
 
+
+		ProgressManager.ProgressStage(TEXT("Remove deleted nodes"));
+		{
+			TArray<FNodeTracker*> DeletedNodeTrackers;
+			for (FNodeTracker* NodeTracker : InvalidatedNodeTrackers)
+			{
+				if (NodeTracker->bDeleted)
+				{
+					DeletedNodeTrackers.Add(NodeTracker);
+				}
+			}
+
+			for (FNodeTracker* NodeTrackerPtr : DeletedNodeTrackers)
+			{
+				RemoveNodeTracker(*NodeTrackerPtr);
+			}
+		}
+
 		ProgressManager.ProgressStage(TEXT("Update node names"));
 		// todo: move to NameChanged and NodeAdded?
 		for (FNodeTracker* NodeTracker : InvalidatedNodeTrackers)
@@ -678,7 +706,6 @@ public:
 			for (FNodeTracker* NodeTracker : InvalidatedNodeTrackers)
 			{
 				ProgressCounter.Next();
-
 				UpdateCollisionStatus(NodeTracker, NodesWithChangedCollisionStatus);
 			}
 			InvalidatedNodeTrackers.Append(NodesWithChangedCollisionStatus);
@@ -827,7 +854,7 @@ public:
 	FORCENOINLINE
 	FNodeTrackerHandle& AddNode(FNodeKey NodeKey, INode* Node)
 	{
-		FNodeTrackerHandle& NodeTracker = NodeTrackers.Emplace(NodeKey, Node);
+		FNodeTrackerHandle& NodeTracker = NodeTrackers.Emplace(NodeKey, FNodeTrackerHandle(NodeKey, Node));
 		
 		NodeTrackersNames.FindOrAdd(NodeTracker.GetNodeTracker()->Name).Add(NodeTracker.GetNodeTracker());
 		InvalidatedNodeTrackers.Add(NodeTracker.GetNodeTracker());
@@ -1044,8 +1071,34 @@ public:
 		}
 	}
 
+	void RemoveNodeTracker(FNodeTracker& NodeTracker)
+	{
+		InvalidatedNodeTrackers.Remove(&NodeTracker);
+
+		RemoveFromConverted(NodeTracker);
+
+		if (TSet<FNodeTracker*>* NodeTrackersPtr = NodeTrackersNames.Find(NodeTracker.Name))
+		{
+			NodeTrackersPtr->Remove(&NodeTracker);
+		}
+
+		if (TSet<FNodeTracker*>* CollisionUsersPtr = CollisionNodes.Find(NodeTracker.Collision))
+		{
+			TSet<FNodeTracker*>& CollisionUsers = *CollisionUsersPtr;
+			CollisionUsers.Remove(&NodeTracker);
+
+			if (CollisionUsers.IsEmpty())
+			{
+				CollisionNodes.Remove(NodeTracker.Collision);
+			}
+		}
+
+		NodeTrackers.Remove(NodeTracker.NodeKey);
+	}
+
 	void UpdateNode(FNodeTracker& NodeTracker)
 	{
+
 		// Forget anything that this node was before update: place in datasmith hierarchy, datasmith objects, instances connection. Updating may change anything 
 		RemoveFromConverted(NodeTracker);
 		ConvertNodeObject(NodeTracker);
@@ -1669,26 +1722,8 @@ public:
 		{
 			// todo: schedule for delete on Update?
 			FNodeTracker* NodeTracker = NodeTrackerHandle->GetNodeTracker();
-			InvalidatedNodeTrackers.Remove(NodeTracker);
-			NodeTrackers.Remove(NodeKey);
-
-			if (TSet<FNodeTracker*>* NodeTrackersPtr = NodeTrackersNames.Find(NodeTracker->Name))
-			{
-				NodeTrackersPtr->Remove(NodeTracker);
-			}
-
-			if (TSet<FNodeTracker*>* CollisionUsersPtr = CollisionNodes.Find(NodeTracker->Collision))
-			{
-				TSet<FNodeTracker*>& CollisionUsers = *CollisionUsersPtr;
-				CollisionUsers.Remove(NodeTracker);
-
-				if (CollisionUsers.IsEmpty())
-				{
-					CollisionNodes.Remove(NodeTracker->Collision);
-				}
-			}
-
-			RemoveFromConverted(*NodeTracker);
+			InvalidatedNodeTrackers.Add(NodeTracker);
+			NodeTracker->bDeleted = true;
 		}
 	}
 
diff --git a/Engine/Source/Programs/Enterprise/Datasmith/DatasmithMaxExporter/Public/DatasmithMaxDirectLink.h b/Engine/Source/Programs/Enterprise/Datasmith/DatasmithMaxExporter/Public/DatasmithMaxDirectLink.h
--- a/Engine/Source/Programs/Enterprise/Datasmith/DatasmithMaxExporter/Public/DatasmithMaxDirectLink.h
+++ b/Engine/Source/Programs/Enterprise/Datasmith/DatasmithMaxExporter/Public/DatasmithMaxDirectLink.h
@@ -88,11 +88,15 @@ void ShutdownScripts();
 
 IPersistentExportOptions& GetPersistentExportOptions();
 
+// Max nodes are matched by these keys
+typedef NodeEventNamespace::NodeKey FNodeKey;
+typedef MtlBase* FMaterialKey;
+
 // Identifies Max node to track its changes
 class FNodeTracker: FNoncopyable
 {
 public:
-	explicit FNodeTracker(INode* InNode) : Node(InNode), Name(Node->GetName()) {}
+	explicit FNodeTracker(FNodeKey InNodeKey, INode* InNode) : NodeKey(InNodeKey), Node(InNode), Name(Node->GetName()) {}
 
 	void Invalidate()
 	{
@@ -109,12 +113,14 @@ public:
 		return InstanceHandle != 0;
 	}
 
+	FNodeKey NodeKey;
 	INode* const Node;
 	FNodeTracker* Collision = nullptr;
 	FLayerTracker* Layer = nullptr;
 	AnimHandle InstanceHandle = 0; // todo: rename - this is handle for object this node is instance of
 	FString Name; 
 	bool bInvalidated = true;
+	bool bDeleted = false;
 
 	TSharedPtr<IDatasmithActorElement> DatasmithActorElement;
 
@@ -124,10 +130,6 @@ public:
 
 };
 
-// Max nodes are matched by these keys
-typedef NodeEventNamespace::NodeKey FNodeKey;
-typedef MtlBase* FMaterialKey;
-
 class FRenderMeshForConversion: FNoncopyable
 {
 public:
