Reinstate LWC for ArchiCAD exporter

diff --git a/Engine/Source/Programs/Enterprise/Datasmith/DatasmithARCHICADExporter/DatsmithUE4ArchiCAD/DatasmithUE4ArchiCAD.Target.cs b/Engine/Source/Programs/Enterprise/Datasmith/DatasmithARCHICADExporter/DatsmithUE4ArchiCAD/DatasmithUE4ArchiCAD.Target.cs
--- a/Engine/Source/Programs/Enterprise/Datasmith/DatasmithARCHICADExporter/DatsmithUE4ArchiCAD/DatasmithUE4ArchiCAD.Target.cs
+++ b/Engine/Source/Programs/Enterprise/Datasmith/DatasmithARCHICADExporter/DatsmithUE4ArchiCAD/DatasmithUE4ArchiCAD.Target.cs
@@ -36,6 +36,5 @@ public class DatasmithUE4ArchiCADTarget : TargetRules
 		bHasExports = true;
 
 		GlobalDefinitions.Add("UE_EXTERNAL_PROFILING_ENABLED=0"); // For DirectLinkUI (see FDatasmithExporterManager::FInitOptions)
-		GlobalDefinitions.Add("UE_LARGE_WORLD_COORDINATES_DISABLED=1"); // LARGE WORLD COORDINATES
 	}
 }
diff --git a/Engine/Source/Programs/Enterprise/Datasmith/DatasmithARCHICADExporter/Private/DatasmithHashTools.cpp b/Engine/Source/Programs/Enterprise/Datasmith/DatasmithARCHICADExporter/Private/DatasmithHashTools.cpp
--- a/Engine/Source/Programs/Enterprise/Datasmith/DatasmithARCHICADExporter/Private/DatasmithHashTools.cpp
+++ b/Engine/Source/Programs/Enterprise/Datasmith/DatasmithARCHICADExporter/Private/DatasmithHashTools.cpp
@@ -20,11 +20,11 @@ float		FDatasmithHashTools::FixedPointTolerance = 1 / 0.01f;
 const float FDatasmithHashTools::MaxInvFixPointTolerance = FLOAT_NON_FRACTIONAL;
 
 // Hash float with a fixed tolerance value.
-void FDatasmithHashTools::HashFixedPointFloatTolerance(float InValue, float InvTolerance)
+void FDatasmithHashTools::HashFixedPointFloatTolerance(double InValue, float InvTolerance)
 {
 	if (InvTolerance < MaxInvFixPointTolerance)
 	{
-		float Rounded = std::floorf(InValue * InvTolerance + 0.5f);
+		float Rounded = std::floorf((float)InValue * InvTolerance + 0.5f);
 		if (Rounded == -0.0f)
 		{
 			Rounded = 0.0f; // We want to confond negative near zero to positive near zero
@@ -34,7 +34,7 @@ void FDatasmithHashTools::HashFixedPointFloatTolerance(float InValue, float InvT
 	else
 	{
 		// No tolerance -> we take the hash value itself
-		TUpdate(InValue);
+		TUpdate((float)InValue);
 	}
 }
 
@@ -43,15 +43,15 @@ float		FDatasmithHashTools::FloatTolerance = KINDA_SMALL_NUMBER;
 const float FDatasmithHashTools::MaxFloatTolerance = 1.0f / FLOAT_NON_FRACTIONAL;
 
 // Hash float value, Use tolerance to absorb compute error
-void FDatasmithHashTools::HashFloatTolerance(float InValue, float InTolerance)
+void FDatasmithHashTools::HashFloatTolerance(double InValue, float InTolerance)
 {
 	if (InTolerance > MaxFloatTolerance)
 	{
 		// Near zero or negative
-		if (InValue <= InTolerance)
+		if ((float)InValue <= InTolerance)
 		{
 			InValue = -InValue;
-			if (InValue <= InTolerance)
+			if ((float)InValue <= InTolerance)
 			{
 				TUpdate(0.0f);
 				return; // Near zero ?
@@ -60,7 +60,7 @@ void FDatasmithHashTools::HashFloatTolerance(float InValue, float InTolerance)
 		}
 
 		// Tricky way to hash value so that near values (because of compute error) give same hash values
-		float LogValue = log10(InValue);
+		float LogValue = log10((float)InValue);
 		float IntLogValue = std::floorf(LogValue / InTolerance + 0.5f);
 		if (IntLogValue == -0.0f)
 		{
@@ -71,7 +71,7 @@ void FDatasmithHashTools::HashFloatTolerance(float InValue, float InTolerance)
 	else
 	{
 		// No tolerance -> we take the hash value itself
-		TUpdate(InValue);
+		TUpdate((float)InValue);
 	}
 }
 
@@ -128,11 +128,11 @@ static bool QuatNeglectible(FQuat InQuat, float InTolerance)
 /* Quat components values usualy are in turn (-0.5..0.5).
  * Bring back to 0.0 to 1.0 and make 0.0 turn == 1.0 turn
  */
-static float ZeroOne(float x)
+static float ZeroOne(double x)
 {
-	if (x < 0.0f)
-		x = 1.0f + x;
-	return sqrt(abs(x - x * x)); // = sqrt(abs(x * (1 - x))
+	if (x < 0.0)
+		x = 1.0 + x;
+	return (float)FMath::Sqrt(FMath::Abs(x - x * x)); // = sqrt(abs(x * (1 - x))
 }
 
 // Hash quaternion value, taking care to absorb compute error
diff --git a/Engine/Source/Programs/Enterprise/Datasmith/DatasmithARCHICADExporter/Private/DatasmithHashTools.h b/Engine/Source/Programs/Enterprise/Datasmith/DatasmithARCHICADExporter/Private/DatasmithHashTools.h
--- a/Engine/Source/Programs/Enterprise/Datasmith/DatasmithARCHICADExporter/Private/DatasmithHashTools.h
+++ b/Engine/Source/Programs/Enterprise/Datasmith/DatasmithARCHICADExporter/Private/DatasmithHashTools.h
@@ -33,7 +33,7 @@ class FDatasmithHashTools
 	static const float MaxInvFixPointTolerance;
 
 	// Hash float with a fixed tolerance value (good for coordinates, colors).
-	void HashFixedPointFloatTolerance(float InValue, float InInvTolerance = FixedPointTolerance);
+	void HashFixedPointFloatTolerance(double InValue, float InInvTolerance = FixedPointTolerance);
 
 	// Recommended for values that come from computation, like normals, UVs, rotations...
 	static float FloatTolerance;
@@ -42,7 +42,7 @@ class FDatasmithHashTools
 	static const float MaxFloatTolerance;
 
 	// Hash float value, use tolerance to absorb compute error (good for normal, rotation...)
-	void HashFloatTolerance(float InValue, float InTolerance = FloatTolerance);
+	void HashFloatTolerance(double InValue, float InTolerance = FloatTolerance);
 
 	// Hash 3d point related value, taking care to absorb input error
 	void HashFixVector(const FVector& InVector, float InInvTolerance = FixedPointTolerance);
diff --git a/Engine/Source/Programs/Enterprise/Datasmith/DatasmithARCHICADExporter/Private/Element2StaticMesh.cpp b/Engine/Source/Programs/Enterprise/Datasmith/DatasmithARCHICADExporter/Private/Element2StaticMesh.cpp
--- a/Engine/Source/Programs/Enterprise/Datasmith/DatasmithARCHICADExporter/Private/Element2StaticMesh.cpp
+++ b/Engine/Source/Programs/Enterprise/Datasmith/DatasmithARCHICADExporter/Private/Element2StaticMesh.cpp
@@ -503,7 +503,7 @@ void FElement2StaticMesh::FillMesh(FDatasmithMesh* OutMesh)
 			for (int32 IndexComponent = 0; IndexComponent < 3; IndexComponent++)
 			{
 				const FVector& Normal = triangle.Normals[IndexComponent];
-				OutMesh->SetNormal(IndexFace * 3 + IndexComponent, Normal.X, Normal.Y, Normal.Z);
+				OutMesh->SetNormal(IndexFace * 3 + IndexComponent, (float)Normal.X, (float)Normal.Y, (float)Normal.Z);
 			}
 
 			OutMesh->SetFaceUV(IndexFace, UVChannel, triangle.UV0, triangle.UV1, triangle.UV2);
