Revit - Metadata not sent to Direct Link destination

diff --git a/Engine/Source/Programs/Enterprise/Datasmith/DatasmithRevitExporter/Private/DatasmithRevitDirectLink.cs b/Engine/Source/Programs/Enterprise/Datasmith/DatasmithRevitExporter/Private/DatasmithRevitDirectLink.cs
--- a/Engine/Source/Programs/Enterprise/Datasmith/DatasmithRevitExporter/Private/DatasmithRevitDirectLink.cs
+++ b/Engine/Source/Programs/Enterprise/Datasmith/DatasmithRevitExporter/Private/DatasmithRevitDirectLink.cs
@@ -745,7 +745,6 @@ namespace DatasmithRevitExporter
 						CurrentBatchSize = 0;
 
 						// Send metadata to DirectLink.
-						DatasmithScene.CleanUp();
 						DatasmithDirectLink.UpdateScene(DatasmithScene);
 
 						MetadataEvent.WaitOne(DelayExport);
@@ -780,6 +779,8 @@ namespace DatasmithRevitExporter
 
 					FUtils.AddActorMetadata(RevitElement, ElementData.ElementMetaData);
 
+					DatasmithScene.AddMetaData(ElementData.ElementMetaData);
+
 					++CurrentBatchSize;
 
 #if DEBUG
@@ -820,7 +821,6 @@ namespace DatasmithRevitExporter
 			if (CurrentBatchSize > 0)
 			{
 				// Send remaining chunk of metadata.
-				DatasmithScene.CleanUp();
 				DatasmithDirectLink.UpdateScene(DatasmithScene);
 			}
 
