Disable AutoSync from Revit exporter in 5.0

diff --git a/Engine/Source/Programs/Enterprise/Datasmith/DatasmithRevitExporter/Private/DatasmithRevitApplication.cs b/Engine/Source/Programs/Enterprise/Datasmith/DatasmithRevitExporter/Private/DatasmithRevitApplication.cs
--- a/Engine/Source/Programs/Enterprise/Datasmith/DatasmithRevitExporter/Private/DatasmithRevitApplication.cs
+++ b/Engine/Source/Programs/Enterprise/Datasmith/DatasmithRevitExporter/Private/DatasmithRevitApplication.cs
@@ -23,18 +23,19 @@ namespace DatasmithRevitExporter
 		private EventHandler<DocumentClosingEventArgs> DocumentClosingHandler;
 		private EventHandler<ViewActivatedEventArgs> ViewActivatedHandler;
 
-		PushButton AutoSyncPushButton;
+		//PushButton AutoSyncPushButton;
 		PushButton SyncPushButton;
 
-		BitmapImage AutoSyncIconOn_Small;
-		BitmapImage AutoSyncIconOn_Large;
-		BitmapImage AutoSyncIconOff_Small;
-		BitmapImage AutoSyncIconOff_Large;
+		//BitmapImage AutoSyncIconOn_Small;
+		//BitmapImage AutoSyncIconOn_Large;
+		//BitmapImage AutoSyncIconOff_Small;
+		//BitmapImage AutoSyncIconOff_Large;
 
 		public object Properties { get; private set; }
 
 		public static DatasmithRevitApplication Instance { get; private set; }
 
+#if false
 		public void SetAutoSyncButtonToggled(bool bToggled)
 		{
 			if (bToggled)
@@ -49,7 +50,7 @@ namespace DatasmithRevitExporter
 			}
 			SyncPushButton.Enabled = !bToggled;
 		}
-
+#endif
 		// Implement the interface to execute some tasks when Revit starts.
 		public Result OnStartup(
 			UIControlledApplication InApplication // handle to the application being started
@@ -69,13 +70,13 @@ namespace DatasmithRevitExporter
 			string AssemblyPath = Assembly.GetExecutingAssembly().Location;
 			PushButtonData ExportButtonData = new PushButtonData("Export3DView", DatasmithRevitResources.Strings.ButtonExport3DView, AssemblyPath, "DatasmithRevitExporter.DatasmithExportRevitCommand");
 			PushButtonData SyncButtonData = new PushButtonData("Sync3DView", DatasmithRevitResources.Strings.ButtonSync, AssemblyPath, "DatasmithRevitExporter.DatasmithSyncRevitCommand");
-			PushButtonData AutoSyncButtonData = new PushButtonData("AutoSync3DView", DatasmithRevitResources.Strings.ButtonAutoSync, AssemblyPath, "DatasmithRevitExporter.DatasmithAutoSyncRevitCommand");
+			//PushButtonData AutoSyncButtonData = new PushButtonData("AutoSync3DView", DatasmithRevitResources.Strings.ButtonAutoSync, AssemblyPath, "DatasmithRevitExporter.DatasmithAutoSyncRevitCommand");
 			PushButtonData ManageConnectionsButtonData = new PushButtonData("Connections", DatasmithRevitResources.Strings.ButtonConnections, AssemblyPath, "DatasmithRevitExporter.DatasmithManageConnectionsRevitCommand");
 			PushButtonData SettingsButtonData = new PushButtonData("Settings", DatasmithRevitResources.Strings.ButtonSettings, AssemblyPath, "DatasmithRevitExporter.DatasmithShowSettingsRevitCommand");
 			PushButtonData LogButtonData = new PushButtonData("Messages", DatasmithRevitResources.Strings.ButtonMessages, AssemblyPath, "DatasmithRevitExporter.DatasmithShowMessagesRevitCommand");
 
 			SyncPushButton = DirectLinkRibbonPanel.AddItem(SyncButtonData) as PushButton;
-			AutoSyncPushButton = DirectLinkRibbonPanel.AddItem(AutoSyncButtonData) as PushButton;
+			//AutoSyncPushButton = DirectLinkRibbonPanel.AddItem(AutoSyncButtonData) as PushButton;
 			PushButton ManageConnectionsButton = DirectLinkRibbonPanel.AddItem(ManageConnectionsButtonData) as PushButton;
 			PushButton ExportPushButton = FileExportRibbonPanel.AddItem(ExportButtonData) as PushButton;
 			PushButton ShowLogButton = DatasmithRibbonPanel.AddItem(LogButtonData) as PushButton;
@@ -93,14 +94,14 @@ namespace DatasmithRevitExporter
 
 			DatasmithIconBase = Path.Combine(Path.GetDirectoryName(AssemblyPath), "DatasmithAutoSyncIcon");
 
-			AutoSyncIconOn_Small = new BitmapImage(new Uri(DatasmithIconBase + "On16.png"));
-			AutoSyncIconOn_Large = new BitmapImage(new Uri(DatasmithIconBase + "On32.png"));
-			AutoSyncIconOff_Small = new BitmapImage(new Uri(DatasmithIconBase + "Off16.png"));
-			AutoSyncIconOff_Large = new BitmapImage(new Uri(DatasmithIconBase + "Off32.png"));
+			//AutoSyncIconOn_Small = new BitmapImage(new Uri(DatasmithIconBase + "On16.png"));
+			//AutoSyncIconOn_Large = new BitmapImage(new Uri(DatasmithIconBase + "On32.png"));
+			//AutoSyncIconOff_Small = new BitmapImage(new Uri(DatasmithIconBase + "Off16.png"));
+			//AutoSyncIconOff_Large = new BitmapImage(new Uri(DatasmithIconBase + "Off32.png"));
 
-			AutoSyncPushButton.Image = AutoSyncIconOn_Small;
-			AutoSyncPushButton.LargeImage = AutoSyncIconOn_Large;
-			AutoSyncPushButton.ToolTip = DatasmithRevitResources.Strings.ButtonAutoSyncHint;
+			//AutoSyncPushButton.Image = AutoSyncIconOn_Small;
+			//AutoSyncPushButton.LargeImage = AutoSyncIconOn_Large;
+			//AutoSyncPushButton.ToolTip = DatasmithRevitResources.Strings.ButtonAutoSyncHint;
 
 			DatasmithIconBase = Path.Combine(Path.GetDirectoryName(AssemblyPath), "DatasmithManageConnectionsIcon");
 			ManageConnectionsButton.Image = new BitmapImage(new Uri(DatasmithIconBase + "16.png"));
diff --git a/Engine/Source/Programs/Enterprise/Datasmith/DatasmithRevitExporter/Private/DatasmithRevitCommand.cs b/Engine/Source/Programs/Enterprise/Datasmith/DatasmithRevitExporter/Private/DatasmithRevitCommand.cs
--- a/Engine/Source/Programs/Enterprise/Datasmith/DatasmithRevitExporter/Private/DatasmithRevitCommand.cs
+++ b/Engine/Source/Programs/Enterprise/Datasmith/DatasmithRevitExporter/Private/DatasmithRevitCommand.cs
@@ -80,13 +80,6 @@ namespace DatasmithRevitExporter
 
 			Debug.Assert(FDirectLink.Get() != null);
 
-			// Holding ctrl will force full sync. 
-			if (FDirectLink.Get().SyncCount > 0 && (System.Windows.Forms.Control.ModifierKeys & Keys.Control) == Keys.Control)
-			{
-				FDirectLink.DestroyInstance(FDirectLink.Get(), InCommandData.Application.Application);
-				FDirectLink.ActivateInstance(Doc);
-			}
-
 			FDatasmithRevitExportContext ExportContext = new FDatasmithRevitExportContext(
 				InCommandData.Application.Application,
 				Doc,
@@ -133,6 +126,7 @@ namespace DatasmithRevitExporter
 		}
 	}
 
+#if false // AutoSync is temporary disabled
 	[Transaction(TransactionMode.Manual)]
 	public class DatasmithAutoSyncRevitCommand : DatasmithRevitCommand
 	{
@@ -145,6 +139,7 @@ namespace DatasmithRevitExporter
 			return Result.Succeeded;
 		}
 	}
+#endif
 
 	// Add-in external command Export to Unreal Datasmith. 
 	[Transaction(TransactionMode.Manual)]
