Fixing no UV channel exported when no customization is present

diff --git a/Engine/Source/Programs/Enterprise/Datasmith/DatasmithRhinoExporter/Private/ElementExporters/DatasmithRhinoMeshExporter.cs b/Engine/Source/Programs/Enterprise/Datasmith/DatasmithRhinoExporter/Private/ElementExporters/DatasmithRhinoMeshExporter.cs
--- a/Engine/Source/Programs/Enterprise/Datasmith/DatasmithRhinoExporter/Private/ElementExporters/DatasmithRhinoMeshExporter.cs
+++ b/Engine/Source/Programs/Enterprise/Datasmith/DatasmithRhinoExporter/Private/ElementExporters/DatasmithRhinoMeshExporter.cs
@@ -93,9 +93,10 @@ namespace DatasmithRhino.ElementExporters
 			int VertexIndexOffset = 0;
 			int FaceIndexOffset = 0;
 			int UVIndexOffset = 0;
+			int NumberOfUVChannels = System.Math.Max(1, MeshInfo.TextureMappings.Count);
 			List<Mesh> MeshSections = MeshInfo.RhinoMeshes;
 			List<DatasmithMaterialInfo> UniqueMaterialInfo = new List<DatasmithMaterialInfo>();
-			InitializeDatasmithMesh(DatasmithMesh, MeshSections, MeshInfo.TextureMappings.Count);
+			InitializeDatasmithMesh(DatasmithMesh, MeshSections, NumberOfUVChannels);
 
 			for (int MeshIndex = 0; MeshIndex < MeshSections.Count; ++MeshIndex )
 			{
@@ -166,26 +167,30 @@ namespace DatasmithRhino.ElementExporters
 					}
 				}
 
-				for (int UVChannel = 0; UVChannel < MeshInfo.TextureMappings.Count; ++UVChannel)
+				if (MeshInfo.TextureMappings.Count > 0)
 				{
-					// Since we are offsetting the meshes before exporting them, we must apply the same correction to the UV transform.
-					Transform InverveOffsetTranform;
-					MeshInfo.OffsetTransform.TryGetInverse(out InverveOffsetTranform);
-					DatasmithTextureMappingData TextureMappingData = MeshInfo.TextureMappings[UVChannel];
-					Transform CorrectedTransform = InverveOffsetTranform * TextureMappingData.ObjectTransform;
-
-					// Rhino gives no guarantee on the state of the texture mapping in a given mesh.
-					// We must make sure that UV is set to the channel we are exporting.
-					const bool bLazyLoad = false;
-					RhinoMesh.SetTextureCoordinates(TextureMappingData.RhinoTextureMapping, CorrectedTransform, bLazyLoad);
-
-					// Add the UV coordinates for the triangles we just added.
-					for (int UVIndex = 0; UVIndex < RhinoMesh.TextureCoordinates.Count; ++UVIndex)
+					for (int UVChannel = 0; UVChannel < MeshInfo.TextureMappings.Count; ++UVChannel)
 					{
-						Point2f UV = RhinoMesh.TextureCoordinates[UVIndex];
-						DatasmithMesh.SetUV(UVChannel, UVIndex + UVIndexOffset, UV.X, 1 - UV.Y);
+						// Since we are offsetting the meshes before exporting them, we must apply the same correction to the UV transform.
+						Transform InverveOffsetTranform;
+						MeshInfo.OffsetTransform.TryGetInverse(out InverveOffsetTranform);
+						DatasmithTextureMappingData TextureMappingData = MeshInfo.TextureMappings[UVChannel];
+						Transform CorrectedTransform = InverveOffsetTranform * TextureMappingData.ObjectTransform;
+
+						// Rhino gives no guarantee on the state of the texture mapping in a given mesh.
+						// We must make sure that UV is set to the channel we are exporting.
+						const bool bLazyLoad = false;
+						RhinoMesh.SetTextureCoordinates(TextureMappingData.RhinoTextureMapping, CorrectedTransform, bLazyLoad);
+
+						// Add the UV coordinates to the current channel.
+						AddUVsToMesh(DatasmithMesh, RhinoMesh, UVChannel, UVIndexOffset);
 					}
 				}
+				else
+				{
+					// No custom TextureMapping, just export the current UV coordinate in channel 0.
+					AddUVsToMesh(DatasmithMesh, RhinoMesh, 0, UVIndexOffset);
+				}
 
 				VertexIndexOffset += RhinoMesh.Vertices.Count;
 				FaceIndexOffset += RhinoMesh.Faces.Count + RhinoMesh.Faces.QuadCount;
@@ -228,5 +233,15 @@ namespace DatasmithRhino.ElementExporters
 			Mesh.SetNormal(NormalIndex + 1, NormalB.X, NormalB.Y, NormalB.Z);
 			Mesh.SetNormal(NormalIndex + 2, NormalC.X, NormalC.Y, NormalC.Z);
 		}
+
+		private static void AddUVsToMesh(FDatasmithFacadeMesh DatasmithMesh, Mesh RhinoMesh, int UVChannel, int UVIndexOffset)
+		{
+			// Add the UV coordinates for the triangles we just added.
+			for (int UVIndex = 0; UVIndex < RhinoMesh.TextureCoordinates.Count; ++UVIndex)
+			{
+				Point2f UV = RhinoMesh.TextureCoordinates[UVIndex];
+				DatasmithMesh.SetUV(UVChannel, UVIndex + UVIndexOffset, UV.X, 1 - UV.Y);
+			}
+		}
 	}
 }
\ No newline at end of file
