DatasmithRhinoExporter - Fixing material assignment for recursive "parent material" in nested blocks.

diff --git a/Engine/Source/Programs/Enterprise/Datasmith/DatasmithRhinoExporter/Private/DatasmithRhinoExportContext.cs b/Engine/Source/Programs/Enterprise/Datasmith/DatasmithRhinoExporter/Private/DatasmithRhinoExportContext.cs
--- a/Engine/Source/Programs/Enterprise/Datasmith/DatasmithRhinoExporter/Private/DatasmithRhinoExportContext.cs
+++ b/Engine/Source/Programs/Enterprise/Datasmith/DatasmithRhinoExporter/Private/DatasmithRhinoExportContext.cs
@@ -1579,7 +1579,7 @@ namespace DatasmithRhino
 			string Name = string.Format("{0}_{1}", InstanceParentNodeInfo.Name, DefinitionNodeInfo.Name);
 			string BaseLabel = DefinitionNodeInfo.BaseLabel;
 			string UniqueLabel = ActorLabelGenerator.GenerateUniqueNameFromBaseName(DefinitionNodeInfo.UniqueLabel);
-			bool bOverrideMaterial = DefinitionNodeInfo.MaterialIndex != MaterialIndex;
+			bool bOverrideMaterial = DefinitionNodeInfo.MaterialIndex != MaterialIndex || DefinitionNodeInfo.bOverrideMaterial;
 
 			return new DatasmithActorInfo(DefinitionNodeInfo.RhinoCommonObject as ModelComponent, Name, UniqueLabel, BaseLabel, MaterialIndex, bOverrideMaterial, VisibilityLayer);
 		}
