Solidworks Datasmith exporter slow or no exports created

diff --git a/Engine/Source/Programs/Enterprise/Datasmith/DatasmithSolidworksExporter/Private/DatasmithExporter.cs b/Engine/Source/Programs/Enterprise/Datasmith/DatasmithSolidworksExporter/Private/DatasmithExporter.cs
--- a/Engine/Source/Programs/Enterprise/Datasmith/DatasmithSolidworksExporter/Private/DatasmithExporter.cs
+++ b/Engine/Source/Programs/Enterprise/Datasmith/DatasmithSolidworksExporter/Private/DatasmithExporter.cs
@@ -1,5 +1,6 @@
 ﻿// Copyright Epic Games, Inc. All Rights Reserved.
 
+using SolidWorks.Interop.sldworks;
 using System;
 using System.Collections.Concurrent;
 using System.Collections.Generic;
@@ -483,7 +484,6 @@ namespace DatasmithSolidworks
 
 		public void ExportAnimation(FAnimation InAnim)
 		{
-#if false // Animation export is not available for 4.27
 			FDatasmithFacadeLevelSequence LevelSeq = new FDatasmithFacadeLevelSequence(InAnim.Name);
 
 			LevelSeq.SetFrameRate(InAnim.FPS);
@@ -529,7 +529,6 @@ namespace DatasmithSolidworks
 			}
 
 			DatasmithScene.AddLevelSequence(LevelSeq);
-#endif
 		}
 
 		private FDatasmithFacadeTexture ExportDiffuseMap(FMaterial InMaterial, FDatasmithFacadeMasterMaterial InMasterMaterial, string InParamName)
diff --git a/Engine/Source/Programs/Enterprise/Datasmith/DatasmithSolidworksExporter/Private/Document.cs b/Engine/Source/Programs/Enterprise/Datasmith/DatasmithSolidworksExporter/Private/Document.cs
--- a/Engine/Source/Programs/Enterprise/Datasmith/DatasmithSolidworksExporter/Private/Document.cs
+++ b/Engine/Source/Programs/Enterprise/Datasmith/DatasmithSolidworksExporter/Private/Document.cs
@@ -30,10 +30,6 @@ namespace DatasmithSolidworks
 		private FDatasmithFacadeDirectLink DatasmithDirectLink = null;
 		private bool bDocumentIsDirty = true;
 		private string DatasmithFileExportPath = "";
-		private ManualResetEvent DirectLinkSyncEvent = new ManualResetEvent(false);
-		private ManualResetEvent ExportToFileEvent = new ManualResetEvent(false);
-		private ManualResetEvent ExitExportThreadEvent = new ManualResetEvent(false);
-		private Thread SceneExportThread = null;
 		private Thread MaterialCheckerThread = null;
 		private ManualResetEvent MaterialCheckerEvent = null;
 		private bool bExitMaterialUpdateThread = false;
@@ -128,86 +124,74 @@ namespace DatasmithSolidworks
 			}
 		}
 
-		private void SceneExportProc()
+		private void RunExport(bool bIsDirectLinkExport)
 		{
-			while (true)
+			if (bIsDirectLinkExport)
 			{
-				int EventIndex = WaitHandle.WaitAny(new WaitHandle[] { ExitExportThreadEvent, DirectLinkSyncEvent, ExportToFileEvent });
+				// DirectLink sync
 
-				if (EventIndex == 0)
-				{
-					return;
-				}
-
-				if (EventIndex == 1)
-				{
-					// DirectLink sync
-
-					bDirectLinkSyncInProgress = true;
+				bDirectLinkSyncInProgress = true;
 
-					DatasmithScene.SetName(SwDoc.GetTitle());
-					DatasmithScene.SetOutputPath(DirectLinkPath);
+				DatasmithScene.SetName(SwDoc.GetTitle());
+				DatasmithScene.SetOutputPath(DirectLinkPath);
 
 #if DEBUG
-					Stopwatch Watch = Stopwatch.StartNew();
+				Stopwatch Watch = Stopwatch.StartNew();
 #endif
 
-					ExportToDatasmithScene();
+				ExportToDatasmithScene();
 
 #if DEBUG
-					Watch.Stop();
-					Debug.WriteLine($"EXPORT TIME: {(double)Watch.ElapsedMilliseconds / 1000.0}");
+				Watch.Stop();
+				Debug.WriteLine($"EXPORT TIME: {(double)Watch.ElapsedMilliseconds / 1000.0}");
 #endif
 
-					DatasmithDirectLink.UpdateScene(DatasmithScene);
+				DatasmithDirectLink.UpdateScene(DatasmithScene);
+			}
+			else
+			{
+				// Export to file
 
-					DirectLinkSyncEvent.Reset();
-				}
-				else
-				{
-					// Export to file
+				bFileExportInProgress = true;
 
-					bFileExportInProgress = true;
 
-					string OutDir = Path.GetDirectoryName(DatasmithFileExportPath);
-					string CleanFileName = Path.GetFileNameWithoutExtension(DatasmithFileExportPath);
+				string OutDir = Path.GetDirectoryName(DatasmithFileExportPath);
+				string CleanFileName = Path.GetFileNameWithoutExtension(DatasmithFileExportPath);
 
-					// Save/restore scene and exporter in order not to mess up DirectLink state
-					FDatasmithFacadeScene OldScene = DatasmithScene;
-					FDatasmithExporter OldExporter = Exporter;
+				// Save/restore scene and exporter in order not to mess up DirectLink state
+				FDatasmithFacadeScene OldScene = DatasmithScene;
+				FDatasmithExporter OldExporter = Exporter;
 
-					DatasmithScene = new FDatasmithFacadeScene("StdMaterial", "Solidworks", "Solidworks", "2021");
-					DatasmithScene.SetName(CleanFileName);
-					DatasmithScene.SetOutputPath(OutDir);
+				DatasmithScene = new FDatasmithFacadeScene("StdMaterial", "Solidworks", "Solidworks", "2021");
+				DatasmithScene.SetName(CleanFileName);
+				DatasmithScene.SetOutputPath(OutDir);
 
-					Exporter = new FDatasmithExporter(DatasmithScene);
+				Exporter = new FDatasmithExporter(DatasmithScene);
 
-					ExportToDatasmithScene();
+				ExportToDatasmithScene();
 
-					DatasmithScene.PreExport();
-					DatasmithScene.ExportScene(DatasmithFileExportPath);
 
-					DatasmithScene = OldScene;
-					Exporter = OldExporter;
+				DatasmithScene.PreExport();
+				DatasmithScene.ExportScene(DatasmithFileExportPath);
 
-					ExportToFileEvent.Reset();
-				}
-				
-				SetExportStatus("Done");
+				DatasmithScene = OldScene;
+				Exporter = OldExporter;
+			}
 
-				SetDirty(false);
+			SetExportStatus("Done");
 
-				bDirectLinkSyncInProgress = false;
-				bFileExportInProgress = false;
+			SetDirty(false);
 
-				// Kickoff material checker thread after first export
-				if (MaterialCheckerThread == null)
-				{
-					MaterialCheckerThread = new Thread(CheckForMaterialUpdatesProc);
-					MaterialCheckerEvent = new ManualResetEvent(false);
-					MaterialCheckerThread.Start();
-					MaterialCheckerEvent.Set();
-				}
+			bDirectLinkSyncInProgress = false;
+			bFileExportInProgress = false;
+
+			// Kickoff material checker thread after first export
+			if (MaterialCheckerThread == null)
+			{
+				MaterialCheckerThread = new Thread(CheckForMaterialUpdatesProc);
+				MaterialCheckerEvent = new ManualResetEvent(false);
+				MaterialCheckerThread.Start();
+				MaterialCheckerEvent.Set();
 			}
 		}
 
@@ -220,13 +204,8 @@ namespace DatasmithSolidworks
 
 		public virtual void Destroy()
 		{
-			ExitExportThreadEvent.Set();
 			bExitMaterialUpdateThread = true;
 
-			if (SceneExportThread != null && !SceneExportThread.Join(1500))
-			{
-				SceneExportThread.Abort();
-			}
 			if (MaterialCheckerThread != null && !MaterialCheckerThread.Join(500))
 			{
 				MaterialCheckerThread.Abort();
@@ -251,12 +230,6 @@ namespace DatasmithSolidworks
 					}
 				}
 
-				if (SceneExportThread == null)
-				{
-					SceneExportThread = new Thread(SceneExportProc);
-					SceneExportThread.Start();
-				}
-
 				MaterialCheckerEvent?.Set();
 			}
 			else
@@ -272,7 +245,7 @@ namespace DatasmithSolidworks
 		public void OnExportToFile(string InFilePath)
 		{
 			DatasmithFileExportPath = InFilePath;
-			ExportToFileEvent.Set();
+			RunExport(false);
 		}
 
 		public void OnDirectLinkSync()
@@ -285,7 +258,7 @@ namespace DatasmithSolidworks
 			}
 
 			DirectLinkSyncCount++;
-			DirectLinkSyncEvent.Set();
+			RunExport(true);
 		}
 
 		public void OnIdle()
