3ds Max DirectLink: Physical Material with ClearCoat to UEPbr

- Converted Max Physical Material to Datasmith UEPbr, including Clear coat

- Testcases:
- Opaque
- Simple Transparent Max material exported as simple transparency(no refraction, ior)
- Advanced Transparent Max material exported as refractive material
- ClearCoat(ior, roughness used)
- Added ThinTranslucent variant(when Max material is ThinWalled or IOR it 1.0)
- Reflection converted to Specular

- Additionally: Fixed transparency color map, emission map/color map that was broken

diff --git a/Engine/Source/Programs/Enterprise/Datasmith/DatasmithMaxExporter/Private/DatasmithMaxArnold.cpp b/Engine/Source/Programs/Enterprise/Datasmith/DatasmithMaxExporter/Private/DatasmithMaxArnold.cpp
--- a/Engine/Source/Programs/Enterprise/Datasmith/DatasmithMaxExporter/Private/DatasmithMaxArnold.cpp
+++ b/Engine/Source/Programs/Enterprise/Datasmith/DatasmithMaxExporter/Private/DatasmithMaxArnold.cpp
@@ -424,15 +424,15 @@ void FDatasmithMaxMatWriter::ExportPhysicalMaterial(TSharedRef< IDatasmithScene
 			{
 				TransparencyMap = ParamBlock2->GetTexmap(ParamDefinition.ID, GetCOREInterface()->GetTime());
 			}
-			else if (FCString::Stricmp(ParamDefinition.int_name, TEXT("Transparency_Color_Map")) == 0)
+			else if (FCString::Stricmp(ParamDefinition.int_name, TEXT("trans_color_map")) == 0)
 			{
 				TransparencyColorMap = ParamBlock2->GetTexmap(ParamDefinition.ID, GetCOREInterface()->GetTime());
 			}
-			else if (FCString::Stricmp(ParamDefinition.int_name, TEXT("Emittance_Map")) == 0)
+			else if (FCString::Stricmp(ParamDefinition.int_name, TEXT("Emission_Map")) == 0)
 			{
 				EmittanceMap = ParamBlock2->GetTexmap(ParamDefinition.ID, GetCOREInterface()->GetTime());
 			}
-			else if (FCString::Stricmp(ParamDefinition.int_name, TEXT("Emittance_Color_Map")) == 0)
+			else if (FCString::Stricmp(ParamDefinition.int_name, TEXT("Emission_Color_Map")) == 0)
 			{
 				EmittanceColorMap = ParamBlock2->GetTexmap(ParamDefinition.ID, GetCOREInterface()->GetTime());
 			}
diff --git a/Engine/Source/Programs/Enterprise/Datasmith/DatasmithMaxExporter/Private/MaxMaterialsToUEPbr/DatasmithMaxMaterialsToUEPbr.cpp b/Engine/Source/Programs/Enterprise/Datasmith/DatasmithMaxExporter/Private/MaxMaterialsToUEPbr/DatasmithMaxMaterialsToUEPbr.cpp
--- a/Engine/Source/Programs/Enterprise/Datasmith/DatasmithMaxExporter/Private/MaxMaterialsToUEPbr/DatasmithMaxMaterialsToUEPbr.cpp
+++ b/Engine/Source/Programs/Enterprise/Datasmith/DatasmithMaxExporter/Private/MaxMaterialsToUEPbr/DatasmithMaxMaterialsToUEPbr.cpp
@@ -11,7 +11,9 @@
 #include "MaxMaterialsToUEPbr/DatasmithMaxCoronaMaterialsToUEPbr.h"
 #include "MaxMaterialsToUEPbr/DatasmithMaxScanlineMaterialsToUEPbr.h"
 #include "MaxMaterialsToUEPbr/DatasmithMaxVrayMaterialsToUEPbr.h"
+#include "MaxMaterialsToUEPbr/DatasmithMaxPhysicalMaterialToUEPbr.h"
 
+#define PHYSICALMATCLASS_UEPbr
 
 FDatasmithMaxMaterialsToUEPbr* FDatasmithMaxMaterialsToUEPbrManager::GetMaterialConverter( Mtl* Material )
 {
@@ -63,6 +65,14 @@ FDatasmithMaxMaterialsToUEPbr* FDatasmithMaxMaterialsToUEPbrManager::GetMaterial
 		static FDatasmithMaxCoronaBlendMaterialToUEPbr CoronaConverter = FDatasmithMaxCoronaBlendMaterialToUEPbr();
 		MaterialConverter = &CoronaConverter;
 	}
+#ifdef PHYSICALMATCLASS_UEPbr
+
+	else if (MaterialClassID == PHYSICALMATCLASS)
+	{
+		static FDatasmithMaxPhysicalMaterialToUEPbr PhysicalConverter = FDatasmithMaxPhysicalMaterialToUEPbr();
+		MaterialConverter = &PhysicalConverter;
+	}
+#endif
 
 	if ( MaterialConverter && MaterialConverter->IsSupported( Material ) )
 	{
diff --git a/Engine/Source/Programs/Enterprise/Datasmith/DatasmithMaxExporter/Private/MaxMaterialsToUEPbr/DatasmithMaxPhysicalMaterialToUEPbr.cpp b/Engine/Source/Programs/Enterprise/Datasmith/DatasmithMaxExporter/Private/MaxMaterialsToUEPbr/DatasmithMaxPhysicalMaterialToUEPbr.cpp
new file mode 100644
--- /dev/null
+++ b/Engine/Source/Programs/Enterprise/Datasmith/DatasmithMaxExporter/Private/MaxMaterialsToUEPbr/DatasmithMaxPhysicalMaterialToUEPbr.cpp
@@ -0,0 +1,768 @@
+// Copyright Epic Games, Inc. All Rights Reserved.
+
+#include "MaxMaterialsToUEPbr/DatasmithMaxPhysicalMaterialToUEPbr.h"
+
+#include "DatasmithMaxClassIDs.h"
+#include "DatasmithMaxTexmapParser.h"
+#include "DatasmithMaxWriter.h"
+#include "DatasmithSceneFactory.h"
+#include "MaxMaterialsToUEPbr/DatasmithMaxTexmapToUEPbr.h"
+#include "MaxMaterialsToUEPbr/DatasmithMaxCoronaTexmapToUEPbr.h"
+
+#include "Windows/AllowWindowsPlatformTypes.h"
+MAX_INCLUDES_START
+	#include "max.h"
+	#include "mtl.h"
+MAX_INCLUDES_END
+#include "Windows/HideWindowsPlatformTypes.h"
+
+// Create expression from expression parameters when they are all non-null or return default value
+#define COMPOSE_OR_DEFAULT1(Default, Func, Param0) (Param0 ? &Func(*Param0) : Default)
+#define COMPOSE_OR_DEFAULT2(Default, Func, Param0, Param1) ((Param0 && Param1) ? &Func(*Param0, *Param1) : Default)
+
+TSharedPtr<IDatasmithUEPbrMaterialElement> FDatasmithMaxMaterialsToUEPbrExpressions::GetMaterialElement()
+{
+	return ConvertState.MaterialElement;
+}
+
+IDatasmithMaterialExpressionScalar& FDatasmithMaxMaterialsToUEPbrExpressions::Scalar(float Value)
+{
+	IDatasmithMaterialExpressionScalar* ScalarExpression = GetMaterialElement()->AddMaterialExpression<IDatasmithMaterialExpressionScalar>();
+	ScalarExpression->GetScalar() = Value;
+	return *ScalarExpression;
+}
+
+IDatasmithMaterialExpressionColor& FDatasmithMaxMaterialsToUEPbrExpressions::Color(const FLinearColor& Value)
+{
+	IDatasmithMaterialExpressionColor* Result = GetMaterialElement()->AddMaterialExpression<IDatasmithMaterialExpressionColor>();
+	Result->GetColor() = Value;
+	return *Result;
+}
+
+IDatasmithMaterialExpression* FDatasmithMaxMaterialsToUEPbrExpressions::WeightTextureOrScalar(
+	const DatasmithMaxTexmapParser::FMapParameter& TextureWeight, float Weight)
+{
+	if (IDatasmithMaterialExpression* WeightMapExpression = ConvertTexmap(TextureWeight))
+	{
+		return WeightMapExpression;
+	}
+
+	if (!FMath::IsNearlyEqual(Weight, 1.0f))
+	{
+		return &Scalar(Weight);
+	}
+		
+	return nullptr;
+}
+
+IDatasmithMaterialExpressionGeneric& FDatasmithMaxMaterialsToUEPbrExpressions::Add(IDatasmithMaterialExpression& A,
+	IDatasmithMaterialExpression& B)
+{
+	IDatasmithMaterialExpressionGeneric* Result = GetMaterialElement()->AddMaterialExpression<IDatasmithMaterialExpressionGeneric>();
+	Result->SetExpressionName(TEXT("Add"));
+
+	A.ConnectExpression(*Result->GetInput(0));
+	B.ConnectExpression(*Result->GetInput(1));
+	return *Result;
+}
+
+IDatasmithMaterialExpressionGeneric& FDatasmithMaxMaterialsToUEPbrExpressions::Subtract(IDatasmithMaterialExpression& A,
+	IDatasmithMaterialExpression& B)
+{
+	IDatasmithMaterialExpressionGeneric* Result = GetMaterialElement()->AddMaterialExpression<IDatasmithMaterialExpressionGeneric>();
+	Result->SetExpressionName(TEXT("Subtract"));
+
+	A.ConnectExpression(*Result->GetInput(0));
+	B.ConnectExpression(*Result->GetInput(1));
+	return *Result;
+}
+
+IDatasmithMaterialExpressionGeneric& FDatasmithMaxMaterialsToUEPbrExpressions::Multiply(IDatasmithMaterialExpression& A, IDatasmithMaterialExpression& B)
+{
+	IDatasmithMaterialExpressionGeneric* Result = GetMaterialElement()->AddMaterialExpression<IDatasmithMaterialExpressionGeneric>();
+	Result->SetExpressionName(TEXT("Multiply"));
+
+	A.ConnectExpression(*Result->GetInput(0));
+	B.ConnectExpression(*Result->GetInput(1));
+	return *Result;
+}
+
+IDatasmithMaterialExpressionGeneric& FDatasmithMaxMaterialsToUEPbrExpressions::Divide(IDatasmithMaterialExpression& A,
+	IDatasmithMaterialExpression& B)
+{
+	IDatasmithMaterialExpressionGeneric* Result = GetMaterialElement()->AddMaterialExpression<IDatasmithMaterialExpressionGeneric>();
+	Result->SetExpressionName(TEXT("Divide"));
+
+	A.ConnectExpression(*Result->GetInput(0));
+	B.ConnectExpression(*Result->GetInput(1));
+	return *Result;
+}
+
+IDatasmithMaterialExpressionGeneric& FDatasmithMaxMaterialsToUEPbrExpressions::Desaturate(IDatasmithMaterialExpression& A)
+{
+	IDatasmithMaterialExpressionGeneric* Result = GetMaterialElement()->AddMaterialExpression<IDatasmithMaterialExpressionGeneric>();
+	Result->SetExpressionName(TEXT("Desaturation"));
+
+	A.ConnectExpression(*Result->GetInput(0));
+	return *Result;
+}
+
+IDatasmithMaterialExpressionGeneric& FDatasmithMaxMaterialsToUEPbrExpressions::Power(IDatasmithMaterialExpression& A,
+	IDatasmithMaterialExpression& B)
+{
+	IDatasmithMaterialExpressionGeneric* Result = GetMaterialElement()->AddMaterialExpression<IDatasmithMaterialExpressionGeneric>();
+	Result->SetExpressionName(TEXT("Power"));
+
+	A.ConnectExpression(*Result->GetInput(0));
+	B.ConnectExpression(*Result->GetInput(1));
+	return *Result;
+}
+
+IDatasmithMaterialExpressionGeneric& FDatasmithMaxMaterialsToUEPbrExpressions::Lerp(IDatasmithMaterialExpression& A,
+	IDatasmithMaterialExpression& B, IDatasmithMaterialExpression& Alpha)
+{
+	IDatasmithMaterialExpressionGeneric* Result = GetMaterialElement()->AddMaterialExpression<IDatasmithMaterialExpressionGeneric>();
+	Result->SetExpressionName(TEXT("LinearInterpolate"));
+
+	A.ConnectExpression(*Result->GetInput(0));
+	B.ConnectExpression(*Result->GetInput(1));
+	Alpha.ConnectExpression(*Result->GetInput(2));
+	return *Result;
+}
+
+IDatasmithMaterialExpressionGeneric& FDatasmithMaxMaterialsToUEPbrExpressions::Fresnel(IDatasmithMaterialExpression* Exponent,
+	IDatasmithMaterialExpression* BaseReflectFraction)
+{
+	IDatasmithMaterialExpressionGeneric* Fresnel = GetMaterialElement()->AddMaterialExpression<IDatasmithMaterialExpressionGeneric>();
+	Fresnel->SetExpressionName(TEXT("Fresnel"));
+
+	Connect(*Fresnel->GetInput(0), Exponent);
+	Connect(*Fresnel->GetInput(1), BaseReflectFraction);
+	return *Fresnel;
+}
+
+IDatasmithMaterialExpression* FDatasmithMaxMaterialsToUEPbrExpressions::ApplyWeightExpression(
+	IDatasmithMaterialExpression* ValueExpression, IDatasmithMaterialExpression* WeightExpression)
+{
+	return COMPOSE_OR_DEFAULT2(ValueExpression, Multiply, ValueExpression, WeightExpression);
+}
+
+IDatasmithMaterialExpression& FDatasmithMaxMaterialsToUEPbrExpressions::CalcIORComplex(double IORn, double IORk,
+	IDatasmithMaterialExpression& ToBeConnected90, IDatasmithMaterialExpression& ToBeConnected0)
+{
+	if (IORn >= 16)
+	{
+		return ToBeConnected90;
+	}
+
+	IORn = FMath::Max(IORn, 1.0f);
+
+	IDatasmithMaterialExpression& FresnelOrig = Fresnel(&Scalar(1.0), &Scalar(0.0));
+	IDatasmithMaterialExpression& Fresnel = OneMinus(FresnelOrig);
+	IDatasmithMaterialExpression& Fresnel2 = Multiply(Fresnel, Fresnel);
+
+	IDatasmithMaterialExpression& ConstantIORn = Scalar(IORn);
+	// ConstantIORn->Desc = TEXT("IOR Value");
+	IDatasmithMaterialExpression& ConstantIORn2 = Multiply(ConstantIORn, ConstantIORn);
+
+	IDatasmithMaterialExpression& ConstantIORk = Scalar(IORk);
+	IDatasmithMaterialExpression& AddN2K2 = (IORk == 0) ? ConstantIORn2 : Add(ConstantIORn2, Multiply(ConstantIORk, ConstantIORk));
+
+	IDatasmithMaterialExpression& Mul2NCPre = Multiply(ConstantIORn, Scalar(2.0));
+	IDatasmithMaterialExpression& Mul2NC = Multiply(Mul2NCPre, Fresnel);
+
+	IDatasmithMaterialExpression& RsNumPre = Subtract(AddN2K2, Mul2NC);
+	IDatasmithMaterialExpression& RsNum = Add(RsNumPre, Fresnel2);
+	IDatasmithMaterialExpression& RsDenPre = Add(AddN2K2, Mul2NC);
+	IDatasmithMaterialExpression& RsDen = Add(RsDenPre, Fresnel2);
+	IDatasmithMaterialExpression& Rs = Divide(RsNum, RsDen);
+
+	IDatasmithMaterialExpression& AddN2K2MulC2 = Multiply(AddN2K2, Fresnel2);
+
+	IDatasmithMaterialExpression& RpNumPre = Subtract(AddN2K2MulC2, Mul2NC);
+	IDatasmithMaterialExpression& RpNum = Add(RpNumPre, Scalar(1.0));
+	IDatasmithMaterialExpression& RpDenPre = Add(AddN2K2MulC2, Mul2NC);
+	IDatasmithMaterialExpression& RpDen = Add(RpDenPre, Scalar(1.0));
+	IDatasmithMaterialExpression& Rp = Divide(RpNum, RpDen);
+
+	IDatasmithMaterialExpression& ResPre = Add(Rp, Rs);
+	IDatasmithMaterialExpression& Res = Multiply(ResPre, Scalar(0.5));
+	Res = Power(Res, Scalar(0.5));
+
+	return Lerp(ToBeConnected0, ToBeConnected90, Res);
+}
+
+void FDatasmithMaxMaterialsToUEPbrExpressions::Connect(IDatasmithExpressionInput& Input, IDatasmithMaterialExpression& ValueExpression)
+{
+	ValueExpression.ConnectExpression(Input);
+}
+
+bool FDatasmithMaxMaterialsToUEPbrExpressions::Connect(IDatasmithExpressionInput& Input, IDatasmithMaterialExpression* ValueExpression)
+{
+	if (ValueExpression)
+	{
+		Connect(Input, *ValueExpression);
+		return true;
+	}
+	return false;
+}
+
+IDatasmithMaterialExpression* FDatasmithMaxMaterialsToUEPbrExpressions::TextureOrColor(const TCHAR* Name,
+	const DatasmithMaxTexmapParser::FMapParameter& Map, FLinearColor Color)
+{
+	return FDatasmithMaxTexmapToUEPbrUtils::MapOrValue(this, Map, Name, Color, TOptional<float>());
+}
+
+IDatasmithMaterialExpression* FDatasmithMaxMaterialsToUEPbrExpressions::TextureOrScalar(const TCHAR* Name,
+	const DatasmithMaxTexmapParser::FMapParameter& Map, float Value)
+{
+	return FDatasmithMaxTexmapToUEPbrUtils::MapOrValue(this, Map, Name, TOptional<FLinearColor>(), Value);
+}
+
+IDatasmithMaterialExpressionGeneric& FDatasmithMaxMaterialsToUEPbrExpressions::OneMinus(
+	IDatasmithMaterialExpression& Expression)
+{
+	IDatasmithMaterialExpressionGeneric* Result = GetMaterialElement()->AddMaterialExpression< IDatasmithMaterialExpressionGeneric >();
+	Result->SetExpressionName(TEXT("OneMinus"));
+	Connect(*Result->GetInput(0), &Expression);
+	return *Result;
+}
+
+bool FDatasmithMaxPhysicalMaterialToUEPbr::IsSupported( Mtl* Material )
+{
+	return true;
+}
+
+struct FMaxPhysicalMaterial
+{
+	bool bTransparencyRoughessMapOn = true;
+	bool bRoughnessInverted = true;
+	bool bThinWalled = true;
+	bool bCutoutMapOn = true;
+
+	float DiffuseWeight = 0, Roughness = 0, Metalness = 0, Ior = 0;
+	float Transparency = 0;
+	float Reflectivity = 0;
+	float EmittanceMultiplier = 0, EmittanceLuminance = 0, EmittanceTemperature = 0;
+	float DisplacementMapAmount = 0;
+
+	FLinearColor TransparencyColor;
+	BMM_Color_fl DiffuseColor, EmittanceColor, ReflectionColor;
+
+	DatasmithMaxTexmapParser::FMapParameter DiffuseWeightMap;
+	DatasmithMaxTexmapParser::FMapParameter DiffuseColorMap;
+	DatasmithMaxTexmapParser::FMapParameter TransparencyMap;
+	DatasmithMaxTexmapParser::FMapParameter TransparencyColorMap;
+	DatasmithMaxTexmapParser::FMapParameter MetalnessMap;
+	DatasmithMaxTexmapParser::FMapParameter RoughnessMap;
+	DatasmithMaxTexmapParser::FMapParameter BumpMap;
+	DatasmithMaxTexmapParser::FMapParameter ReflectivityMap;
+	DatasmithMaxTexmapParser::FMapParameter ReflectionColorMap;
+	DatasmithMaxTexmapParser::FMapParameter EmittanceMap;
+	DatasmithMaxTexmapParser::FMapParameter EmittanceColorMap;
+	Texmap* CutoutMap = NULL;
+
+	int MaterialMode = 0; //0 means simple, just metalness, 1 means specular + metalness
+
+
+	void Parse(Mtl& Material)
+	{
+		int NumParamBlocks = Material.NumParamBlocks();
+
+
+		for (int j = 0; j < NumParamBlocks; j++)
+		{
+			IParamBlock2 *ParamBlock2 = Material.GetParamBlockByID(j);
+			// The the descriptor to 'decode'
+			ParamBlockDesc2 *ParamBlockDesc = ParamBlock2->GetDesc();
+			// Loop through all the defined parameters therein
+			for (int i = 0; i < ParamBlockDesc->count; i++)
+			{
+				const ParamDef& ParamDefinition = ParamBlockDesc->paramdefs[i];
+
+
+				if (FCString::Stricmp(ParamDefinition.int_name, TEXT("material_mode")) == 0)
+				{
+					MaterialMode = ParamBlock2->GetInt(ParamDefinition.ID, GetCOREInterface()->GetTime());
+				}
+
+				// booleans
+				else if (FCString::Stricmp(ParamDefinition.int_name, TEXT("base_weight_map_on")) == 0)
+				{
+					if (ParamBlock2->GetInt(ParamDefinition.ID, GetCOREInterface()->GetTime()) == 0)
+					{
+						DiffuseWeightMap.bEnabled = false;
+					}
+				}
+				else if (FCString::Stricmp(ParamDefinition.int_name, TEXT("base_color_map_on")) == 0)
+				{
+					if (ParamBlock2->GetInt(ParamDefinition.ID, GetCOREInterface()->GetTime()) == 0)
+					{
+						DiffuseColorMap.bEnabled = false;
+					}
+				}
+				else if (FCString::Stricmp(ParamDefinition.int_name, TEXT("reflectivity_map_on")) == 0)
+				{
+					if (ParamBlock2->GetInt(ParamDefinition.ID, GetCOREInterface()->GetTime()) == 0)
+					{
+						ReflectivityMap.bEnabled = false;
+					}
+				}
+				else if (FCString::Stricmp(ParamDefinition.int_name, TEXT("refl_color_map_on")) == 0)
+				{
+					if (ParamBlock2->GetInt(ParamDefinition.ID, GetCOREInterface()->GetTime()) == 0)
+					{
+						ReflectionColorMap.bEnabled = false;
+					}
+				}
+				else if (FCString::Stricmp(ParamDefinition.int_name, TEXT("roughness_map_on")) == 0)
+				{
+					if (ParamBlock2->GetInt(ParamDefinition.ID, GetCOREInterface()->GetTime()) == 0)
+					{
+						RoughnessMap.bEnabled = false;
+					}
+				}
+				else if (FCString::Stricmp(ParamDefinition.int_name, TEXT("metalness_map_on")) == 0)
+				{
+					if (ParamBlock2->GetInt(ParamDefinition.ID, GetCOREInterface()->GetTime()) == 0)
+					{
+						MetalnessMap.bEnabled = false;
+					}
+				}
+				else if (FCString::Stricmp(ParamDefinition.int_name, TEXT("transparency_map_on")) == 0)
+				{
+					if (ParamBlock2->GetInt(ParamDefinition.ID, GetCOREInterface()->GetTime()) == 0)
+					{
+						TransparencyMap.bEnabled = false;
+					}
+				}
+				else if (FCString::Stricmp(ParamDefinition.int_name, TEXT("trans_color_map_on")) == 0)
+				{
+					if (ParamBlock2->GetInt(ParamDefinition.ID, GetCOREInterface()->GetTime()) == 0)
+					{
+						TransparencyColorMap.bEnabled = false;
+					}
+				}
+				else if (FCString::Stricmp(ParamDefinition.int_name, TEXT("trans_rough_map_on")) == 0)
+				{
+					if (ParamBlock2->GetInt(ParamDefinition.ID, GetCOREInterface()->GetTime()) == 0)
+					{
+						bTransparencyRoughessMapOn = false;
+					}
+				}
+				else if (FCString::Stricmp(ParamDefinition.int_name, TEXT("emission_map_on")) == 0)
+				{
+					if (ParamBlock2->GetInt(ParamDefinition.ID, GetCOREInterface()->GetTime()) == 0)
+					{
+						EmittanceMap.bEnabled = false;
+					}
+				}
+				else if (FCString::Stricmp(ParamDefinition.int_name, TEXT("emit_color_map_on")) == 0)
+				{
+					if (ParamBlock2->GetInt(ParamDefinition.ID, GetCOREInterface()->GetTime()) == 0)
+					{
+						EmittanceColorMap.bEnabled = false;
+					}
+				}
+				else if (FCString::Stricmp(ParamDefinition.int_name, TEXT("roughness_inv")) == 0) // Set when 'Glossiness' selected in UI instead of Roughness
+				{
+					if (ParamBlock2->GetInt(ParamDefinition.ID, GetCOREInterface()->GetTime()) == 0)
+					{
+						bRoughnessInverted = false;
+					}
+				}
+				else if (FCString::Stricmp(ParamDefinition.int_name, TEXT("thin_walled")) == 0)
+				{
+					if (ParamBlock2->GetInt(ParamDefinition.ID, GetCOREInterface()->GetTime()) == 0)
+					{
+						bThinWalled = false;
+					}
+				}
+				else if (FCString::Stricmp(ParamDefinition.int_name, TEXT("bump_map_on")) == 0)
+				{
+					if (ParamBlock2->GetInt(ParamDefinition.ID, GetCOREInterface()->GetTime()) == 0)
+					{
+						BumpMap.bEnabled = false;
+					}
+				}
+				else if (FCString::Stricmp(ParamDefinition.int_name, TEXT("cutout_map_on")) == 0)
+				{
+					if (ParamBlock2->GetInt(ParamDefinition.ID, GetCOREInterface()->GetTime()) == 0)
+					{
+						bCutoutMapOn = false;
+					}
+				}
+
+				// float values
+				else if (FCString::Stricmp(ParamDefinition.int_name, TEXT("base_weight")) == 0)
+				{
+					DiffuseWeight = ParamBlock2->GetFloat(ParamDefinition.ID, GetCOREInterface()->GetTime());
+				}
+				else if (FCString::Stricmp(ParamDefinition.int_name, TEXT("roughness")) == 0)
+				{
+					Roughness = ParamBlock2->GetFloat(ParamDefinition.ID, GetCOREInterface()->GetTime());
+				}
+				else if (FCString::Stricmp(ParamDefinition.int_name, TEXT("metalness")) == 0)
+				{
+					Metalness = ParamBlock2->GetFloat(ParamDefinition.ID, GetCOREInterface()->GetTime());
+				}
+				else if (FCString::Stricmp(ParamDefinition.int_name, TEXT("trans_ior")) == 0)
+				{
+					Ior = ParamBlock2->GetFloat(ParamDefinition.ID, GetCOREInterface()->GetTime());
+				}
+				else if (FCString::Stricmp(ParamDefinition.int_name, TEXT("transparency")) == 0)
+				{
+					Transparency = ParamBlock2->GetFloat(ParamDefinition.ID, GetCOREInterface()->GetTime());
+				}
+				else if (FCString::Stricmp(ParamDefinition.int_name, TEXT("reflectivity")) == 0)
+				{
+					Reflectivity = ParamBlock2->GetFloat(ParamDefinition.ID, GetCOREInterface()->GetTime());
+				}
+				else if (FCString::Stricmp(ParamDefinition.int_name, TEXT("emission")) == 0)
+				{
+					EmittanceMultiplier = ParamBlock2->GetFloat(ParamDefinition.ID, GetCOREInterface()->GetTime());
+				}
+				else if (FCString::Stricmp(ParamDefinition.int_name, TEXT("emit_luminance")) == 0)
+				{
+					EmittanceLuminance = ParamBlock2->GetFloat(ParamDefinition.ID, GetCOREInterface()->GetTime());
+				}
+				else if (FCString::Stricmp(ParamDefinition.int_name, TEXT("emit_kelvin")) == 0)
+				{
+					EmittanceTemperature = ParamBlock2->GetFloat(ParamDefinition.ID, GetCOREInterface()->GetTime());
+				}
+				else if (FCString::Stricmp(ParamDefinition.int_name, TEXT("bump_map_amt")) == 0)
+				{
+					BumpMap.Weight = ParamBlock2->GetFloat(ParamDefinition.ID, GetCOREInterface()->GetTime());
+				}
+				else if (FCString::Stricmp(ParamDefinition.int_name, TEXT("displacement_map_amt")) == 0)
+				{
+					DisplacementMapAmount = ParamBlock2->GetFloat(ParamDefinition.ID, GetCOREInterface()->GetTime());
+				}
+
+				// plain colors
+				else if (FCString::Stricmp(ParamDefinition.int_name, TEXT("base_color")) == 0)
+				{
+					DiffuseColor = (BMM_Color_fl)ParamBlock2->GetColor(ParamDefinition.ID, GetCOREInterface()->GetTime());
+				}
+				else if (FCString::Stricmp(ParamDefinition.int_name, TEXT("trans_color")) == 0)
+				{
+					TransparencyColor = FDatasmithMaxMatHelper::MaxLinearColorToFLinearColor((BMM_Color_fl)ParamBlock2->GetColor(ParamDefinition.ID, GetCOREInterface()->GetTime()));
+				}
+				else if (FCString::Stricmp(ParamDefinition.int_name, TEXT("emit_color")) == 0)
+				{
+					EmittanceColor = (BMM_Color_fl)ParamBlock2->GetColor(ParamDefinition.ID, GetCOREInterface()->GetTime());
+				}
+				else if (FCString::Stricmp(ParamDefinition.int_name, TEXT("refl_color")) == 0)
+				{
+					ReflectionColor = (BMM_Color_fl)ParamBlock2->GetColor(ParamDefinition.ID, GetCOREInterface()->GetTime());
+				}
+
+				// textures
+				else if (FCString::Stricmp(ParamDefinition.int_name, TEXT("Base_Weight_Map")) == 0)
+				{
+					DiffuseWeightMap.Map = ParamBlock2->GetTexmap(ParamDefinition.ID, GetCOREInterface()->GetTime());
+				}
+				else if (FCString::Stricmp(ParamDefinition.int_name, TEXT("Base_Color_Map")) == 0)
+				{
+					DiffuseColorMap.Map = ParamBlock2->GetTexmap(ParamDefinition.ID, GetCOREInterface()->GetTime());
+				}
+				else if (FCString::Stricmp(ParamDefinition.int_name, TEXT("Reflectivity_Map")) == 0)
+				{
+					ReflectivityMap.Map = ParamBlock2->GetTexmap(ParamDefinition.ID, GetCOREInterface()->GetTime());
+				}
+				else if (FCString::Stricmp(ParamDefinition.int_name, TEXT("Roughness_Map")) == 0)
+				{
+					RoughnessMap.Map = ParamBlock2->GetTexmap(ParamDefinition.ID, GetCOREInterface()->GetTime());
+				}
+				else if (FCString::Stricmp(ParamDefinition.int_name, TEXT("Metalness_Map")) == 0)
+				{
+					MetalnessMap.Map = ParamBlock2->GetTexmap(ParamDefinition.ID, GetCOREInterface()->GetTime());
+				}
+				else if (FCString::Stricmp(ParamDefinition.int_name, TEXT("Refl_Color_Map")) == 0)
+				{
+					ReflectionColorMap.Map = ParamBlock2->GetTexmap(ParamDefinition.ID, GetCOREInterface()->GetTime());
+				}
+				else if (FCString::Stricmp(ParamDefinition.int_name, TEXT("Transparency_Map")) == 0)
+				{
+					TransparencyMap.Map = ParamBlock2->GetTexmap(ParamDefinition.ID, GetCOREInterface()->GetTime());
+				}
+				else if (FCString::Stricmp(ParamDefinition.int_name, TEXT("trans_color_map")) == 0)
+				{
+					TransparencyColorMap.Map = ParamBlock2->GetTexmap(ParamDefinition.ID, GetCOREInterface()->GetTime());
+				}
+				else if (FCString::Stricmp(ParamDefinition.int_name, TEXT("Emission_Map")) == 0)
+				{
+					EmittanceMap.Map = ParamBlock2->GetTexmap(ParamDefinition.ID, GetCOREInterface()->GetTime());
+				}
+				else if (FCString::Stricmp(ParamDefinition.int_name, TEXT("emit_color_map")) == 0)
+				{
+					EmittanceColorMap.Map = ParamBlock2->GetTexmap(ParamDefinition.ID, GetCOREInterface()->GetTime());
+				}
+				else if (FCString::Stricmp(ParamDefinition.int_name, TEXT("bump_map")) == 0)
+				{
+					BumpMap.Map = ParamBlock2->GetTexmap(ParamDefinition.ID, GetCOREInterface()->GetTime());
+				}
+				else if (FCString::Stricmp(ParamDefinition.int_name, TEXT("cutout_map")) == 0)
+				{
+					CutoutMap = ParamBlock2->GetTexmap(ParamDefinition.ID, GetCOREInterface()->GetTime());
+				}
+			}
+
+			ParamBlock2->ReleaseDesc();
+		}
+	}
+};
+
+struct FMaxPhysicalMaterialClearCoat
+{
+	bool bCoatRoughnessInverted = true;
+
+	float CoatAmount = 0.f;
+	float CoatRoughness = 0.f;
+	float CoatIOR = 1.5f;
+
+	DatasmithMaxTexmapParser::FMapParameter CoatWeightMap;
+	DatasmithMaxTexmapParser::FMapParameter CoatColorMap;
+	DatasmithMaxTexmapParser::FMapParameter CoatRoughnessMap;
+
+	BMM_Color_fl CoatColor;
+
+	bool HasClearCoat()
+	{
+		// There is no coating effect
+		if (CoatAmount <= 0.0f || ((CoatColorMap.Map == false || CoatColorMap.bEnabled == false) && CoatColor.r == 0.0f && CoatColor.g == 0.0f && CoatColor.b == 0))
+		{
+			return false;
+		}
+		return true;
+	}
+
+	void Parse(Mtl& Material){
+		const int NumParamBlocks = Material.NumParamBlocks();
+
+		for (int j = 0; j < NumParamBlocks; j++)
+		{
+			IParamBlock2* ParamBlock2 = Material.GetParamBlockByID(j);
+			ParamBlockDesc2* ParamBlockDesc = ParamBlock2->GetDesc();
+
+			// Loop through all the defined parameters therein
+			for (int i = 0; i < ParamBlockDesc->count; i++)
+			{
+				const ParamDef& ParamDefinition = ParamBlockDesc->paramdefs[i];
+
+				if (FCString::Stricmp(ParamDefinition.int_name, TEXT("coat_map_on")) == 0)
+				{
+					if (ParamBlock2->GetInt(ParamDefinition.ID, GetCOREInterface()->GetTime()) == 0)
+					{
+						CoatWeightMap.bEnabled = false;
+					}
+				}
+				else if (FCString::Stricmp(ParamDefinition.int_name, TEXT("coat_color_map_on")) == 0)
+				{
+					if (ParamBlock2->GetInt(ParamDefinition.ID, GetCOREInterface()->GetTime()) == 0)
+					{
+						CoatColorMap.bEnabled = false;
+					}
+				}
+				else if (FCString::Stricmp(ParamDefinition.int_name, TEXT("coat_rough_map_on")) == 0)
+				{
+					if (ParamBlock2->GetInt(ParamDefinition.ID, GetCOREInterface()->GetTime()) == 0)
+					{
+						CoatRoughnessMap.bEnabled = false;
+					}
+				}
+				else if (FCString::Stricmp(ParamDefinition.int_name, TEXT("coat_roughness_inv")) == 0)
+				{
+					if (ParamBlock2->GetInt(ParamDefinition.ID, GetCOREInterface()->GetTime()) == 0)
+					{
+						bCoatRoughnessInverted = false;
+					}
+				}
+				else if (FCString::Stricmp(ParamDefinition.int_name, TEXT("Coating_Weight_Map")) == 0)
+				{
+					CoatWeightMap.Map = ParamBlock2->GetTexmap(ParamDefinition.ID, GetCOREInterface()->GetTime());
+				}
+				else if (FCString::Stricmp(ParamDefinition.int_name, TEXT("Coating_Color_Map")) == 0)
+				{
+					CoatColorMap.Map = ParamBlock2->GetTexmap(ParamDefinition.ID, GetCOREInterface()->GetTime());
+				}
+				else if (FCString::Stricmp(ParamDefinition.int_name, TEXT("Coating_Roughness_Map")) == 0)
+				{
+					CoatRoughnessMap.Map = ParamBlock2->GetTexmap(ParamDefinition.ID, GetCOREInterface()->GetTime());
+				}
+				else if (FCString::Stricmp(ParamDefinition.int_name, TEXT("coating")) == 0)
+				{
+					CoatAmount = ParamBlock2->GetFloat(ParamDefinition.ID, GetCOREInterface()->GetTime());
+				}
+				else if (FCString::Stricmp(ParamDefinition.int_name, TEXT("coat_roughness")) == 0)
+				{
+					CoatRoughness = ParamBlock2->GetFloat(ParamDefinition.ID, GetCOREInterface()->GetTime());
+				}
+				else if (FCString::Stricmp(ParamDefinition.int_name, TEXT("coat_ior")) == 0)
+				{
+					CoatIOR = ParamBlock2->GetFloat(ParamDefinition.ID, GetCOREInterface()->GetTime());
+				}
+				else if (FCString::Stricmp(ParamDefinition.int_name, TEXT("coat_color")) == 0)
+				{
+					CoatColor = (BMM_Color_fl)ParamBlock2->GetColor(ParamDefinition.ID, GetCOREInterface()->GetTime());
+				}
+			}
+
+			ParamBlock2->ReleaseDesc();
+		}
+		
+	}
+};
+
+void FDatasmithMaxPhysicalMaterialToUEPbr::Convert( TSharedRef< IDatasmithScene > DatasmithScene, TSharedPtr< IDatasmithBaseMaterialElement >& MaterialElement, Mtl* Material, const TCHAR* AssetsPath )
+{
+	if ( !Material )
+	{
+		return;
+	}
+
+	TSharedRef< IDatasmithUEPbrMaterialElement > PbrMaterialElement = FDatasmithSceneFactory::CreateUEPbrMaterial( Material->GetName().data() );
+	FScopedConvertState ScopedConvertState(ConvertState);
+	ConvertState.DatasmithScene = DatasmithScene;
+	ConvertState.MaterialElement = PbrMaterialElement;
+	ConvertState.AssetsPath = AssetsPath;
+
+	FMaxPhysicalMaterial PhysicalMaterialProperties;
+	PhysicalMaterialProperties.Parse(*Material);
+
+	IDatasmithMaterialExpression* DiffuseColorExpression = ApplyWeightExpression(
+		TextureOrColor(TEXT("Diffuse Color"), PhysicalMaterialProperties.DiffuseColorMap, FDatasmithMaxMatHelper::MaxLinearColorToFLinearColor(PhysicalMaterialProperties.DiffuseColor)),
+		WeightTextureOrScalar(PhysicalMaterialProperties.DiffuseWeightMap, PhysicalMaterialProperties.DiffuseWeight));
+
+	IDatasmithMaterialExpression* TransparencyExpression = nullptr;
+	if (PhysicalMaterialProperties.TransparencyColorMap.Map || PhysicalMaterialProperties.TransparencyMap.Map)
+	{
+		TransparencyExpression = ApplyWeightExpression(
+					   TextureOrColor(TEXT("Opacity"), PhysicalMaterialProperties.TransparencyColorMap, PhysicalMaterialProperties.TransparencyColor),
+					   WeightTextureOrScalar(PhysicalMaterialProperties.TransparencyMap, PhysicalMaterialProperties.Transparency));
+	}
+	else
+	{
+		FLinearColor Transparency = PhysicalMaterialProperties.TransparencyColor * PhysicalMaterialProperties.Transparency;
+
+		if (!Transparency.IsAlmostBlack())
+		{
+			TransparencyExpression = &Color(Transparency);
+		}
+	}
+
+	Connect(PbrMaterialElement->GetEmissiveColor(), 
+		ApplyWeightExpression(
+			TextureOrColor(TEXT("Emissive Color"), PhysicalMaterialProperties.EmittanceColorMap, FDatasmithMaxMatHelper::MaxLinearColorToFLinearColor(PhysicalMaterialProperties.EmittanceColor)),
+			WeightTextureOrScalar(PhysicalMaterialProperties.EmittanceMap, PhysicalMaterialProperties.EmittanceMultiplier))
+	);
+
+
+	// todo: set opacity as mask and set blend mode
+	// PbrMaterialElement->SetBlendMode(/*EBlendMode::BLEND_Masked*/1)
+	// ExportPhysicalMaterialProperty(DatasmithScene, CutoutMap, bCutoutMapOn, NULL, NULL, BMM_Color_fl(0.0, 0.0, 0.0, 0.0), 1.0, MaterialShader->GetMaskComp(), DATASMITH_CLIPTEXNAME, DATASMITH_CLIPTEXNAME, false, true);
+
+	Connect(PbrMaterialElement->GetMetallic(), TextureOrScalar(TEXT("Metallic"), PhysicalMaterialProperties.MetalnessMap, PhysicalMaterialProperties.Metalness));
+
+	if (PhysicalMaterialProperties.MaterialMode == 1)
+	{
+		IDatasmithMaterialExpression* ReflectivityExpression = ApplyWeightExpression(
+			TextureOrColor(TEXT("Specular"), PhysicalMaterialProperties.ReflectionColorMap, FDatasmithMaxMatHelper::MaxLinearColorToFLinearColor(PhysicalMaterialProperties.ReflectionColor)),
+			WeightTextureOrScalar(PhysicalMaterialProperties.ReflectivityMap, PhysicalMaterialProperties.Reflectivity));
+
+		Connect(PbrMaterialElement->GetSpecular(), COMPOSE_OR_DEFAULT1(nullptr, Desaturate, ReflectivityExpression));
+	}
+
+	if (IDatasmithMaterialExpression* RoughnessTextureExpression = ConvertTexmap(PhysicalMaterialProperties.RoughnessMap))
+	{
+		Connect(PbrMaterialElement->GetRoughness(), PhysicalMaterialProperties.bRoughnessInverted ? COMPOSE_OR_DEFAULT1(nullptr, OneMinus, RoughnessTextureExpression) : RoughnessTextureExpression);
+	}
+	else
+	{
+		float Roughness = 0.75f;
+		if ((PhysicalMaterialProperties.MetalnessMap.Map && PhysicalMaterialProperties.MetalnessMap.bEnabled) 
+			|| PhysicalMaterialProperties.Metalness > 0 
+			|| (PhysicalMaterialProperties.MaterialMode == 1 && ((PhysicalMaterialProperties.ReflectionColorMap.Map && PhysicalMaterialProperties.ReflectionColorMap.bEnabled) || PhysicalMaterialProperties.Reflectivity > 0)))
+		{
+			Roughness = PhysicalMaterialProperties.bRoughnessInverted  ? 1.0f - PhysicalMaterialProperties.Roughness : PhysicalMaterialProperties.Roughness;
+		}
+		Connect(PbrMaterialElement->GetRoughness(), Scalar(Roughness));
+	}
+
+	if (PhysicalMaterialProperties.BumpMap.Map && PhysicalMaterialProperties.BumpMap.bEnabled== true)
+	{
+		ConvertState.DefaultTextureMode = EDatasmithTextureMode::Bump; // Will change to normal if we pass through a normal map texmap
+		ConvertState.bCanBake = false; // Current baking fails to produce proper normal maps
+
+		IDatasmithMaterialExpression* BumpExpression = FDatasmithMaxTexmapToUEPbrUtils::MapOrValue(this, PhysicalMaterialProperties.BumpMap, TEXT("Bump Map"), TOptional<FLinearColor>(), TOptional<float>());
+
+		Connect(PbrMaterialElement->GetNormal(), BumpExpression);
+		ConvertState.bCanBake = true;
+	}
+
+	IDatasmithMaterialExpression* BaseColorExpression = DiffuseColorExpression;
+	// Build Material based on estimated material type
+	if (TransparencyExpression)
+	{
+		IDatasmithMaterialExpression& TransparencyExpressionRef = *TransparencyExpression;
+
+		// ThinTranslucent, including when Ior is nearly 1
+		if (PhysicalMaterialProperties.bThinWalled || FMath::IsNearlyEqual(PhysicalMaterialProperties.Ior, 1.0f))
+		{
+			PbrMaterialElement->SetShadingModel(EDatasmithShadingModel::ThinTranslucent);
+			IDatasmithMaterialExpressionGeneric* ThinTranslucencyMaterialOutput = PbrMaterialElement->AddMaterialExpression<IDatasmithMaterialExpressionGeneric>();
+			ThinTranslucencyMaterialOutput->SetExpressionName(TEXT("ThinTranslucentMaterialOutput"));
+
+			TransparencyExpression->ConnectExpression(*ThinTranslucencyMaterialOutput->GetInput(0));
+
+			Connect(PbrMaterialElement->GetOpacity(), OneMinus(Desaturate(TransparencyExpressionRef)));
+		}
+		else
+		{
+			if (PhysicalMaterialProperties.MaterialMode == 1) // Advanced mode, setup refractive glass
+			{
+				// Ported FDatasmithMaterialExpressions::AddGlassNode to compute base color and opacity
+				float Ior = PhysicalMaterialProperties.Ior;
+
+				IDatasmithMaterialExpression& Multiply09 = Multiply(Scalar(0.9), TransparencyExpressionRef);
+
+				Connect(PbrMaterialElement->GetOpacity(), OneMinus(Desaturate(CalcIORComplex(Ior, 0, Multiply(Scalar(0.5), TransparencyExpressionRef), Multiply09))));
+				Connect(PbrMaterialElement->GetRefraction(), Lerp(Scalar(1.0f), Scalar(PhysicalMaterialProperties.Ior), Fresnel()));
+
+				// todo: may add Diffuse color into the mix to accomodate non-clear opacity
+				BaseColorExpression = &Multiply(Scalar(0.6), Power(Multiply09, Scalar(5.0))); 
+			}
+			else
+			{
+				Connect(PbrMaterialElement->GetOpacity(), OneMinus(Desaturate(TransparencyExpressionRef)));
+			}
+		}
+	}
+	Connect(PbrMaterialElement->GetBaseColor(), BaseColorExpression);
+
+	FMaxPhysicalMaterialClearCoat ClearCoatProperties;
+	ClearCoatProperties.Parse(*Material);
+
+	if (ClearCoatProperties.HasClearCoat())
+	{
+		IDatasmithMaterialExpression* ClearCoatExpression = ApplyWeightExpression(
+			TextureOrColor(TEXT("Clear Coat"), ClearCoatProperties.CoatColorMap, FDatasmithMaxMatHelper::MaxLinearColorToFLinearColor(ClearCoatProperties.CoatColor)),
+			WeightTextureOrScalar(ClearCoatProperties.CoatWeightMap, ClearCoatProperties.CoatAmount) 
+			);
+
+		if (ClearCoatExpression)
+		{
+			// at 0 degrees we use the minumum and the actual input at 90 degrees
+			IDatasmithMaterialExpression& ToBeConnected0 = Multiply(*ClearCoatExpression, Scalar(0.1));
+
+			Connect(PbrMaterialElement->GetClearCoat(), CalcIORComplex(ClearCoatProperties.CoatIOR, 0, *ClearCoatExpression, ToBeConnected0));
+		}
+
+		Connect(PbrMaterialElement->GetClearCoatRoughness(), TextureOrScalar(TEXT("Clear Coat Roughness"), ClearCoatProperties.CoatRoughnessMap, ClearCoatProperties.CoatRoughness));
+
+		PbrMaterialElement->SetShadingModel(EDatasmithShadingModel::ClearCoat);
+	}
+
+
+	MaterialElement = PbrMaterialElement;
+}
diff --git a/Engine/Source/Programs/Enterprise/Datasmith/DatasmithMaxExporter/Private/MaxMaterialsToUEPbr/DatasmithMaxPhysicalMaterialToUEPbr.h b/Engine/Source/Programs/Enterprise/Datasmith/DatasmithMaxExporter/Private/MaxMaterialsToUEPbr/DatasmithMaxPhysicalMaterialToUEPbr.h
new file mode 100644
--- /dev/null
+++ b/Engine/Source/Programs/Enterprise/Datasmith/DatasmithMaxExporter/Private/MaxMaterialsToUEPbr/DatasmithMaxPhysicalMaterialToUEPbr.h
@@ -0,0 +1,64 @@
+// Copyright Epic Games, Inc. All Rights Reserved.
+
+#pragma once
+
+#include "DatasmithMaxMaterialsToUEPbr.h"
+
+class IDatasmithMaterialExpression;
+class IDatasmithMaterialExpressionScalar;
+class IDatasmithMaterialExpressionColor;
+class IDatasmithMaterialExpressionGeneric;
+class IDatasmithExpressionInput;
+
+class FDatasmithMaxMaterialsToUEPbrExpressions: public FDatasmithMaxMaterialsToUEPbr
+{
+public:
+	TSharedPtr<IDatasmithUEPbrMaterialElement> GetMaterialElement();
+
+	IDatasmithMaterialExpressionScalar& Scalar(float Value);
+
+	IDatasmithMaterialExpressionColor& Color(const FLinearColor& Value);
+
+	IDatasmithMaterialExpression* WeightTextureOrScalar(const DatasmithMaxTexmapParser::FMapParameter& TextureWeight, float Weight);
+
+	IDatasmithMaterialExpressionGeneric& Add(IDatasmithMaterialExpression& A, IDatasmithMaterialExpression& B);
+
+	IDatasmithMaterialExpressionGeneric& Subtract(IDatasmithMaterialExpression& A, IDatasmithMaterialExpression& B);
+
+	IDatasmithMaterialExpressionGeneric& Multiply(IDatasmithMaterialExpression& A, IDatasmithMaterialExpression& B);
+
+	IDatasmithMaterialExpressionGeneric& Divide(IDatasmithMaterialExpression& A, IDatasmithMaterialExpression& B);
+
+	IDatasmithMaterialExpressionGeneric& Desaturate(IDatasmithMaterialExpression& A);
+
+	IDatasmithMaterialExpressionGeneric& Power(IDatasmithMaterialExpression& A, IDatasmithMaterialExpression& B);
+
+	IDatasmithMaterialExpressionGeneric& Lerp(IDatasmithMaterialExpression& A, IDatasmithMaterialExpression& B, IDatasmithMaterialExpression& Alpha);
+
+	IDatasmithMaterialExpressionGeneric& Fresnel(IDatasmithMaterialExpression* Exponent=nullptr, IDatasmithMaterialExpression* BaseReflectFraction=nullptr); // Any input can be null
+
+
+	IDatasmithMaterialExpression* ApplyWeightExpression(IDatasmithMaterialExpression* ValueExpression, IDatasmithMaterialExpression* WeightExpression);
+
+	IDatasmithMaterialExpression& CalcIORComplex(double IORn, double IORk, IDatasmithMaterialExpression& ToBeConnected90, IDatasmithMaterialExpression& ToBeConnected0);
+
+
+	void Connect(IDatasmithExpressionInput& Input, IDatasmithMaterialExpression& ValueExpression);
+	bool Connect(IDatasmithExpressionInput& Input, IDatasmithMaterialExpression* ValueExpression); // Connect if not null
+
+	IDatasmithMaterialExpression* TextureOrColor(const TCHAR* Name,
+	                                             const DatasmithMaxTexmapParser::FMapParameter& Map, FLinearColor Color);
+
+	IDatasmithMaterialExpression* TextureOrScalar(const TCHAR* Name,
+	                                              const DatasmithMaxTexmapParser::FMapParameter& Map, float Value);
+
+	IDatasmithMaterialExpressionGeneric& OneMinus(IDatasmithMaterialExpression& Expression);
+};
+
+class FDatasmithMaxPhysicalMaterialToUEPbr : public FDatasmithMaxMaterialsToUEPbrExpressions
+{
+public:
+	virtual bool IsSupported( Mtl* Material ) override;
+	virtual void Convert( TSharedRef< IDatasmithScene > DatasmithScene, TSharedPtr< IDatasmithBaseMaterialElement >& MaterialElement, Mtl* Material, const TCHAR* AssetsPath ) override;
+};
+
diff --git a/Engine/Source/Programs/Enterprise/Datasmith/DatasmithMaxExporter/Private/MaxMaterialsToUEPbr/DatasmithMaxScanlineMaterialsToUEPbr.cpp b/Engine/Source/Programs/Enterprise/Datasmith/DatasmithMaxExporter/Private/MaxMaterialsToUEPbr/DatasmithMaxScanlineMaterialsToUEPbr.cpp
--- a/Engine/Source/Programs/Enterprise/Datasmith/DatasmithMaxExporter/Private/MaxMaterialsToUEPbr/DatasmithMaxScanlineMaterialsToUEPbr.cpp
+++ b/Engine/Source/Programs/Enterprise/Datasmith/DatasmithMaxExporter/Private/MaxMaterialsToUEPbr/DatasmithMaxScanlineMaterialsToUEPbr.cpp
@@ -76,7 +76,7 @@ namespace DatasmithMaxScanlineMaterialsToUEPbrImpl
 		ScanlineMaterialProperties.SpecularLevel = Material.GetShinStr();
 		ScanlineMaterialProperties.Glossiness = Material.GetShininess();
 		ScanlineMaterialProperties.bUseSelfIllumColor = ( Material.GetSelfIllumColorOn() != 0 );
-		ScanlineMaterialProperties.SelfIllumColor = FDatasmithMaxMatHelper::MaxColorToFLinearColor( (BMM_Color_fl)Material.GetSelfIllumColor() );
+		ScanlineMaterialProperties.SelfIllumColor = FDatasmithMaxMatHelper::MaxColorToFLinearColor( (BMM_Color_fl)Material.GetSelfIllumColor() ); // todo: this seems wrong(converting color to gamme space)
 
 		const TimeValue CurrentTime = GetCOREInterface()->GetTime();
 
diff --git a/Engine/Source/Programs/Enterprise/Datasmith/DatasmithMaxExporter/Private/MaxMaterialsToUEPbr/DatasmithMaxTexmapToUEPbr.cpp b/Engine/Source/Programs/Enterprise/Datasmith/DatasmithMaxExporter/Private/MaxMaterialsToUEPbr/DatasmithMaxTexmapToUEPbr.cpp
--- a/Engine/Source/Programs/Enterprise/Datasmith/DatasmithMaxExporter/Private/MaxMaterialsToUEPbr/DatasmithMaxTexmapToUEPbr.cpp
+++ b/Engine/Source/Programs/Enterprise/Datasmith/DatasmithMaxExporter/Private/MaxMaterialsToUEPbr/DatasmithMaxTexmapToUEPbr.cpp
@@ -1736,6 +1736,10 @@ bool FDatasmithMaxPassthroughToUEPbr::IsSupported( const FDatasmithMaxMaterialsT
 		{
 			return true;
 		}
+		else if ( InTexmap->ClassID() == COLORMAPCLASS )
+		{
+			return true;
+		}
 	}
 
 	return false;
diff --git a/Engine/Source/Programs/Enterprise/Datasmith/DatasmithMaxExporter/Public/DatasmithMaxClassIDs.h b/Engine/Source/Programs/Enterprise/Datasmith/DatasmithMaxExporter/Public/DatasmithMaxClassIDs.h
--- a/Engine/Source/Programs/Enterprise/Datasmith/DatasmithMaxExporter/Public/DatasmithMaxClassIDs.h
+++ b/Engine/Source/Programs/Enterprise/Datasmith/DatasmithMaxExporter/Public/DatasmithMaxClassIDs.h
@@ -36,6 +36,8 @@
 #define  MRPHYSICALSKYBCLASS		Class_ID(0x223349b4, 0x562a7d8b)
 #define  THIRDPARTYMULTITEXCLASS	Class_ID(0x55c6d763, 0xc637bb2e)
 #define  TILESMAPCLASS				Class_ID(0x64035fb9, 0x69664cdc)
+#define  COLORMAPCLASS				Class_ID(0x139f22c6, 0x13f6a914)
+
 
 #define  VRAYNORMALCLASS			Class_ID(0x71fa6e51, 0x72057c2f)
 #define  VRAYCOLORCLASS				Class_ID(0x58f82b74, 0x73b75d7f)
diff --git a/Engine/Source/Programs/Enterprise/Datasmith/DatasmithMaxExporter/Public/DatasmithMaxWriter.h b/Engine/Source/Programs/Enterprise/Datasmith/DatasmithMaxExporter/Public/DatasmithMaxWriter.h
--- a/Engine/Source/Programs/Enterprise/Datasmith/DatasmithMaxExporter/Public/DatasmithMaxWriter.h
+++ b/Engine/Source/Programs/Enterprise/Datasmith/DatasmithMaxExporter/Public/DatasmithMaxWriter.h
@@ -90,7 +90,7 @@ public:
 
 	static bool HasNonBakeableSubmap(Texmap* InTexmap);
 
-	static FLinearColor MaxColorToFLinearColor(BMM_Color_fl Color, float Multiplier = 1.0f);
+	static FLinearColor MaxColorToFLinearColor(BMM_Color_fl Color, float Multiplier = 1.0f); // todo: rename, although result type is FLinearColor this function does pow(rgb, 1/gamma), i.e. converting to gamma-space
 	static FLinearColor MaxLinearColorToFLinearColor(BMM_Color_fl Color, float Multiplier = 1.0f);
 	static float GetBitmapGamma(BitmapTex* InBitmapTex);
 	static float GetBitmapGamma(BitmapInfo* InBitmapInfo);
