3ds Max DirectLink: Corona Light Material to UEPbr

- Emission and Opacity Mask
- Added "Unlit" ShadingModel to Datasmith core/importer

diff --git a/Engine/Plugins/Enterprise/DatasmithImporter/Source/DatasmithImporter/Private/DatasmithMaterialExpressions.cpp b/Engine/Plugins/Enterprise/DatasmithImporter/Source/DatasmithImporter/Private/DatasmithMaterialExpressions.cpp
--- a/Engine/Plugins/Enterprise/DatasmithImporter/Source/DatasmithImporter/Private/DatasmithMaterialExpressions.cpp
+++ b/Engine/Plugins/Enterprise/DatasmithImporter/Source/DatasmithImporter/Private/DatasmithMaterialExpressions.cpp
@@ -3289,12 +3289,13 @@ UMaterialInterface* FDatasmithMaterialExpressions::CreateUEPbrMaterial(UPackage*
 		case EDatasmithShadingModel::ThinTranslucent:
 			UnrealMaterial->SetShadingModel(MSM_ThinTranslucent);
 			break;
-
+		case EDatasmithShadingModel::Unlit:
+			UnrealMaterial->SetShadingModel(MSM_Unlit);
+			break;
 		default:
 			break;
 	}
 
-
 	if ( MaterialElement->GetOpacity().GetExpression() )
 	{
 		if ( MaterialElement->GetShadingModel() == EDatasmithShadingModel::ThinTranslucent )
diff --git a/Engine/Plugins/Enterprise/DatasmithImporter/Source/DatasmithImporter/Private/DatasmithMaterialImporter.cpp b/Engine/Plugins/Enterprise/DatasmithImporter/Source/DatasmithImporter/Private/DatasmithMaterialImporter.cpp
--- a/Engine/Plugins/Enterprise/DatasmithImporter/Source/DatasmithImporter/Private/DatasmithMaterialImporter.cpp
+++ b/Engine/Plugins/Enterprise/DatasmithImporter/Source/DatasmithImporter/Private/DatasmithMaterialImporter.cpp
@@ -121,7 +121,11 @@ namespace DatasmithMaterialImporterUtils
 
 		for ( int32 InputIndex = 0; InputIndex < MaterialExpression->GetInputCount(); ++InputIndex )
 		{
-			Hash = HashCombine( Hash, ComputeExpressionInputHash( MaterialExpression->GetInput( InputIndex ), VisitedExpressions ) );
+			IDatasmithExpressionInput* DatasmithExpressionInput = MaterialExpression->GetInput( InputIndex );
+			if(ensure(!DatasmithExpressionInput || (MaterialExpression != DatasmithExpressionInput->GetExpression()))) // Prevent infinite recursion when an expression has itself as input
+			{
+				Hash = HashCombine( Hash, ComputeExpressionInputHash( DatasmithExpressionInput, VisitedExpressions ) );
+			}
 		}
 
 		return Hash;
diff --git a/Engine/Source/Programs/Enterprise/Datasmith/DatasmithMaxExporter/Private/MaxMaterialsToUEPbr/DatasmithMaxCoronaMaterialsToUEPbr.cpp b/Engine/Source/Programs/Enterprise/Datasmith/DatasmithMaxExporter/Private/MaxMaterialsToUEPbr/DatasmithMaxCoronaMaterialsToUEPbr.cpp
--- a/Engine/Source/Programs/Enterprise/Datasmith/DatasmithMaxExporter/Private/MaxMaterialsToUEPbr/DatasmithMaxCoronaMaterialsToUEPbr.cpp
+++ b/Engine/Source/Programs/Enterprise/Datasmith/DatasmithMaxExporter/Private/MaxMaterialsToUEPbr/DatasmithMaxCoronaMaterialsToUEPbr.cpp
@@ -289,6 +289,69 @@ namespace DatasmithMaxCoronaMaterialsToUEPbrImpl
 
 		return CoronaBlendMaterialProperties;
 	}
+
+	struct FCoronaLightMaterial
+	{
+		DatasmithMaxTexmapParser::FMapParameter EmitTexture;
+		DatasmithMaxTexmapParser::FMapParameter ClipTexture;
+		BMM_Color_fl EmitColor;
+		float Multiplier = 1.0f;
+
+
+		void Parse(Mtl& Material)
+		{
+
+			int NumParamBlocks = Material.NumParamBlocks();
+
+			for (int j = 0; j < NumParamBlocks; j++)
+			{
+				IParamBlock2* ParamBlock2 = Material.GetParamBlockByID((short)j);
+				// The the descriptor to 'decode'
+				ParamBlockDesc2* ParamBlockDesc = ParamBlock2->GetDesc();
+				// Loop through all the defined parameters therein
+				for (int i = 0; i < ParamBlockDesc->count; i++)
+				{
+					const ParamDef& ParamDefinition = ParamBlockDesc->paramdefs[i];
+
+					if (FCString::Stricmp(ParamDefinition.int_name, TEXT("texmap")) == 0)
+					{
+						EmitTexture.Map = ParamBlock2->GetTexmap(ParamDefinition.ID, GetCOREInterface()->GetTime());
+					}
+					else if (FCString::Stricmp(ParamDefinition.int_name, TEXT("opacityTexmap")) == 0)
+					{
+						ClipTexture.Map = ParamBlock2->GetTexmap(ParamDefinition.ID, GetCOREInterface()->GetTime());
+					}
+					else if (FCString::Stricmp(ParamDefinition.int_name, TEXT("texmapOn")) == 0)
+					{
+						if (ParamBlock2->GetInt(ParamDefinition.ID, GetCOREInterface()->GetTime()) == 0)
+						{
+							EmitTexture.bEnabled = false;
+						}
+					}
+					else if (FCString::Stricmp(ParamDefinition.int_name, TEXT("opacityTexmapOn")) == 0)
+					{
+						if (ParamBlock2->GetInt(ParamDefinition.ID, GetCOREInterface()->GetTime()) == 0)
+						{
+							ClipTexture.bEnabled = false;
+						}
+					}
+					else if (FCString::Stricmp(ParamDefinition.int_name, TEXT("color")) == 0)
+					{
+						EmitColor = (BMM_Color_fl)ParamBlock2->GetColor(ParamDefinition.ID, GetCOREInterface()->GetTime());
+					}
+					else if (FCString::Stricmp(ParamDefinition.int_name, TEXT("twosidedEmission")) == 0)
+					{
+						// int twoSided = ParamBlock2->GetInt(ParamDefinition.ID, GetCOREInterface()->GetTime());
+					}
+					else if (FCString::Stricmp(ParamDefinition.int_name, TEXT("intensity")) == 0)
+					{
+						Multiplier = ParamBlock2->GetFloat(ParamDefinition.ID, GetCOREInterface()->GetTime());
+					}
+				}
+				ParamBlock2->ReleaseDesc();
+			}
+		}
+	};
 }
 
 FDatasmithMaxCoronaMaterialsToUEPbr::FDatasmithMaxCoronaMaterialsToUEPbr()
@@ -717,4 +780,43 @@ void FDatasmithMaxCoronaBlendMaterialToUEPbr::Convert( TSharedRef<IDatasmithScen
 	PbrMaterialElement->SetUseMaterialAttributes(true);
 	PreviousExpression->ConnectExpression(PbrMaterialElement->GetMaterialAttributes());
 	MaterialElement = PbrMaterialElement;
-}
\ No newline at end of file
+}
+
+bool FDatasmithMaxCoronaLightMaterialToUEPbr::IsSupported(Mtl* Material)
+{
+	return true;
+}
+
+void FDatasmithMaxCoronaLightMaterialToUEPbr::Convert(TSharedRef<IDatasmithScene> DatasmithScene,
+	TSharedPtr<IDatasmithBaseMaterialElement>& MaterialElement, Mtl* Material, const TCHAR* AssetsPath)
+{
+	if ( !Material )
+	{
+		return;
+	}
+
+	TSharedRef<IDatasmithUEPbrMaterialElement> PbrMaterialElement = FDatasmithSceneFactory::CreateUEPbrMaterial(Material->GetName().data());
+	FScopedConvertState ScopedConvertState(ConvertState);
+	ConvertState.DatasmithScene = DatasmithScene;
+	ConvertState.MaterialElement = PbrMaterialElement;
+	ConvertState.AssetsPath = AssetsPath;
+
+	DatasmithMaxCoronaMaterialsToUEPbrImpl::FCoronaLightMaterial MaterialProperties;
+	MaterialProperties.Parse(*Material);
+
+	Connect(PbrMaterialElement->GetEmissiveColor(),
+		COMPOSE_OR_DEFAULT2(nullptr, Multiply,
+			TextureOrColor(TEXT("Emissive Color"), MaterialProperties.EmitTexture, FDatasmithMaxMatHelper::MaxLinearColorToFLinearColor(MaterialProperties.EmitColor)),
+			&Scalar(MaterialProperties.Multiplier * 2.0f))
+	);
+
+	if (IDatasmithMaterialExpression* OpacictyExpression = ConvertTexmap(MaterialProperties.ClipTexture))
+	{
+		PbrMaterialElement->SetBlendMode(/*EBlendMode::BLEND_Masked*/1); // Set blend mode when there's opacity mask
+		Connect(PbrMaterialElement->GetOpacity(), OpacictyExpression);
+	}
+
+	PbrMaterialElement->SetShadingModel(EDatasmithShadingModel::Unlit);
+
+	MaterialElement = PbrMaterialElement;
+}
diff --git a/Engine/Source/Programs/Enterprise/Datasmith/DatasmithMaxExporter/Private/MaxMaterialsToUEPbr/DatasmithMaxCoronaMaterialsToUEPbr.h b/Engine/Source/Programs/Enterprise/Datasmith/DatasmithMaxExporter/Private/MaxMaterialsToUEPbr/DatasmithMaxCoronaMaterialsToUEPbr.h
--- a/Engine/Source/Programs/Enterprise/Datasmith/DatasmithMaxExporter/Private/MaxMaterialsToUEPbr/DatasmithMaxCoronaMaterialsToUEPbr.h
+++ b/Engine/Source/Programs/Enterprise/Datasmith/DatasmithMaxExporter/Private/MaxMaterialsToUEPbr/DatasmithMaxCoronaMaterialsToUEPbr.h
@@ -19,3 +19,10 @@ public:
 	virtual bool IsSupported( Mtl* Material ) override;
 	virtual void Convert( TSharedRef< IDatasmithScene > DatasmithScene, TSharedPtr< IDatasmithBaseMaterialElement >& MaterialElement, Mtl* Material, const TCHAR* AssetsPath ) override;
 };
+
+class FDatasmithMaxCoronaLightMaterialToUEPbr : public FDatasmithMaxMaterialsToUEPbrExpressions
+{
+public:
+	virtual bool IsSupported( Mtl* Material ) override;
+	virtual void Convert( TSharedRef< IDatasmithScene > DatasmithScene, TSharedPtr< IDatasmithBaseMaterialElement >& MaterialElement, Mtl* Material, const TCHAR* AssetsPath ) override;
+};
diff --git a/Engine/Source/Programs/Enterprise/Datasmith/DatasmithMaxExporter/Private/MaxMaterialsToUEPbr/DatasmithMaxMaterialsToUEPbr.cpp b/Engine/Source/Programs/Enterprise/Datasmith/DatasmithMaxExporter/Private/MaxMaterialsToUEPbr/DatasmithMaxMaterialsToUEPbr.cpp
--- a/Engine/Source/Programs/Enterprise/Datasmith/DatasmithMaxExporter/Private/MaxMaterialsToUEPbr/DatasmithMaxMaterialsToUEPbr.cpp
+++ b/Engine/Source/Programs/Enterprise/Datasmith/DatasmithMaxExporter/Private/MaxMaterialsToUEPbr/DatasmithMaxMaterialsToUEPbr.cpp
@@ -13,7 +13,6 @@
 #include "MaxMaterialsToUEPbr/DatasmithMaxVrayMaterialsToUEPbr.h"
 #include "MaxMaterialsToUEPbr/DatasmithMaxPhysicalMaterialToUEPbr.h"
 
-#define PHYSICALMATCLASS_UEPbr
 
 FDatasmithMaxMaterialsToUEPbr* FDatasmithMaxMaterialsToUEPbrManager::GetMaterialConverter( Mtl* Material )
 {
@@ -65,14 +64,16 @@ FDatasmithMaxMaterialsToUEPbr* FDatasmithMaxMaterialsToUEPbrManager::GetMaterial
 		static FDatasmithMaxCoronaBlendMaterialToUEPbr CoronaConverter = FDatasmithMaxCoronaBlendMaterialToUEPbr();
 		MaterialConverter = &CoronaConverter;
 	}
-#ifdef PHYSICALMATCLASS_UEPbr
-
+	else if ( MaterialClassID == CORONALIGHTMATCLASS )
+	{
+		static FDatasmithMaxCoronaLightMaterialToUEPbr CoronaConverter = FDatasmithMaxCoronaLightMaterialToUEPbr();
+		MaterialConverter = &CoronaConverter;
+	}
 	else if (MaterialClassID == PHYSICALMATCLASS)
 	{
 		static FDatasmithMaxPhysicalMaterialToUEPbr PhysicalConverter = FDatasmithMaxPhysicalMaterialToUEPbr();
 		MaterialConverter = &PhysicalConverter;
 	}
-#endif
 
 	if ( MaterialConverter && MaterialConverter->IsSupported( Material ) )
 	{
diff --git a/Engine/Source/Programs/Enterprise/Datasmith/DatasmithMaxExporter/Private/MaxMaterialsToUEPbr/DatasmithMaxMaterialsToUEPbr.h b/Engine/Source/Programs/Enterprise/Datasmith/DatasmithMaxExporter/Private/MaxMaterialsToUEPbr/DatasmithMaxMaterialsToUEPbr.h
--- a/Engine/Source/Programs/Enterprise/Datasmith/DatasmithMaxExporter/Private/MaxMaterialsToUEPbr/DatasmithMaxMaterialsToUEPbr.h
+++ b/Engine/Source/Programs/Enterprise/Datasmith/DatasmithMaxExporter/Private/MaxMaterialsToUEPbr/DatasmithMaxMaterialsToUEPbr.h
@@ -88,3 +88,63 @@ public:
 	virtual bool IsSupported( const FDatasmithMaxMaterialsToUEPbr* MaxMaterialToUEPbr, Texmap* InTexmap ) const = 0;
 	virtual IDatasmithMaterialExpression* Convert( FDatasmithMaxMaterialsToUEPbr* MaxMaterialToUEPbr, Texmap* InTexmap ) = 0;
 };
+
+
+class IDatasmithMaterialExpression;
+class IDatasmithMaterialExpressionScalar;
+class IDatasmithMaterialExpressionColor;
+class IDatasmithMaterialExpressionGeneric;
+class IDatasmithExpressionInput;
+
+class FDatasmithMaxMaterialsToUEPbrExpressions: public FDatasmithMaxMaterialsToUEPbr
+{
+public:
+	TSharedPtr<IDatasmithUEPbrMaterialElement> GetMaterialElement();
+
+	IDatasmithMaterialExpressionScalar& Scalar(float Value);
+
+	IDatasmithMaterialExpressionColor& Color(const FLinearColor& Value);
+
+	IDatasmithMaterialExpression* WeightTextureOrScalar(const DatasmithMaxTexmapParser::FMapParameter& TextureWeight, float Weight);
+
+	IDatasmithMaterialExpressionGeneric& Add(IDatasmithMaterialExpression& A, IDatasmithMaterialExpression& B);
+
+	IDatasmithMaterialExpressionGeneric& Subtract(IDatasmithMaterialExpression& A, IDatasmithMaterialExpression& B);
+
+	IDatasmithMaterialExpressionGeneric& Multiply(IDatasmithMaterialExpression& A, IDatasmithMaterialExpression& B);
+
+	IDatasmithMaterialExpressionGeneric& Divide(IDatasmithMaterialExpression& A, IDatasmithMaterialExpression& B);
+
+	IDatasmithMaterialExpressionGeneric& Desaturate(IDatasmithMaterialExpression& A);
+
+	IDatasmithMaterialExpressionGeneric& Power(IDatasmithMaterialExpression& A, IDatasmithMaterialExpression& B);
+
+	IDatasmithMaterialExpressionGeneric& Lerp(IDatasmithMaterialExpression& A, IDatasmithMaterialExpression& B, IDatasmithMaterialExpression& Alpha);
+
+	IDatasmithMaterialExpressionGeneric& Fresnel(IDatasmithMaterialExpression* Exponent=nullptr, IDatasmithMaterialExpression* BaseReflectFraction=nullptr); // Any input can be null
+
+
+	IDatasmithMaterialExpression* ApplyWeightExpression(IDatasmithMaterialExpression* ValueExpression, IDatasmithMaterialExpression* WeightExpression);
+
+	IDatasmithMaterialExpression& CalcIORComplex(double IORn, double IORk, IDatasmithMaterialExpression& ToBeConnected90, IDatasmithMaterialExpression& ToBeConnected0);
+
+
+	void Connect(IDatasmithExpressionInput& Input, IDatasmithMaterialExpression& ValueExpression);
+	bool Connect(IDatasmithExpressionInput& Input, IDatasmithMaterialExpression* ValueExpression); // Connect if not null
+
+	IDatasmithMaterialExpression* TextureOrColor(const TCHAR* Name,
+	                                             const DatasmithMaxTexmapParser::FMapParameter& Map, FLinearColor Color);
+
+	IDatasmithMaterialExpression* TextureOrScalar(const TCHAR* Name,
+	                                              const DatasmithMaxTexmapParser::FMapParameter& Map, float Value);
+
+	IDatasmithMaterialExpressionGeneric& OneMinus(IDatasmithMaterialExpression& Expression);
+};
+
+// Create expression from expression parameters when they are all non-null or return default value
+// Pass params by value in order to prevent statement creating expression to be evaluated twice
+// e.g. COMPOSE_OR_DEFAULT1(nullptr, Multiply, &Scalar(1.0), SomeExpression) won't create two expressions of each Scalar
+// Func should be a method of this of free function - so it can be just substituted 
+#define COMPOSE_OR_DEFAULT1(Default, Func, Param0) ([this](IDatasmithMaterialExpression* P0, IDatasmithMaterialExpression* D) {return P0 ? &Func(*P0) : D; } (Param0, Default))
+#define COMPOSE_OR_DEFAULT2(Default, Func, Param0, Param1) ([this](IDatasmithMaterialExpression* P0, IDatasmithMaterialExpression* P1, IDatasmithMaterialExpression* D) { return (P0 && P1) ? &Func(*P0, *P1) : D;} (Param0, Param1, Default))
+
diff --git a/Engine/Source/Programs/Enterprise/Datasmith/DatasmithMaxExporter/Private/MaxMaterialsToUEPbr/DatasmithMaxPhysicalMaterialToUEPbr.cpp b/Engine/Source/Programs/Enterprise/Datasmith/DatasmithMaxExporter/Private/MaxMaterialsToUEPbr/DatasmithMaxPhysicalMaterialToUEPbr.cpp
--- a/Engine/Source/Programs/Enterprise/Datasmith/DatasmithMaxExporter/Private/MaxMaterialsToUEPbr/DatasmithMaxPhysicalMaterialToUEPbr.cpp
+++ b/Engine/Source/Programs/Enterprise/Datasmith/DatasmithMaxExporter/Private/MaxMaterialsToUEPbr/DatasmithMaxPhysicalMaterialToUEPbr.cpp
@@ -16,10 +16,6 @@ MAX_INCLUDES_START
 MAX_INCLUDES_END
 #include "Windows/HideWindowsPlatformTypes.h"
 
-// Create expression from expression parameters when they are all non-null or return default value
-#define COMPOSE_OR_DEFAULT1(Default, Func, Param0) (Param0 ? &Func(*Param0) : Default)
-#define COMPOSE_OR_DEFAULT2(Default, Func, Param0, Param1) ((Param0 && Param1) ? &Func(*Param0, *Param1) : Default)
-
 TSharedPtr<IDatasmithUEPbrMaterialElement> FDatasmithMaxMaterialsToUEPbrExpressions::GetMaterialElement()
 {
 	return ConvertState.MaterialElement;
diff --git a/Engine/Source/Programs/Enterprise/Datasmith/DatasmithMaxExporter/Private/MaxMaterialsToUEPbr/DatasmithMaxPhysicalMaterialToUEPbr.h b/Engine/Source/Programs/Enterprise/Datasmith/DatasmithMaxExporter/Private/MaxMaterialsToUEPbr/DatasmithMaxPhysicalMaterialToUEPbr.h
--- a/Engine/Source/Programs/Enterprise/Datasmith/DatasmithMaxExporter/Private/MaxMaterialsToUEPbr/DatasmithMaxPhysicalMaterialToUEPbr.h
+++ b/Engine/Source/Programs/Enterprise/Datasmith/DatasmithMaxExporter/Private/MaxMaterialsToUEPbr/DatasmithMaxPhysicalMaterialToUEPbr.h
@@ -4,57 +4,6 @@
 
 #include "DatasmithMaxMaterialsToUEPbr.h"
 
-class IDatasmithMaterialExpression;
-class IDatasmithMaterialExpressionScalar;
-class IDatasmithMaterialExpressionColor;
-class IDatasmithMaterialExpressionGeneric;
-class IDatasmithExpressionInput;
-
-class FDatasmithMaxMaterialsToUEPbrExpressions: public FDatasmithMaxMaterialsToUEPbr
-{
-public:
-	TSharedPtr<IDatasmithUEPbrMaterialElement> GetMaterialElement();
-
-	IDatasmithMaterialExpressionScalar& Scalar(float Value);
-
-	IDatasmithMaterialExpressionColor& Color(const FLinearColor& Value);
-
-	IDatasmithMaterialExpression* WeightTextureOrScalar(const DatasmithMaxTexmapParser::FMapParameter& TextureWeight, float Weight);
-
-	IDatasmithMaterialExpressionGeneric& Add(IDatasmithMaterialExpression& A, IDatasmithMaterialExpression& B);
-
-	IDatasmithMaterialExpressionGeneric& Subtract(IDatasmithMaterialExpression& A, IDatasmithMaterialExpression& B);
-
-	IDatasmithMaterialExpressionGeneric& Multiply(IDatasmithMaterialExpression& A, IDatasmithMaterialExpression& B);
-
-	IDatasmithMaterialExpressionGeneric& Divide(IDatasmithMaterialExpression& A, IDatasmithMaterialExpression& B);
-
-	IDatasmithMaterialExpressionGeneric& Desaturate(IDatasmithMaterialExpression& A);
-
-	IDatasmithMaterialExpressionGeneric& Power(IDatasmithMaterialExpression& A, IDatasmithMaterialExpression& B);
-
-	IDatasmithMaterialExpressionGeneric& Lerp(IDatasmithMaterialExpression& A, IDatasmithMaterialExpression& B, IDatasmithMaterialExpression& Alpha);
-
-	IDatasmithMaterialExpressionGeneric& Fresnel(IDatasmithMaterialExpression* Exponent=nullptr, IDatasmithMaterialExpression* BaseReflectFraction=nullptr); // Any input can be null
-
-
-	IDatasmithMaterialExpression* ApplyWeightExpression(IDatasmithMaterialExpression* ValueExpression, IDatasmithMaterialExpression* WeightExpression);
-
-	IDatasmithMaterialExpression& CalcIORComplex(double IORn, double IORk, IDatasmithMaterialExpression& ToBeConnected90, IDatasmithMaterialExpression& ToBeConnected0);
-
-
-	void Connect(IDatasmithExpressionInput& Input, IDatasmithMaterialExpression& ValueExpression);
-	bool Connect(IDatasmithExpressionInput& Input, IDatasmithMaterialExpression* ValueExpression); // Connect if not null
-
-	IDatasmithMaterialExpression* TextureOrColor(const TCHAR* Name,
-	                                             const DatasmithMaxTexmapParser::FMapParameter& Map, FLinearColor Color);
-
-	IDatasmithMaterialExpression* TextureOrScalar(const TCHAR* Name,
-	                                              const DatasmithMaxTexmapParser::FMapParameter& Map, float Value);
-
-	IDatasmithMaterialExpressionGeneric& OneMinus(IDatasmithMaterialExpression& Expression);
-};
-
 class FDatasmithMaxPhysicalMaterialToUEPbr : public FDatasmithMaxMaterialsToUEPbrExpressions
 {
 public:
diff --git a/Engine/Source/Runtime/Datasmith/DatasmithCore/Public/DatasmithDefinitions.h b/Engine/Source/Runtime/Datasmith/DatasmithCore/Public/DatasmithDefinitions.h
--- a/Engine/Source/Runtime/Datasmith/DatasmithCore/Public/DatasmithDefinitions.h
+++ b/Engine/Source/Runtime/Datasmith/DatasmithCore/Public/DatasmithDefinitions.h
@@ -341,14 +341,15 @@ enum class EDatasmithShaderUsage
 	LightFunction
 };
 
-static const TCHAR* DatasmithShadingModelStrings[] = { TEXT("DefaultLit"), TEXT("ThinTranslucent"), TEXT("Subsurface"), TEXT("ClearCoat") };
+static const TCHAR* DatasmithShadingModelStrings[] = { TEXT("DefaultLit"), TEXT("ThinTranslucent"), TEXT("Subsurface"), TEXT("ClearCoat"), TEXT("Unlit") };
 
 enum class EDatasmithShadingModel : uint8
 {
 	DefaultLit,
 	ThinTranslucent,
 	Subsurface,
-	ClearCoat
+	ClearCoat,
+	Unlit
 };
 
 UENUM()
