3ds Max DirectLink: MentalRay Arch & Design and VRay Light Materials to UEPbr

diff --git a/Engine/Plugins/Enterprise/DatasmithImporter/Source/DatasmithImporter/Private/DatasmithMaterialExpressions.cpp b/Engine/Plugins/Enterprise/DatasmithImporter/Source/DatasmithImporter/Private/DatasmithMaterialExpressions.cpp
--- a/Engine/Plugins/Enterprise/DatasmithImporter/Source/DatasmithImporter/Private/DatasmithMaterialExpressions.cpp
+++ b/Engine/Plugins/Enterprise/DatasmithImporter/Source/DatasmithImporter/Private/DatasmithMaterialExpressions.cpp
@@ -5,6 +5,7 @@
 #include "DatasmithAssetUserData.h"
 #include "DatasmithImportContext.h"
 #include "DatasmithMaterialElements.h"
+#include "DatasmithMaterialsUtils.h"
 #include "DatasmithScene.h"
 #include "DatasmithTypes.h"
 #include "IDatasmithSceneElements.h"
@@ -234,43 +235,6 @@ namespace
 	}
 }
 
-FLinearColor FDatasmithMaterialExpressions::TemperatureToColor(float Kelvin)
-{
-	float Temp = Kelvin / 100.0f;
-	float Red, Green, Blue;
-
-	if (Temp <= 66.0f)
-	{
-		Red = 255.0f;
-		Green = Temp;
-		Green = 99.4708025861f * FMath::Loge(Green) - 161.1195681661f;
-		if (Temp <= 19.0f)
-		{
-			Blue = 0.0f;
-		}
-		else
-		{
-			Blue = Temp - 10.0f;
-			Blue = 138.5177312231f * FMath::Loge(Blue) - 305.0447927307f;
-		}
-	}
-	else
-	{
-		Red = Temp - 60.0f;
-		Red = 329.698727446f * FMath::Pow(Red, -0.1332047592f);
-		Green = Temp - 60.f;
-		Green = 288.1221695283f * FMath::Pow(Green, -0.0755148492f);
-		Blue = 255.0f;
-	}
-
-	FLinearColor Color;
-	Color.R = Red / 255.0f;
-	Color.G = Green / 255.0f;
-	Color.B = Blue / 255.0f;
-
-	return Color;
-}
-
 UMaterialExpressionMakeMaterialAttributes* FDatasmithMaterialExpressions::FindOrAddAttributesFromMatFunc(UMaterialFunction* Func)
 {
 	for (int32 i = 0; i < Func->FunctionExpressions.Num(); i++)
@@ -2169,7 +2133,7 @@ void FDatasmithMaterialExpressions::CreateDatasmithMaterialHelper(UPackage* Pack
 			}
 			else if (ShaderElement->GetEmitTemperature() > 0)
 			{
-				EmissiveExpression = AddRGBExpression(TemperatureToColor(ShaderElement->GetEmitTemperature()), 0, 0, UnrealMaterial, EDatasmithTextureSlot::NOSLOT);
+				EmissiveExpression = AddRGBExpression(DatasmithMaterialsUtils::TemperatureToColor(ShaderElement->GetEmitTemperature()), 0, 0, UnrealMaterial, EDatasmithTextureSlot::NOSLOT);
 			}
 		}
 	}
diff --git a/Engine/Source/Programs/Enterprise/Datasmith/DatasmithMaxExporter/Private/DatasmithMaxCorona.cpp b/Engine/Source/Programs/Enterprise/Datasmith/DatasmithMaxExporter/Private/DatasmithMaxCorona.cpp
--- a/Engine/Source/Programs/Enterprise/Datasmith/DatasmithMaxExporter/Private/DatasmithMaxCorona.cpp
+++ b/Engine/Source/Programs/Enterprise/Datasmith/DatasmithMaxExporter/Private/DatasmithMaxCorona.cpp
@@ -7,6 +7,8 @@
 #include "DatasmithMaxSceneParser.h"
 #include "DatasmithMaxSceneExporter.h"
 
+#include "DatasmithMaterialsUtils.h"
+
 #include "Misc/Paths.h"
 
 float GetCoronaTexmapGamma(BitmapTex* InBitmapTex)
@@ -691,7 +693,7 @@ FString FDatasmithMaxMatWriter::DumpCoronaColor(TSharedPtr<IDatasmithCompositeTe
 		CoronaColor.B = ColorParameters.ColorHdr.Z;
 		break;
 	case 2:
-		CoronaColor = FDatasmithMaxMatHelper::MaxLinearColorToFLinearColor( FDatasmithMaxMatHelper::TemperatureToColor( ColorParameters.Temperature ) );
+		CoronaColor = DatasmithMaterialsUtils::TemperatureToColor( ColorParameters.Temperature );
 		ColorParameters.bInputIsLinear = true;
 		break;
 	case 3:
diff --git a/Engine/Source/Programs/Enterprise/Datasmith/DatasmithMaxExporter/Private/DatasmithMaxTexmapParser.h b/Engine/Source/Programs/Enterprise/Datasmith/DatasmithMaxExporter/Private/DatasmithMaxTexmapParser.h
--- a/Engine/Source/Programs/Enterprise/Datasmith/DatasmithMaxExporter/Private/DatasmithMaxTexmapParser.h
+++ b/Engine/Source/Programs/Enterprise/Datasmith/DatasmithMaxExporter/Private/DatasmithMaxTexmapParser.h
@@ -21,6 +21,11 @@ namespace DatasmithMaxTexmapParser
 		{
 		}
 
+		bool IsMapPresentAndEnabled()
+		{
+			return Map && bEnabled;
+		}
+
 		Texmap* Map;
 		bool bEnabled;
 		float Weight;
diff --git a/Engine/Source/Programs/Enterprise/Datasmith/DatasmithMaxExporter/Private/DatasmithMaxWriter.cpp b/Engine/Source/Programs/Enterprise/Datasmith/DatasmithMaxExporter/Private/DatasmithMaxWriter.cpp
--- a/Engine/Source/Programs/Enterprise/Datasmith/DatasmithMaxExporter/Private/DatasmithMaxWriter.cpp
+++ b/Engine/Source/Programs/Enterprise/Datasmith/DatasmithMaxExporter/Private/DatasmithMaxWriter.cpp
@@ -367,43 +367,6 @@ Mtl* FDatasmithMaxMatHelper::GetRenderedXRefMaterial(Mtl* Material)
 	return RenderedMaterial;
 }
 
-BMM_Color_fl FDatasmithMaxMatHelper::TemperatureToColor(float Kelvin)
-{
-	float Temp = Kelvin / 100.0f;
-	float Red;
-	float Green;
-	float Blue;
-
-	if (Temp <= 66.0f)
-	{
-		Red = 255.0f;
-		Green = Temp;
-		Green = 99.4708025861f * log(Green) - 161.1195681661f;
-		if (Temp <= 19.0f)
-		{
-			Blue = 0.0f;
-		}
-		else
-		{
-			Blue = Temp - 10.0f;
-			Blue = 138.5177312231f * log(Blue) - 305.0447927307f;
-		}
-	}
-	else
-	{
-		Red = Temp - 60.0f;
-		Red = 329.698727446f * FMath::Pow(Red, -0.1332047592f);
-		Green = Temp - 60;
-		Green = 288.1221695283f * FMath::Pow(Green, -0.0755148492f);
-		Blue = 255.0f;
-	}
-	BMM_Color_fl Color;
-	Color.r = Red / 255.0f;
-	Color.g = Green / 255.0f;
-	Color.b = Blue / 255.0f;
-	return Color;
-}
-
 FLinearColor FDatasmithMaxMatHelper::MaxColorToFLinearColor(BMM_Color_fl Color, float Multiplier /*= 1.0f*/)
 {
 	Color.r = FMath::Pow(Color.r, 1.0f / FDatasmithExportOptions::ColorGamma);
diff --git a/Engine/Source/Programs/Enterprise/Datasmith/DatasmithMaxExporter/Private/MaxMaterialsToUEPbr/DatasmithMaxCoronaTexmapToUEPbr.cpp b/Engine/Source/Programs/Enterprise/Datasmith/DatasmithMaxExporter/Private/MaxMaterialsToUEPbr/DatasmithMaxCoronaTexmapToUEPbr.cpp
--- a/Engine/Source/Programs/Enterprise/Datasmith/DatasmithMaxExporter/Private/MaxMaterialsToUEPbr/DatasmithMaxCoronaTexmapToUEPbr.cpp
+++ b/Engine/Source/Programs/Enterprise/Datasmith/DatasmithMaxExporter/Private/MaxMaterialsToUEPbr/DatasmithMaxCoronaTexmapToUEPbr.cpp
@@ -9,6 +9,8 @@
 #include "DatasmithMaxWriter.h"
 #include "MaxMaterialsToUEPbr/DatasmithMaxTexmapToUEPbr.h"
 
+#include "DatasmithMaterialsUtils.h"
+
 #include "Misc/Paths.h"
 
 #include "Windows/AllowWindowsPlatformTypes.h"
@@ -107,7 +109,7 @@ IDatasmithMaterialExpression* FDatasmithMaxCoronaColorToUEPbr::Convert( FDatasmi
 		CoronaColor.B = ColorParameters.ColorHdr.Z;
 		break;
 	case 2:
-		CoronaColor = FDatasmithMaxMatHelper::MaxLinearColorToFLinearColor( FDatasmithMaxMatHelper::TemperatureToColor( ColorParameters.Temperature ) );
+		CoronaColor = DatasmithMaterialsUtils::TemperatureToColor( ColorParameters.Temperature );
 		ColorParameters.bInputIsLinear = true;
 		break;
 	case 3:
diff --git a/Engine/Source/Programs/Enterprise/Datasmith/DatasmithMaxExporter/Private/MaxMaterialsToUEPbr/DatasmithMaxMaterialsToUEPbr.cpp b/Engine/Source/Programs/Enterprise/Datasmith/DatasmithMaxExporter/Private/MaxMaterialsToUEPbr/DatasmithMaxMaterialsToUEPbr.cpp
--- a/Engine/Source/Programs/Enterprise/Datasmith/DatasmithMaxExporter/Private/MaxMaterialsToUEPbr/DatasmithMaxMaterialsToUEPbr.cpp
+++ b/Engine/Source/Programs/Enterprise/Datasmith/DatasmithMaxExporter/Private/MaxMaterialsToUEPbr/DatasmithMaxMaterialsToUEPbr.cpp
@@ -12,6 +12,7 @@
 #include "MaxMaterialsToUEPbr/DatasmithMaxScanlineMaterialsToUEPbr.h"
 #include "MaxMaterialsToUEPbr/DatasmithMaxVrayMaterialsToUEPbr.h"
 #include "MaxMaterialsToUEPbr/DatasmithMaxPhysicalMaterialToUEPbr.h"
+#include "MaxMaterialsToUEPbr/DatasmithMaxMentalMaterialToUEPbr.h"
 
 
 FDatasmithMaxMaterialsToUEPbr* FDatasmithMaxMaterialsToUEPbrManager::GetMaterialConverter( Mtl* Material )
@@ -44,6 +45,12 @@ FDatasmithMaxMaterialsToUEPbr* FDatasmithMaxMaterialsToUEPbrManager::GetMaterial
 		static FDatasmithMaxVRayBlendMaterialToUEPbr VrayBlendConverter = FDatasmithMaxVRayBlendMaterialToUEPbr();
 		MaterialConverter = &VrayBlendConverter;
 	}
+	else if ( MaterialClassID == VRAYLIGHTMATCLASS )
+	{
+		static FDatasmithMaxVRayLightMaterialToUEPbr CoronaConverter = FDatasmithMaxVRayLightMaterialToUEPbr();
+		MaterialConverter = &CoronaConverter;
+	}
+	
 	else if ( MaterialClassID == STANDARDMATCLASS )
 	{
 		static FDatasmithMaxScanlineMaterialsToUEPbr ScanlineConverter = FDatasmithMaxScanlineMaterialsToUEPbr();
@@ -74,6 +81,11 @@ FDatasmithMaxMaterialsToUEPbr* FDatasmithMaxMaterialsToUEPbrManager::GetMaterial
 		static FDatasmithMaxPhysicalMaterialToUEPbr PhysicalConverter = FDatasmithMaxPhysicalMaterialToUEPbr();
 		MaterialConverter = &PhysicalConverter;
 	}
+	else if (MaterialClassID == ARCHDESIGNMATCLASS)
+	{
+		static FDatasmithMaxMentalMaterialToUEPbr PhysicalConverter = FDatasmithMaxMentalMaterialToUEPbr();
+		MaterialConverter = &PhysicalConverter;
+	}
 
 	if ( MaterialConverter && MaterialConverter->IsSupported( Material ) )
 	{
diff --git a/Engine/Source/Programs/Enterprise/Datasmith/DatasmithMaxExporter/Private/MaxMaterialsToUEPbr/DatasmithMaxMaterialsToUEPbr.h b/Engine/Source/Programs/Enterprise/Datasmith/DatasmithMaxExporter/Private/MaxMaterialsToUEPbr/DatasmithMaxMaterialsToUEPbr.h
--- a/Engine/Source/Programs/Enterprise/Datasmith/DatasmithMaxExporter/Private/MaxMaterialsToUEPbr/DatasmithMaxMaterialsToUEPbr.h
+++ b/Engine/Source/Programs/Enterprise/Datasmith/DatasmithMaxExporter/Private/MaxMaterialsToUEPbr/DatasmithMaxMaterialsToUEPbr.h
@@ -141,10 +141,23 @@ public:
 	IDatasmithMaterialExpressionGeneric& OneMinus(IDatasmithMaterialExpression& Expression);
 };
 
+// Utility macros to simplify material expression composition:
+
 // Create expression from expression parameters when they are all non-null or return default value
-// Pass params by value in order to prevent statement creating expression to be evaluated twice
+// E.g. instead of:
+// IDatasmithMaterialExpression* SomeExpression = MakeExpressionCanReturnNull(...)
+// IDatasmithMaterialExpression* ResultExpression = SomeExpression ? Desaturate(SomeExpression) : nullptr;
+// Can do this:
+// IDatasmithMaterialExpression* ResultExpression = COMPOSE_OR_NULL(Desaturate, MakeExpressionCanReturnNull(...));
+// Which allows chaining expressions: COMPOSE_OR_DEFAULT2(&Scalar(1), Multiply, COMPOSE_OR_NULL(Desaturate, MakeExpressionCanReturnNull(...)), AnotherExpression)
+// Implementation details: Pass params by value in order to prevent statement creating expression to be evaluated twice
 // e.g. COMPOSE_OR_DEFAULT1(nullptr, Multiply, &Scalar(1.0), SomeExpression) won't create two expressions of each Scalar
 // Func should be a method of this of free function - so it can be just substituted 
 #define COMPOSE_OR_DEFAULT1(Default, Func, Param0) ([this](IDatasmithMaterialExpression* P0, IDatasmithMaterialExpression* D) {return P0 ? &Func(*P0) : D; } (Param0, Default))
 #define COMPOSE_OR_DEFAULT2(Default, Func, Param0, Param1) ([this](IDatasmithMaterialExpression* P0, IDatasmithMaterialExpression* P1, IDatasmithMaterialExpression* D) { return (P0 && P1) ? &Func(*P0, *P1) : D;} (Param0, Param1, Default))
+// Calls Func when param is not NULL, returns NULL otherwise
+#define COMPOSE_OR_NULL(Func, Param0) ([this](IDatasmithMaterialExpression* P0) {return P0 ? &Func(*P0) : nullptr; } (Param0 ))
+#define COMPOSE_OR_NULL2(Func, Param0, Param1) ([this](IDatasmithMaterialExpression* P0, IDatasmithMaterialExpression* P1) { return (P0 && P1) ? &Func(*P0, *P1) : nullptr;} (Param0, Param1))
+
+
 
diff --git a/Engine/Source/Programs/Enterprise/Datasmith/DatasmithMaxExporter/Private/MaxMaterialsToUEPbr/DatasmithMaxMentalMaterialToUEPbr.cpp b/Engine/Source/Programs/Enterprise/Datasmith/DatasmithMaxExporter/Private/MaxMaterialsToUEPbr/DatasmithMaxMentalMaterialToUEPbr.cpp
new file mode 100644
--- /dev/null
+++ b/Engine/Source/Programs/Enterprise/Datasmith/DatasmithMaxExporter/Private/MaxMaterialsToUEPbr/DatasmithMaxMentalMaterialToUEPbr.cpp
@@ -0,0 +1,363 @@
+// Copyright Epic Games, Inc. All Rights Reserved.
+
+#include "MaxMaterialsToUEPbr/DatasmithMaxMentalMaterialToUEPbr.h"
+
+#include "DatasmithMaxClassIDs.h"
+#include "DatasmithMaxTexmapParser.h"
+#include "DatasmithMaxWriter.h"
+#include "DatasmithSceneFactory.h"
+#include "MaxMaterialsToUEPbr/DatasmithMaxTexmapToUEPbr.h"
+#include "MaxMaterialsToUEPbr/DatasmithMaxCoronaTexmapToUEPbr.h"
+
+#include "DatasmithMaterialsUtils.h"
+
+
+#include "Windows/AllowWindowsPlatformTypes.h"
+MAX_INCLUDES_START
+	#include "max.h"
+	#include "mtl.h"
+MAX_INCLUDES_END
+#include "Windows/HideWindowsPlatformTypes.h"
+
+bool FDatasmithMaxMentalMaterialToUEPbr::IsSupported( Mtl* Material )
+{
+	return true;
+}
+
+struct FMaxMentalMaterial
+{
+	float DiffuseWeight = 0.f;
+	DatasmithMaxTexmapParser::FMapParameter DiffuseColorMap;
+	FLinearColor DiffuseColor;
+
+	float Ior = 0;
+
+	float ReflectionWeight = 0.f;
+	DatasmithMaxTexmapParser::FMapParameter ReflectionColorMap;
+	FLinearColor ReflectionColor;
+	bool bReflectionUseFresnel = true;
+	float bReflectionFuncLow = 0;
+	float bReflectionFuncHigh = 1;
+	float bReflectionFuncCurve = 1;
+	bool bReflectionMetal = false;
+
+	bool bThinWalled = false;
+	float RefractionWeight = 0.f;
+	DatasmithMaxTexmapParser::FMapParameter RefractionColorMap;
+	FLinearColor RefractionColor;
+
+	DatasmithMaxTexmapParser::FMapParameter CutoutMap;
+
+	DatasmithMaxTexmapParser::FMapParameter GlossinessMap;
+	float Glossiness = 0;
+
+	bool bSelfIllumEnabled = false;
+	bool bSelfIllumUseTemperature = true;
+	FLinearColor SelfIllumColorFilter;
+	float SelfIllumColorKelvin = 0.0f;
+	DatasmithMaxTexmapParser::FMapParameter SelfIllumMap;
+
+	float BumpAmount = 0;
+	DatasmithMaxTexmapParser::FMapParameter BumpMap;
+
+	void Parse(Mtl& Material)
+	{
+		int NumParamBlocks = Material.NumParamBlocks();
+
+		for (int j = 0; j < NumParamBlocks; j++)
+		{
+			IParamBlock2* ParamBlock2 = Material.GetParamBlockByID((short)j);
+			// The the descriptor to 'decode'
+			ParamBlockDesc2* ParamBlockDesc = ParamBlock2->GetDesc();
+			// Loop through all the defined parameters therein
+			for (int i = 0; i < ParamBlockDesc->count; i++)
+			{
+				const ParamDef& ParamDefinition = ParamBlockDesc->paramdefs[i];
+
+				if (FCString::Stricmp(ParamDefinition.int_name, TEXT("mapM0")) == 0)
+				{
+					DiffuseColorMap.Map = ParamBlock2->GetTexmap(ParamDefinition.ID, GetCOREInterface()->GetTime());
+				}
+				else if (FCString::Stricmp(ParamDefinition.int_name, TEXT("mapM2")) == 0)
+				{
+					ReflectionColorMap.Map = ParamBlock2->GetTexmap(ParamDefinition.ID, GetCOREInterface()->GetTime());
+				}
+				else if (FCString::Stricmp(ParamDefinition.int_name, TEXT("mapM4")) == 0)
+				{
+					RefractionColorMap.Map = ParamBlock2->GetTexmap(ParamDefinition.ID, GetCOREInterface()->GetTime());
+				}
+				else if (FCString::Stricmp(ParamDefinition.int_name, TEXT("cutout_map")) == 0)
+				{
+					CutoutMap.Map = ParamBlock2->GetTexmap(ParamDefinition.ID, GetCOREInterface()->GetTime());
+				}
+				else if (FCString::Stricmp(ParamDefinition.int_name, TEXT("cutout_map_on")) == 0)
+				{
+					CutoutMap.bEnabled = ParamBlock2->GetInt(ParamDefinition.ID, GetCOREInterface()->GetTime()) != 0;
+				}
+				else if (FCString::Stricmp(ParamDefinition.int_name, TEXT("mapM3")) == 0)
+				{
+					GlossinessMap.Map = ParamBlock2->GetTexmap(ParamDefinition.ID, GetCOREInterface()->GetTime());
+				}
+				else if (FCString::Stricmp(ParamDefinition.int_name, TEXT("self_illum_color_mode")) == 0)
+				{
+					bSelfIllumUseTemperature = ParamBlock2->GetInt(ParamDefinition.ID, GetCOREInterface()->GetTime()) != 0;
+				}
+				else if (FCString::Stricmp(ParamDefinition.int_name, TEXT("self_illum_color_kelvin")) == 0)
+				{
+					SelfIllumColorKelvin = ParamBlock2->GetFloat(ParamDefinition.ID, GetCOREInterface()->GetTime());
+				}
+				else if (FCString::Stricmp(ParamDefinition.int_name, TEXT("self_illum_color_filter")) == 0)
+				{
+					SelfIllumColorFilter = FDatasmithMaxMatHelper::MaxLinearColorToFLinearColor((BMM_Color_fl)ParamBlock2->GetAColor(ParamDefinition.ID, GetCOREInterface()->GetTime()));
+				}
+				else if (FCString::Stricmp(ParamDefinition.int_name, TEXT("self_illum_on")) == 0)
+				{
+					bSelfIllumEnabled = ParamBlock2->GetInt(ParamDefinition.ID, GetCOREInterface()->GetTime()) != 0;
+				}
+				else if (FCString::Stricmp(ParamDefinition.int_name, TEXT("self_illum_map")) == 0)
+				{
+					SelfIllumMap.Map = ParamBlock2->GetTexmap(ParamDefinition.ID, GetCOREInterface()->GetTime());
+				}
+				else if (FCString::Stricmp(ParamDefinition.int_name, TEXT("bump_map")) == 0)
+				{
+					BumpMap.Map = ParamBlock2->GetTexmap(ParamDefinition.ID, GetCOREInterface()->GetTime());
+				}
+				else if (FCString::Stricmp(ParamDefinition.int_name, TEXT("bump_map_amt")) == 0)
+				{
+					BumpAmount = ParamBlock2->GetFloat(ParamDefinition.ID, GetCOREInterface()->GetTime());
+				}
+				else if (FCString::Stricmp(ParamDefinition.int_name, TEXT("diff_color_map_on")) == 0)
+				{
+					DiffuseColorMap.bEnabled = ParamBlock2->GetInt(ParamDefinition.ID, GetCOREInterface()->GetTime()) != 0;
+				}
+				else if (FCString::Stricmp(ParamDefinition.int_name, TEXT("refl_color_map_on")) == 0)
+				{
+					ReflectionColorMap.bEnabled = ParamBlock2->GetInt(ParamDefinition.ID, GetCOREInterface()->GetTime()) != 0;
+				}
+				else if (FCString::Stricmp(ParamDefinition.int_name, TEXT("refr_color_map_on")) == 0)
+				{
+					RefractionColorMap.bEnabled = ParamBlock2->GetInt(ParamDefinition.ID, GetCOREInterface()->GetTime()) != 0;
+				}
+				else if (FCString::Stricmp(ParamDefinition.int_name, TEXT("opts_1sided")) == 0)
+				{
+					bThinWalled = ParamBlock2->GetInt(ParamDefinition.ID, GetCOREInterface()->GetTime()) != 0;
+				}
+				else if (FCString::Stricmp(ParamDefinition.int_name, TEXT("refl_gloss_map_on")) == 0)
+				{
+					GlossinessMap.bEnabled = ParamBlock2->GetInt(ParamDefinition.ID, GetCOREInterface()->GetTime()) != 0;
+				}
+				else if (FCString::Stricmp(ParamDefinition.int_name, TEXT("bump_map_on")) == 0)
+				{
+					BumpMap.bEnabled = ParamBlock2->GetInt(ParamDefinition.ID, GetCOREInterface()->GetTime()) != 0;
+				}
+				else if (FCString::Stricmp(ParamDefinition.int_name, TEXT("refl_func_fresnel")) == 0)
+				{
+					bReflectionUseFresnel = ParamBlock2->GetInt(ParamDefinition.ID, GetCOREInterface()->GetTime()) != 0;
+				}
+				else if (FCString::Stricmp(ParamDefinition.int_name, TEXT("refl_func_low")) == 0)
+				{
+					bReflectionFuncLow = ParamBlock2->GetFloat(ParamDefinition.ID, GetCOREInterface()->GetTime());
+				}
+				else if (FCString::Stricmp(ParamDefinition.int_name, TEXT("refl_func_high")) == 0)
+				{
+					bReflectionFuncHigh = ParamBlock2->GetFloat(ParamDefinition.ID, GetCOREInterface()->GetTime());
+				}
+				else if (FCString::Stricmp(ParamDefinition.int_name, TEXT("refl_func_curve")) == 0)
+				{
+					bReflectionFuncCurve = ParamBlock2->GetFloat(ParamDefinition.ID, GetCOREInterface()->GetTime());
+				}
+				else if (FCString::Stricmp(ParamDefinition.int_name, TEXT("refl_metal")) == 0)
+				{
+					bReflectionMetal = ParamBlock2->GetInt(ParamDefinition.ID, GetCOREInterface()->GetTime()) != 0;
+				}
+				else if (FCString::Stricmp(ParamDefinition.int_name, TEXT("diff_weight")) == 0)
+				{
+					DiffuseWeight = ParamBlock2->GetFloat(ParamDefinition.ID, GetCOREInterface()->GetTime());
+				}
+				else if (FCString::Stricmp(ParamDefinition.int_name, TEXT("refl_weight")) == 0)
+				{
+					ReflectionWeight = ParamBlock2->GetFloat(ParamDefinition.ID, GetCOREInterface()->GetTime());
+				}
+				else if (FCString::Stricmp(ParamDefinition.int_name, TEXT("refr_weight")) == 0)
+				{
+					RefractionWeight = ParamBlock2->GetFloat(ParamDefinition.ID, GetCOREInterface()->GetTime());
+				}
+				else if (FCString::Stricmp(ParamDefinition.int_name, TEXT("refl_color")) == 0 )
+				{
+					ReflectionColor = FDatasmithMaxMatHelper::MaxLinearColorToFLinearColor((BMM_Color_fl)ParamBlock2->GetAColor(ParamDefinition.ID, GetCOREInterface()->GetTime()));
+				}
+				else if (FCString::Stricmp(ParamDefinition.int_name, TEXT("refr_ior")) == 0) 
+				{
+					Ior = ParamBlock2->GetFloat(ParamDefinition.ID, GetCOREInterface()->GetTime());
+				}
+				else if (FCString::Stricmp(ParamDefinition.int_name, TEXT("diff_color")) == 0)
+				{
+					DiffuseColor = FDatasmithMaxMatHelper::MaxLinearColorToFLinearColor((BMM_Color_fl)ParamBlock2->GetAColor(ParamDefinition.ID, GetCOREInterface()->GetTime()));
+				}
+				else if (FCString::Stricmp(ParamDefinition.int_name, TEXT("refl_gloss")) == 0)
+				{
+					Glossiness = ParamBlock2->GetFloat(ParamDefinition.ID, GetCOREInterface()->GetTime());
+				}			
+				else if (FCString::Stricmp(ParamDefinition.int_name, TEXT("refr_color")) == 0)
+				{
+					RefractionColor = FDatasmithMaxMatHelper::MaxLinearColorToFLinearColor((BMM_Color_fl)ParamBlock2->GetAColor(ParamDefinition.ID, GetCOREInterface()->GetTime()));
+				}
+			}
+			ParamBlock2->ReleaseDesc();
+		}
+
+	}
+};
+
+
+void FDatasmithMaxMentalMaterialToUEPbr::Convert( TSharedRef< IDatasmithScene > DatasmithScene, TSharedPtr< IDatasmithBaseMaterialElement >& MaterialElement, Mtl* Material, const TCHAR* AssetsPath )
+{
+	if ( !Material )
+	{
+		return;
+	}
+
+	TSharedRef< IDatasmithUEPbrMaterialElement > PbrMaterialElement = FDatasmithSceneFactory::CreateUEPbrMaterial( Material->GetName().data() );
+	FScopedConvertState ScopedConvertState(ConvertState);
+	ConvertState.DatasmithScene = DatasmithScene;
+	ConvertState.MaterialElement = PbrMaterialElement;
+	ConvertState.AssetsPath = AssetsPath;
+
+	FMaxMentalMaterial MaterialProperties;
+	MaterialProperties.Parse(*Material);
+
+	// Diffuse
+	IDatasmithMaterialExpression* DiffuseColorExpression = COMPOSE_OR_NULL2(Multiply, 
+		TextureOrColor(TEXT("Diffuse Color"), MaterialProperties.DiffuseColorMap, MaterialProperties.DiffuseColor), 
+		&Scalar(MaterialProperties.DiffuseWeight));
+
+	// Emission
+	if (MaterialProperties.bSelfIllumEnabled)
+	{
+		Connect(PbrMaterialElement->GetEmissiveColor(), MaterialProperties.bSelfIllumUseTemperature
+			                                                ? &Color(DatasmithMaterialsUtils::TemperatureToColor(MaterialProperties.SelfIllumColorKelvin))
+			                                                : TextureOrColor(
+				                                                TEXT("Emissive Color"), MaterialProperties.SelfIllumMap,
+				                                                MaterialProperties.SelfIllumColorFilter));
+	}
+
+	// Specular
+	IDatasmithMaterialExpression* ReflectivityExpression = COMPOSE_OR_NULL2(Multiply,
+		TextureOrColor(TEXT("Specular"), MaterialProperties.ReflectionColorMap, MaterialProperties.ReflectionColor),
+		&Scalar(MaterialProperties.ReflectionWeight));
+	Connect(PbrMaterialElement->GetSpecular(), COMPOSE_OR_NULL(Desaturate, ReflectivityExpression));
+
+	if (MaterialProperties.bReflectionMetal)
+	{
+		Connect(PbrMaterialElement->GetMetallic(), COMPOSE_OR_NULL(Desaturate, ReflectivityExpression));		
+	}
+
+	// Transparency
+	IDatasmithMaterialExpression* TransparencyExpression = nullptr;
+	if (MaterialProperties.RefractionColorMap.IsMapPresentAndEnabled())
+	{
+		TransparencyExpression = COMPOSE_OR_NULL2(Multiply,
+			TextureOrColor(TEXT("Opacity"), MaterialProperties.RefractionColorMap, MaterialProperties.RefractionColor), 
+			&Scalar(MaterialProperties.RefractionWeight));
+	}
+	else
+	{
+		FLinearColor Transparency = MaterialProperties.RefractionColor * MaterialProperties.RefractionWeight;
+	
+		if (!Transparency.IsAlmostBlack()) // Don't create transparency when material is actually black(otherwise present transparency expression will make imported material Translucent instead of Opaque)
+		{
+			TransparencyExpression = &Color(Transparency);
+		}
+	}
+
+	// Roughness
+	if (IDatasmithMaterialExpression* RoughnessTextureExpression = ConvertTexmap(MaterialProperties.GlossinessMap))
+	{
+		Connect(PbrMaterialElement->GetRoughness(), COMPOSE_OR_NULL(OneMinus, COMPOSE_OR_NULL(Desaturate, RoughnessTextureExpression)));
+	}
+	else
+	{
+		// Use source glossiness when there's some reflectivity set in Max material(otherwise it's just a matte)
+		float Roughness = (MaterialProperties.ReflectionColorMap.IsMapPresentAndEnabled() || !MaterialProperties.ReflectionColor.IsAlmostBlack())
+			                  ? 1.0f - MaterialProperties.Glossiness // Just invert Glossiness to get Roughness
+			                  : 0.75f;
+		Connect(PbrMaterialElement->GetRoughness(), Scalar(Roughness));
+	}
+
+	// Opacity/refraction/masking
+	IDatasmithMaterialExpression* CutoutExpression = ConvertTexmap(MaterialProperties.CutoutMap);
+	IDatasmithMaterialExpression* BaseColorExpression = DiffuseColorExpression;
+	if (TransparencyExpression)
+	{
+		IDatasmithMaterialExpression& TransparencyExpressionRef = *TransparencyExpression;
+		IDatasmithMaterialExpression* BaseOpacityExpression = nullptr; // opacity before cutout mask applied
+		
+		// todo: PhysicalMaterial uses similar code to setup transparency/refraction only has additional param to force ThinTranslucent
+		// ThinTranslucent, including when Ior is nearly 1
+		bool ForceThinTranslucency = MaterialProperties.bThinWalled;
+		if (ForceThinTranslucency || FMath::IsNearlyEqual(MaterialProperties.Ior, 1.0f))
+		{
+			PbrMaterialElement->SetShadingModel(EDatasmithShadingModel::ThinTranslucent);
+			IDatasmithMaterialExpressionGeneric* ThinTranslucencyMaterialOutput = PbrMaterialElement->AddMaterialExpression<IDatasmithMaterialExpressionGeneric>();
+			ThinTranslucencyMaterialOutput->SetExpressionName(TEXT("ThinTranslucentMaterialOutput"));
+
+			Connect(*ThinTranslucencyMaterialOutput->GetInput(0), TransparencyExpression);
+
+			BaseOpacityExpression = &OneMinus(Desaturate(TransparencyExpressionRef));
+		}
+		else
+		{
+			const bool bIsRefractiveGlass = MaterialProperties.bReflectionUseFresnel; // UseFresnel option is used for glass, water and other refractive dielectricsetc
+			if (bIsRefractiveGlass)
+			{
+				// Ported FDatasmithMaterialExpressions::AddGlassNode to compute base color and opacity
+				float Ior = MaterialProperties.Ior;
+
+				IDatasmithMaterialExpression& Multiply09 = Multiply(Scalar(0.9), TransparencyExpressionRef);
+
+				BaseOpacityExpression = &OneMinus(Desaturate(CalcIORComplex(Ior, 0, Multiply(Scalar(0.5), TransparencyExpressionRef), Multiply09)));
+
+				Connect(PbrMaterialElement->GetRefraction(), Lerp(Scalar(1.0f), Scalar(MaterialProperties.Ior), Fresnel()));
+
+				// todo: may add Diffuse color into the mix to accomodate non-clear opacity
+				BaseColorExpression = &Multiply(Scalar(0.6), Power(Multiply09, Scalar(5.0))); 
+			}
+			else
+			{
+				BaseOpacityExpression = &OneMinus(Desaturate(TransparencyExpressionRef));
+			}
+		}
+		Connect(PbrMaterialElement->GetOpacity(), COMPOSE_OR_DEFAULT2(BaseOpacityExpression, Multiply, BaseOpacityExpression, CutoutExpression));
+	}
+	else // NOT Transparent
+	{
+		if (CutoutExpression)
+		{
+			PbrMaterialElement->SetBlendMode(/*EBlendMode::BLEND_Masked*/1); // Specify blend mode when there's opacity mask
+			Connect(PbrMaterialElement->GetOpacity(), CutoutExpression);
+		}
+	}
+
+	Connect(PbrMaterialElement->GetBaseColor(), BaseColorExpression);
+
+	// Bump
+	// todo: similar code used in multiple places for UEPbr bump conversion
+	{
+		ConvertState.DefaultTextureMode = EDatasmithTextureMode::Bump; // Will change to normal if we pass through a normal map texmap
+		ConvertState.bCanBake = false; // Current baking fails to produce proper normal maps
+
+		MaterialProperties.BumpMap.Weight = MaterialProperties.BumpAmount;
+
+		IDatasmithMaterialExpression* BumpExpression = FDatasmithMaxTexmapToUEPbrUtils::MapOrValue( this, MaterialProperties.BumpMap, TEXT("Bump Map"), TOptional< FLinearColor >(), TOptional< float >() );
+
+		if ( BumpExpression )
+		{
+			BumpExpression->ConnectExpression( PbrMaterialElement->GetNormal() );
+			BumpExpression->SetName( TEXT("Bump Map") );
+		}
+
+		ConvertState.bCanBake = true;
+	}
+
+	MaterialElement = PbrMaterialElement;
+}
diff --git a/Engine/Source/Programs/Enterprise/Datasmith/DatasmithMaxExporter/Private/MaxMaterialsToUEPbr/DatasmithMaxMentalMaterialToUEPbr.h b/Engine/Source/Programs/Enterprise/Datasmith/DatasmithMaxExporter/Private/MaxMaterialsToUEPbr/DatasmithMaxMentalMaterialToUEPbr.h
new file mode 100644
--- /dev/null
+++ b/Engine/Source/Programs/Enterprise/Datasmith/DatasmithMaxExporter/Private/MaxMaterialsToUEPbr/DatasmithMaxMentalMaterialToUEPbr.h
@@ -0,0 +1,13 @@
+// Copyright Epic Games, Inc. All Rights Reserved.
+
+#pragma once
+
+#include "DatasmithMaxMaterialsToUEPbr.h"
+
+class FDatasmithMaxMentalMaterialToUEPbr : public FDatasmithMaxMaterialsToUEPbrExpressions
+{
+public:
+	virtual bool IsSupported( Mtl* Material ) override;
+	virtual void Convert( TSharedRef< IDatasmithScene > DatasmithScene, TSharedPtr< IDatasmithBaseMaterialElement >& MaterialElement, Mtl* Material, const TCHAR* AssetsPath ) override;
+};
+
diff --git a/Engine/Source/Programs/Enterprise/Datasmith/DatasmithMaxExporter/Private/MaxMaterialsToUEPbr/DatasmithMaxPhysicalMaterialToUEPbr.cpp b/Engine/Source/Programs/Enterprise/Datasmith/DatasmithMaxExporter/Private/MaxMaterialsToUEPbr/DatasmithMaxPhysicalMaterialToUEPbr.cpp
--- a/Engine/Source/Programs/Enterprise/Datasmith/DatasmithMaxExporter/Private/MaxMaterialsToUEPbr/DatasmithMaxPhysicalMaterialToUEPbr.cpp
+++ b/Engine/Source/Programs/Enterprise/Datasmith/DatasmithMaxExporter/Private/MaxMaterialsToUEPbr/DatasmithMaxPhysicalMaterialToUEPbr.cpp
@@ -666,12 +666,12 @@ void FDatasmithMaxPhysicalMaterialToUEPbr::Convert( TSharedRef< IDatasmithScene
 			TextureOrColor(TEXT("Specular"), PhysicalMaterialProperties.ReflectionColorMap, FDatasmithMaxMatHelper::MaxLinearColorToFLinearColor(PhysicalMaterialProperties.ReflectionColor)),
 			WeightTextureOrScalar(PhysicalMaterialProperties.ReflectivityMap, PhysicalMaterialProperties.Reflectivity));
 
-		Connect(PbrMaterialElement->GetSpecular(), COMPOSE_OR_DEFAULT1(nullptr, Desaturate, ReflectivityExpression));
+		Connect(PbrMaterialElement->GetSpecular(), COMPOSE_OR_NULL(Desaturate, ReflectivityExpression));
 	}
 
 	if (IDatasmithMaterialExpression* RoughnessTextureExpression = ConvertTexmap(PhysicalMaterialProperties.RoughnessMap))
 	{
-		Connect(PbrMaterialElement->GetRoughness(), PhysicalMaterialProperties.bRoughnessInverted ? COMPOSE_OR_DEFAULT1(nullptr, OneMinus, RoughnessTextureExpression) : RoughnessTextureExpression);
+		Connect(PbrMaterialElement->GetRoughness(), PhysicalMaterialProperties.bRoughnessInverted ? COMPOSE_OR_NULL(OneMinus, RoughnessTextureExpression) : RoughnessTextureExpression);
 	}
 	else
 	{
diff --git a/Engine/Source/Programs/Enterprise/Datasmith/DatasmithMaxExporter/Private/MaxMaterialsToUEPbr/DatasmithMaxVrayMaterialsToUEPbr.cpp b/Engine/Source/Programs/Enterprise/Datasmith/DatasmithMaxExporter/Private/MaxMaterialsToUEPbr/DatasmithMaxVrayMaterialsToUEPbr.cpp
--- a/Engine/Source/Programs/Enterprise/Datasmith/DatasmithMaxExporter/Private/MaxMaterialsToUEPbr/DatasmithMaxVrayMaterialsToUEPbr.cpp
+++ b/Engine/Source/Programs/Enterprise/Datasmith/DatasmithMaxExporter/Private/MaxMaterialsToUEPbr/DatasmithMaxVrayMaterialsToUEPbr.cpp
@@ -435,6 +435,70 @@ namespace DatasmithMaxVRayMaterialsToUEPbrImpl
 
 		return VRayBlendMaterialProperties;
 	}
+
+	struct FVRayLightMaterial
+	{
+		DatasmithMaxTexmapParser::FMapParameter EmitTexture;
+		DatasmithMaxTexmapParser::FMapParameter ClipTexture;
+		BMM_Color_fl EmitColor;
+
+		float Multiplier = 1.0f;
+
+		void Parse(Mtl& Material)
+		{
+			int NumParamBlocks = Material.NumParamBlocks();
+
+			for (int j = 0; j < NumParamBlocks; j++)
+			{
+				IParamBlock2* ParamBlock2 = Material.GetParamBlockByID((short)j);
+				// The the descriptor to 'decode'
+				ParamBlockDesc2* ParamBlockDesc = ParamBlock2->GetDesc();
+				// Loop through all the defined parameters therein
+				for (int i = 0; i < ParamBlockDesc->count; i++)
+				{
+					const ParamDef& ParamDefinition = ParamBlockDesc->paramdefs[i];
+
+					if (FCString::Stricmp(ParamDefinition.int_name, TEXT("texmap")) == 0)
+					{
+						EmitTexture.Map = ParamBlock2->GetTexmap(ParamDefinition.ID, GetCOREInterface()->GetTime());
+					}
+					else if (FCString::Stricmp(ParamDefinition.int_name, TEXT("opacity_texmap")) == 0)
+					{
+						ClipTexture.Map = ParamBlock2->GetTexmap(ParamDefinition.ID, GetCOREInterface()->GetTime());
+					}
+					else if (FCString::Stricmp(ParamDefinition.int_name, TEXT("texmap_on")) == 0)
+					{
+						if (ParamBlock2->GetInt(ParamDefinition.ID, GetCOREInterface()->GetTime()) == 0)
+						{
+							EmitTexture.bEnabled = false;
+						}
+					}
+					else if (FCString::Stricmp(ParamDefinition.int_name, TEXT("opacity_texmap_on")) == 0)
+					{
+						if (ParamBlock2->GetInt(ParamDefinition.ID, GetCOREInterface()->GetTime()) == 0)
+						{
+							ClipTexture.bEnabled = false;
+						}
+					}
+					else if (FCString::Stricmp(ParamDefinition.int_name, TEXT("Color")) == 0)
+					{
+						EmitColor = (BMM_Color_fl)ParamBlock2->GetColor(ParamDefinition.ID, GetCOREInterface()->GetTime());
+					}
+					else if (FCString::Stricmp(ParamDefinition.int_name, TEXT("twoSided")) == 0)
+					{
+						// int twoSided = ParamBlock2->GetInt(ParamDefinition.ID, GetCOREInterface()->GetTime());
+					}
+					else if (FCString::Stricmp(ParamDefinition.int_name, TEXT("Multiplier")) == 0)
+					{
+						Multiplier = ParamBlock2->GetFloat(ParamDefinition.ID, GetCOREInterface()->GetTime());
+					}
+				}
+				ParamBlock2->ReleaseDesc();
+			}
+		}
+
+	};
+
 }
 
 FDatasmithMaxVRayMaterialsToUEPbr::FDatasmithMaxVRayMaterialsToUEPbr()
@@ -1029,3 +1093,42 @@ void FDatasmithMaxVRayBlendMaterialToUEPbr::Convert( TSharedRef<IDatasmithScene>
 	PreviousExpression->ConnectExpression(PbrMaterialElement->GetMaterialAttributes());
 	MaterialElement = PbrMaterialElement;
 }
+
+bool FDatasmithMaxVRayLightMaterialToUEPbr::IsSupported(Mtl* Material)
+{
+	return true;
+}
+
+void FDatasmithMaxVRayLightMaterialToUEPbr::Convert(TSharedRef<IDatasmithScene> DatasmithScene,
+	TSharedPtr<IDatasmithBaseMaterialElement>& MaterialElement, Mtl* Material, const TCHAR* AssetsPath)
+{
+	if ( !Material )
+	{
+		return;
+	}
+
+	TSharedRef<IDatasmithUEPbrMaterialElement> PbrMaterialElement = FDatasmithSceneFactory::CreateUEPbrMaterial(Material->GetName().data());
+	FScopedConvertState ScopedConvertState(ConvertState);
+	ConvertState.DatasmithScene = DatasmithScene;
+	ConvertState.MaterialElement = PbrMaterialElement;
+	ConvertState.AssetsPath = AssetsPath;
+
+	DatasmithMaxVRayMaterialsToUEPbrImpl::FVRayLightMaterial MaterialProperties;
+	MaterialProperties.Parse(*Material);
+
+	Connect(PbrMaterialElement->GetEmissiveColor(),
+		COMPOSE_OR_DEFAULT2(nullptr, Multiply,
+			TextureOrColor(TEXT("Emissive Color"), MaterialProperties.EmitTexture, FDatasmithMaxMatHelper::MaxLinearColorToFLinearColor(MaterialProperties.EmitColor)),
+			&Scalar(MaterialProperties.Multiplier))
+	);
+
+	if (IDatasmithMaterialExpression* OpacictyExpression = ConvertTexmap(MaterialProperties.ClipTexture))
+	{
+		PbrMaterialElement->SetBlendMode(/*EBlendMode::BLEND_Masked*/1); // Set blend mode when there's opacity mask
+		Connect(PbrMaterialElement->GetOpacity(), OpacictyExpression);
+	}
+
+	PbrMaterialElement->SetShadingModel(EDatasmithShadingModel::Unlit);
+
+	MaterialElement = PbrMaterialElement;
+}
diff --git a/Engine/Source/Programs/Enterprise/Datasmith/DatasmithMaxExporter/Private/MaxMaterialsToUEPbr/DatasmithMaxVrayMaterialsToUEPbr.h b/Engine/Source/Programs/Enterprise/Datasmith/DatasmithMaxExporter/Private/MaxMaterialsToUEPbr/DatasmithMaxVrayMaterialsToUEPbr.h
--- a/Engine/Source/Programs/Enterprise/Datasmith/DatasmithMaxExporter/Private/MaxMaterialsToUEPbr/DatasmithMaxVrayMaterialsToUEPbr.h
+++ b/Engine/Source/Programs/Enterprise/Datasmith/DatasmithMaxExporter/Private/MaxMaterialsToUEPbr/DatasmithMaxVrayMaterialsToUEPbr.h
@@ -33,3 +33,10 @@ public:
 	virtual bool IsSupported( Mtl* Material ) override;
 	virtual void Convert( TSharedRef< IDatasmithScene > DatasmithScene, TSharedPtr< IDatasmithBaseMaterialElement >& MaterialElement, Mtl* Material, const TCHAR* AssetsPath ) override;
 };
+
+class FDatasmithMaxVRayLightMaterialToUEPbr : public FDatasmithMaxMaterialsToUEPbrExpressions
+{
+public:
+	virtual bool IsSupported( Mtl* Material ) override;
+	virtual void Convert( TSharedRef< IDatasmithScene > DatasmithScene, TSharedPtr< IDatasmithBaseMaterialElement >& MaterialElement, Mtl* Material, const TCHAR* AssetsPath ) override;
+};
diff --git a/Engine/Source/Programs/Enterprise/Datasmith/DatasmithMaxExporter/Public/DatasmithMaxWriter.h b/Engine/Source/Programs/Enterprise/Datasmith/DatasmithMaxExporter/Public/DatasmithMaxWriter.h
--- a/Engine/Source/Programs/Enterprise/Datasmith/DatasmithMaxExporter/Public/DatasmithMaxWriter.h
+++ b/Engine/Source/Programs/Enterprise/Datasmith/DatasmithMaxExporter/Public/DatasmithMaxWriter.h
@@ -80,7 +80,6 @@ public:
 	// CLASS RELATED FUNCTIONS
 	static EDSBitmapType GetTextureClass(Texmap* mTexMap);
 	static EDSMaterialType GetMaterialClass(Mtl* Material);
-	static BMM_Color_fl TemperatureToColor(float Kelvin);
 	/**
 	 * Gets a XRefMaterial and returns the rendered material which is either the SubMaterial or SourceMaterial depending on if the source material is overridden.
 	 * @param XRefMaterial	The XRefMaterial we want the rendered material of.
diff --git a/Engine/Source/Runtime/Datasmith/DatasmithCore/Private/DatasmithMaterialsUtils.cpp b/Engine/Source/Runtime/Datasmith/DatasmithCore/Private/DatasmithMaterialsUtils.cpp
--- a/Engine/Source/Runtime/Datasmith/DatasmithCore/Private/DatasmithMaterialsUtils.cpp
+++ b/Engine/Source/Runtime/Datasmith/DatasmithCore/Private/DatasmithMaterialsUtils.cpp
@@ -199,3 +199,43 @@ IDatasmithMaterialExpression* DatasmithMaterialsUtils::CreateWeightedMaterialExp
 
 	return Expression;
 }
+
+FLinearColor DatasmithMaterialsUtils::TemperatureToColor(float Kelvin)
+{
+	float Temp = Kelvin / 100.0f;
+	float Red;
+	float Green;
+	float Blue;
+
+	if (Temp <= 66.0f)
+	{
+		Red = 255.0f;
+		Green = Temp;
+		Green = 99.4708025861f * FMath::Loge(Green) - 161.1195681661f;
+		if (Temp <= 19.0f)
+		{
+			Blue = 0.0f;
+		}
+		else
+		{
+			Blue = Temp - 10.0f;
+			Blue = 138.5177312231f * FMath::Loge(Blue) - 305.0447927307f;
+		}
+	}
+	else
+	{
+		Red = Temp - 60.0f;
+		Red = 329.698727446f * FMath::Pow(Red, -0.1332047592f);
+		Green = Temp - 60.f;
+		Green = 288.1221695283f * FMath::Pow(Green, -0.0755148492f);
+		Blue = 255.0f;
+	}
+
+	FLinearColor Color;
+	Color.R = Red / 255.0f;
+	Color.G = Green / 255.0f;
+	Color.B = Blue / 255.0f;
+
+	return Color;
+}
+
diff --git a/Engine/Source/Runtime/Datasmith/DatasmithCore/Public/DatasmithMaterialsUtils.h b/Engine/Source/Runtime/Datasmith/DatasmithCore/Public/DatasmithMaterialsUtils.h
--- a/Engine/Source/Runtime/Datasmith/DatasmithCore/Public/DatasmithMaterialsUtils.h
+++ b/Engine/Source/Runtime/Datasmith/DatasmithCore/Public/DatasmithMaterialsUtils.h
@@ -68,4 +68,6 @@ namespace DatasmithMaterialsUtils
 	 * @return						Material expression that outputs a Color or Scalar interpolated with Expression; nullptr if no Color, Scalar or Expression are given
 	 */
 	DATASMITHCORE_API IDatasmithMaterialExpression* CreateWeightedMaterialExpression( const TSharedRef< IDatasmithUEPbrMaterialElement >& MaterialElement, const TCHAR* ParameterName, TOptional< FLinearColor > Color, TOptional< float > Scalar, IDatasmithMaterialExpression* Expression, float Weight, EDatasmithTextureMode TextureMode = EDatasmithTextureMode::Diffuse );
+
+	DATASMITHCORE_API FLinearColor TemperatureToColor(float Kelvin);
 }
