3ds Max DirectLink: XRefMat and VRayBlendMtl to UEPbr

- also simplified conversion code(removed redundand checks)

diff --git a/Engine/Source/Programs/Enterprise/Datasmith/DatasmithMaxExporter/Private/MaxMaterialsToUEPbr/DatasmithMaxCoronaMaterialsToUEPbr.cpp b/Engine/Source/Programs/Enterprise/Datasmith/DatasmithMaxExporter/Private/MaxMaterialsToUEPbr/DatasmithMaxCoronaMaterialsToUEPbr.cpp
--- a/Engine/Source/Programs/Enterprise/Datasmith/DatasmithMaxExporter/Private/MaxMaterialsToUEPbr/DatasmithMaxCoronaMaterialsToUEPbr.cpp
+++ b/Engine/Source/Programs/Enterprise/Datasmith/DatasmithMaxExporter/Private/MaxMaterialsToUEPbr/DatasmithMaxCoronaMaterialsToUEPbr.cpp
@@ -687,7 +687,7 @@ bool FDatasmithMaxCoronaBlendMaterialToUEPbr::IsSupported( Mtl* Material )
 	if (CoronaBlendMaterialProperties.BaseMaterial)
 	{
 		FDatasmithMaxMaterialsToUEPbr* MaterialConverter = FDatasmithMaxMaterialsToUEPbrManager::GetMaterialConverter(CoronaBlendMaterialProperties.BaseMaterial);
-		bAllMaterialsSupported &= MaterialConverter && MaterialConverter->IsSupported(CoronaBlendMaterialProperties.BaseMaterial);
+		bAllMaterialsSupported &= MaterialConverter != nullptr;
 	}
 	else
 	{
@@ -702,7 +702,7 @@ bool FDatasmithMaxCoronaBlendMaterialToUEPbr::IsSupported( Mtl* Material )
 		{
 			//Only support if all the blended materials are UEPbr materials.
 			FDatasmithMaxMaterialsToUEPbr* MaterialConverter = FDatasmithMaxMaterialsToUEPbrManager::GetMaterialConverter(CoatedMaterial.Material);
-			bAllMaterialsSupported &= MaterialConverter && MaterialConverter->IsSupported(CoatedMaterial.Material);
+			bAllMaterialsSupported &= MaterialConverter != nullptr;
 		}
 	}
 
diff --git a/Engine/Source/Programs/Enterprise/Datasmith/DatasmithMaxExporter/Private/MaxMaterialsToUEPbr/DatasmithMaxMaterialsToUEPbr.cpp b/Engine/Source/Programs/Enterprise/Datasmith/DatasmithMaxExporter/Private/MaxMaterialsToUEPbr/DatasmithMaxMaterialsToUEPbr.cpp
--- a/Engine/Source/Programs/Enterprise/Datasmith/DatasmithMaxExporter/Private/MaxMaterialsToUEPbr/DatasmithMaxMaterialsToUEPbr.cpp
+++ b/Engine/Source/Programs/Enterprise/Datasmith/DatasmithMaxExporter/Private/MaxMaterialsToUEPbr/DatasmithMaxMaterialsToUEPbr.cpp
@@ -86,15 +86,16 @@ FDatasmithMaxMaterialsToUEPbr* FDatasmithMaxMaterialsToUEPbrManager::GetMaterial
 		static FDatasmithMaxMentalMaterialToUEPbr PhysicalConverter = FDatasmithMaxMentalMaterialToUEPbr();
 		MaterialConverter = &PhysicalConverter;
 	}
+	else if (MaterialClassID == XREFMATCLASS)
+	{
+		MaterialConverter = GetMaterialConverter(FDatasmithMaxMatHelper::GetRenderedXRefMaterial(Material));
+	}
 
 	if ( MaterialConverter && MaterialConverter->IsSupported( Material ) )
 	{
 		return MaterialConverter;
 	}
-	else
-	{
-		return nullptr;
-	}
+	return nullptr;
 }
 
 FDatasmithMaxMaterialsToUEPbr::FDatasmithMaxMaterialsToUEPbr()
diff --git a/Engine/Source/Programs/Enterprise/Datasmith/DatasmithMaxExporter/Private/MaxMaterialsToUEPbr/DatasmithMaxMaterialsToUEPbr.h b/Engine/Source/Programs/Enterprise/Datasmith/DatasmithMaxExporter/Private/MaxMaterialsToUEPbr/DatasmithMaxMaterialsToUEPbr.h
--- a/Engine/Source/Programs/Enterprise/Datasmith/DatasmithMaxExporter/Private/MaxMaterialsToUEPbr/DatasmithMaxMaterialsToUEPbr.h
+++ b/Engine/Source/Programs/Enterprise/Datasmith/DatasmithMaxExporter/Private/MaxMaterialsToUEPbr/DatasmithMaxMaterialsToUEPbr.h
@@ -35,9 +35,9 @@ public:
 	FDatasmithMaxMaterialsToUEPbr();
 	virtual ~FDatasmithMaxMaterialsToUEPbr() = default;
 
-	virtual bool IsSupported( Mtl* Material ) = 0;
 	virtual void Convert( TSharedRef< IDatasmithScene > DatasmithScene, TSharedPtr< IDatasmithBaseMaterialElement >& MaterialElement, Mtl* Material, const TCHAR* AssetsPath ) = 0;
 
+	virtual bool IsSupported( Mtl* Material ) = 0; // Called by FDatasmithMaxMaterialsToUEPbrManager to see if instantiated converter can actually convert a material instance
 	virtual bool IsTexmapSupported( Texmap* InTexmap ) const;
 
 	/**
diff --git a/Engine/Source/Programs/Enterprise/Datasmith/DatasmithMaxExporter/Private/MaxMaterialsToUEPbr/DatasmithMaxScanlineMaterialsToUEPbr.cpp b/Engine/Source/Programs/Enterprise/Datasmith/DatasmithMaxExporter/Private/MaxMaterialsToUEPbr/DatasmithMaxScanlineMaterialsToUEPbr.cpp
--- a/Engine/Source/Programs/Enterprise/Datasmith/DatasmithMaxExporter/Private/MaxMaterialsToUEPbr/DatasmithMaxScanlineMaterialsToUEPbr.cpp
+++ b/Engine/Source/Programs/Enterprise/Datasmith/DatasmithMaxExporter/Private/MaxMaterialsToUEPbr/DatasmithMaxScanlineMaterialsToUEPbr.cpp
@@ -295,13 +295,13 @@ bool FDatasmithMaxBlendMaterialsToUEPbr::IsSupported( Mtl* Material )
 	if ( Mtl* BaseMaterial = Material->GetSubMtl(0) )
 	{
 		FDatasmithMaxMaterialsToUEPbr* MaterialConverter = FDatasmithMaxMaterialsToUEPbrManager::GetMaterialConverter( BaseMaterial );
-		bAllMaterialsSupported &= MaterialConverter && MaterialConverter->IsSupported( BaseMaterial );
+		bAllMaterialsSupported &= MaterialConverter != nullptr;
 	}
 
 	if ( Mtl* CoatMaterial = Material->GetSubMtl(1) )
 	{
 		FDatasmithMaxMaterialsToUEPbr* MaterialConverter = FDatasmithMaxMaterialsToUEPbrManager::GetMaterialConverter( CoatMaterial );
-		bAllMaterialsSupported &= MaterialConverter && MaterialConverter->IsSupported( CoatMaterial );
+		bAllMaterialsSupported &= MaterialConverter != nullptr;
 	}
 
 	return bAllMaterialsSupported;
diff --git a/Engine/Source/Programs/Enterprise/Datasmith/DatasmithMaxExporter/Private/MaxMaterialsToUEPbr/DatasmithMaxVrayMaterialsToUEPbr.cpp b/Engine/Source/Programs/Enterprise/Datasmith/DatasmithMaxExporter/Private/MaxMaterialsToUEPbr/DatasmithMaxVrayMaterialsToUEPbr.cpp
--- a/Engine/Source/Programs/Enterprise/Datasmith/DatasmithMaxExporter/Private/MaxMaterialsToUEPbr/DatasmithMaxVrayMaterialsToUEPbr.cpp
+++ b/Engine/Source/Programs/Enterprise/Datasmith/DatasmithMaxExporter/Private/MaxMaterialsToUEPbr/DatasmithMaxVrayMaterialsToUEPbr.cpp
@@ -920,12 +920,7 @@ bool FDatasmithMaxVRay2SidedMaterialsToUEPbr::IsSupported( Mtl* Material )
 
 	DatasmithMaxVRayMaterialsToUEPbrImpl::FMaxVRay2SidedMaterial VRay2SidedMaterialProperties = DatasmithMaxVRayMaterialsToUEPbrImpl::ParseVRay2SidedMaterialProperties( *Material );
 
-	if ( FDatasmithMaxMaterialsToUEPbr* MaterialConverter = FDatasmithMaxMaterialsToUEPbrManager::GetMaterialConverter( VRay2SidedMaterialProperties.FrontMaterial ) )
-	{
-		return MaterialConverter->IsSupported( VRay2SidedMaterialProperties.FrontMaterial );
-	}
-
-	return false;
+	return FDatasmithMaxMaterialsToUEPbrManager::GetMaterialConverter( VRay2SidedMaterialProperties.FrontMaterial ) != nullptr;
 }
 
 void FDatasmithMaxVRay2SidedMaterialsToUEPbr::Convert( TSharedRef<IDatasmithScene> DatasmithScene, TSharedPtr<IDatasmithBaseMaterialElement>& MaterialElement, Mtl* Material, const TCHAR* AssetsPath )
@@ -965,13 +960,7 @@ bool FDatasmithMaxVRayWrapperMaterialsToUEPbr::IsSupported( Mtl* Material )
 	}
 
 	DatasmithMaxVRayMaterialsToUEPbrImpl::FMaxVRayWrapperMaterial VRayWrapperMaterialProperties = DatasmithMaxVRayMaterialsToUEPbrImpl::ParseVRayWrapperMaterialProperties( *Material );
-
-	if ( FDatasmithMaxMaterialsToUEPbr* MaterialConverter = FDatasmithMaxMaterialsToUEPbrManager::GetMaterialConverter( VRayWrapperMaterialProperties.BaseMaterial ) )
-	{
-		return MaterialConverter->IsSupported( VRayWrapperMaterialProperties.BaseMaterial );
-	}
-
-	return false;
+	return FDatasmithMaxMaterialsToUEPbrManager::GetMaterialConverter( VRayWrapperMaterialProperties.BaseMaterial ) != nullptr;
 }
 
 void FDatasmithMaxVRayWrapperMaterialsToUEPbr::Convert( TSharedRef<IDatasmithScene> DatasmithScene, TSharedPtr<IDatasmithBaseMaterialElement>& MaterialElement, Mtl* Material, const TCHAR* AssetsPath )
@@ -1012,11 +1001,7 @@ bool FDatasmithMaxVRayBlendMaterialToUEPbr::IsSupported( Mtl* Material )
 	if (VRayBlendMaterialProperties.BaseMaterial)
 	{
 		FDatasmithMaxMaterialsToUEPbr* MaterialConverter = FDatasmithMaxMaterialsToUEPbrManager::GetMaterialConverter(VRayBlendMaterialProperties.BaseMaterial);
-		bAllMaterialsSupported &= MaterialConverter && MaterialConverter->IsSupported(VRayBlendMaterialProperties.BaseMaterial);
-	}
-	else
-	{
-		return false;
+		bAllMaterialsSupported &= MaterialConverter != nullptr;
 	}
 
 	for (int CoatIndex = 0; bAllMaterialsSupported && CoatIndex < FMaxVRayBlendMaterial::MaximumNumberOfCoat; ++CoatIndex)
@@ -1027,7 +1012,7 @@ bool FDatasmithMaxVRayBlendMaterialToUEPbr::IsSupported( Mtl* Material )
 		{
 			//Only support if all the blended materials are UEPbr materials.
 			FDatasmithMaxMaterialsToUEPbr* MaterialConverter = FDatasmithMaxMaterialsToUEPbrManager::GetMaterialConverter(CoatedMaterial.Material);
-			bAllMaterialsSupported &= MaterialConverter && MaterialConverter->IsSupported(CoatedMaterial.Material);
+			bAllMaterialsSupported &= MaterialConverter != nullptr;
 		}
 	}
 
@@ -1046,51 +1031,72 @@ void FDatasmithMaxVRayBlendMaterialToUEPbr::Convert( TSharedRef<IDatasmithScene>
 
 	FMaxVRayBlendMaterial VRayBlendMaterialProperties = ParseVRayBlendMaterialProperties( *Material );
 
+	IDatasmithMaterialExpression* PreviousExpression = nullptr;
+
 	//Exporting the base material.
-	IDatasmithMaterialExpressionFunctionCall* BaseMaterialFunctionCall = PbrMaterialElement->AddMaterialExpression<IDatasmithMaterialExpressionFunctionCall>();
-	if (TSharedPtr<IDatasmithBaseMaterialElement> ExportedMaterial = FDatasmithMaxMatExport::ExportUniqueMaterial(DatasmithScene, VRayBlendMaterialProperties.BaseMaterial, AssetsPath))
+	if (VRayBlendMaterialProperties.BaseMaterial) 
 	{
-		BaseMaterialFunctionCall->SetFunctionPathName(ExportedMaterial->GetName());
+		IDatasmithMaterialExpressionFunctionCall* BaseMaterialFunctionCall = PbrMaterialElement->AddMaterialExpression<IDatasmithMaterialExpressionFunctionCall>();
+		if (TSharedPtr<IDatasmithBaseMaterialElement> ExportedMaterial = FDatasmithMaxMatExport::ExportUniqueMaterial(DatasmithScene, VRayBlendMaterialProperties.BaseMaterial, AssetsPath))
+		{
+			BaseMaterialFunctionCall->SetFunctionPathName(ExportedMaterial->GetName());
+		}
+		PreviousExpression = BaseMaterialFunctionCall;
 	}
 
 	//Exporting the blended materials.
-	IDatasmithMaterialExpression* PreviousExpression = BaseMaterialFunctionCall;
 	for (int CoatIndex = 0; CoatIndex < FMaxVRayBlendMaterial::MaximumNumberOfCoat; ++CoatIndex)
 	{
 		const FMaxVRayBlendMaterial::FVRayCoatMaterialProperties& CoatedMaterial = VRayBlendMaterialProperties.CoatedMaterials[CoatIndex];
 
 		if (CoatedMaterial.Material != nullptr && CoatedMaterial.MaterialBlendParameter.bEnabled)
 		{
-			IDatasmithMaterialExpressionFunctionCall* BlendFunctionCall = PbrMaterialElement->AddMaterialExpression<IDatasmithMaterialExpressionFunctionCall>();
-			BlendFunctionCall->SetFunctionPathName(TEXT("/Engine/Functions/MaterialLayerFunctions/MatLayerBlend_Standard.MatLayerBlend_Standard"));
-			PreviousExpression->ConnectExpression(*BlendFunctionCall->GetInput(0));
-			PreviousExpression = BlendFunctionCall;
-
-			IDatasmithMaterialExpressionFunctionCall* CoatedMaterialFunctionCall = PbrMaterialElement->AddMaterialExpression<IDatasmithMaterialExpressionFunctionCall>();
-			if (TSharedPtr<IDatasmithBaseMaterialElement> ExportedMaterial = FDatasmithMaxMatExport::ExportUniqueMaterial(DatasmithScene, CoatedMaterial.Material, AssetsPath))
+			TSharedPtr<IDatasmithBaseMaterialElement> ExportedCoatedMaterial = FDatasmithMaxMatExport::ExportUniqueMaterial(DatasmithScene, CoatedMaterial.Material, AssetsPath);
+			if (PreviousExpression)
 			{
-				CoatedMaterialFunctionCall->SetFunctionPathName(ExportedMaterial->GetName());
-			}
-			CoatedMaterialFunctionCall->ConnectExpression(*BlendFunctionCall->GetInput(1));
+				IDatasmithMaterialExpressionFunctionCall* BlendFunctionCall = PbrMaterialElement->AddMaterialExpression<IDatasmithMaterialExpressionFunctionCall>();
+				BlendFunctionCall->SetFunctionPathName(TEXT("/Engine/Functions/MaterialLayerFunctions/MatLayerBlend_Standard.MatLayerBlend_Standard"));
+				PreviousExpression->ConnectExpression(*BlendFunctionCall->GetInput(0));
+				PreviousExpression = BlendFunctionCall;
+
+				IDatasmithMaterialExpressionFunctionCall* CoatedMaterialFunctionCall = PbrMaterialElement->AddMaterialExpression<IDatasmithMaterialExpressionFunctionCall>();
+				if (ExportedCoatedMaterial)
+				{
+					CoatedMaterialFunctionCall->SetFunctionPathName(ExportedCoatedMaterial->GetName());
+				}
+				CoatedMaterialFunctionCall->ConnectExpression(*BlendFunctionCall->GetInput(1));
 
-			IDatasmithMaterialExpression* AlphaExpression = FDatasmithMaxTexmapToUEPbrUtils::MapOrValue(this, CoatedMaterial.MaterialBlendParameter, TEXT("MixAmount"),
-				CoatedMaterial.MixColor, TOptional< float >());
+				IDatasmithMaterialExpression* AlphaExpression = FDatasmithMaxTexmapToUEPbrUtils::MapOrValue(this, CoatedMaterial.MaterialBlendParameter, TEXT("MixAmount"),
+					CoatedMaterial.MixColor, TOptional< float >());
+
+				//AlphaExpression is nullptr only when there is no mask and the mask weight is ~100% so we add scalar 0 instead.
+				if(!AlphaExpression)
+				{
+					IDatasmithMaterialExpressionScalar* WeightExpression = PbrMaterialElement->AddMaterialExpression< IDatasmithMaterialExpressionScalar >();
+					WeightExpression->SetName(TEXT("MixAmount"));
+					WeightExpression->GetScalar() = 0.f;
+					AlphaExpression = WeightExpression;
+				}
 
-			//AlphaExpression is nullptr only when there is no mask and the mask weight is ~100% so we add scalar 0 instead.
-			if(!AlphaExpression)
+				AlphaExpression->ConnectExpression(*BlendFunctionCall->GetInput(2));
+			}
+			else
 			{
-				IDatasmithMaterialExpressionScalar* WeightExpression = PbrMaterialElement->AddMaterialExpression< IDatasmithMaterialExpressionScalar >();
-				WeightExpression->SetName(TEXT("MixAmount"));
-				WeightExpression->GetScalar() = 0.f;
-				AlphaExpression = WeightExpression;
+				IDatasmithMaterialExpressionFunctionCall* CoatedMaterialFunctionCall = PbrMaterialElement->AddMaterialExpression<IDatasmithMaterialExpressionFunctionCall>();
+				if (ExportedCoatedMaterial)
+				{
+					CoatedMaterialFunctionCall->SetFunctionPathName(ExportedCoatedMaterial->GetName());
+				}
+				PreviousExpression = CoatedMaterialFunctionCall;
 			}
-
-			AlphaExpression->ConnectExpression(*BlendFunctionCall->GetInput(2));
 		}
 	}
 
 	PbrMaterialElement->SetUseMaterialAttributes(true);
-	PreviousExpression->ConnectExpression(PbrMaterialElement->GetMaterialAttributes());
+	if (PreviousExpression)
+	{
+		PreviousExpression->ConnectExpression(PbrMaterialElement->GetMaterialAttributes());
+	}
 	MaterialElement = PbrMaterialElement;
 }
 
