Solidworks - Deleted object in Solidworks is not removed from UE

diff --git a/Engine/Source/Programs/Enterprise/Datasmith/DatasmithSolidworksExporter/Private/DatasmithExporter.cs b/Engine/Source/Programs/Enterprise/Datasmith/DatasmithSolidworksExporter/Private/DatasmithExporter.cs
--- a/Engine/Source/Programs/Enterprise/Datasmith/DatasmithSolidworksExporter/Private/DatasmithExporter.cs
+++ b/Engine/Source/Programs/Enterprise/Datasmith/DatasmithSolidworksExporter/Private/DatasmithExporter.cs
@@ -21,6 +21,7 @@ namespace DatasmithSolidworks
 		public string Label;
 		public string Name;
 		public string ParentName;
+		public string MeshName;
 		public float[] Transform;
 		public bool bVisible;
 	};
@@ -98,6 +99,12 @@ namespace DatasmithSolidworks
 			Actor.SetLabel(InExportInfo.Label);
 			Actor.SetVisibility(InExportInfo.bVisible);
 			Actor.SetWorldTransform(AdjustTransformForDatasmith(InExportInfo.Transform));
+
+			if (InExportInfo.Type == EActorType.MeshActor)
+			{
+				FDatasmithFacadeActorMesh MeshActor = Actor as FDatasmithFacadeActorMesh;
+				MeshActor.SetMesh(InExportInfo.MeshName);
+			}
 		}
 
 		public void RemoveActor(string InActorName)
@@ -111,18 +118,18 @@ namespace DatasmithSolidworks
 			}
 		}
 
-		public void ExportMesh(string InMeshName, FMeshData InData, string InUpdateMeshActor, out Tuple<FDatasmithFacadeMeshElement, FDatasmithFacadeMesh> OutMeshPair)
+		public string ExportMesh(string InMeshName, FMeshData InData, string InUpdateMeshActor, out Tuple<FDatasmithFacadeMeshElement, FDatasmithFacadeMesh> OutMeshPair)
 		{
 			OutMeshPair = null;
 
 			if (InData.Vertices == null || InData.Normals == null || InData.TexCoords == null || InData.Triangles == null)
 			{
-				return;
+				return null;
 			}
 
 			if (InData.Vertices.Length == 0 || InData.Normals.Length == 0 || InData.TexCoords.Length == 0 || InData.Triangles.Length == 0)
 			{
-				return;
+				return null;
 			}
 
 			InMeshName = SanitizeName(InMeshName);
@@ -201,6 +208,8 @@ namespace DatasmithSolidworks
 					MeshActor.SetMesh(InMeshName);
 				}
 			}
+
+			return InMeshName;
 		}
 
 		public void ExportMetadata(FMetadata InMetadata)
diff --git a/Engine/Source/Programs/Enterprise/Datasmith/DatasmithSolidworksExporter/Private/PartDocument.cs b/Engine/Source/Programs/Enterprise/Datasmith/DatasmithSolidworksExporter/Private/PartDocument.cs
--- a/Engine/Source/Programs/Enterprise/Datasmith/DatasmithSolidworksExporter/Private/PartDocument.cs
+++ b/Engine/Source/Programs/Enterprise/Datasmith/DatasmithSolidworksExporter/Private/PartDocument.cs
@@ -80,35 +80,44 @@ namespace DatasmithSolidworks
 
 		public override void ExportToDatasmithScene()
 		{
-			FDatasmithActorExportInfo ExportInfo = new FDatasmithActorExportInfo();
-			ExportInfo.Name = Path.GetFileNameWithoutExtension(PathName);
-			ExportInfo.bVisible = true;
-			ExportInfo.Label = ExportInfo.Name;
-			ExportInfo.Type = EActorType.MeshActor;
-
-			Exporter.ExportOrUpdateActor(ExportInfo);
+			string ActorName = Path.GetFileNameWithoutExtension(PathName);
 
-			SetExportStatus($"{ExportInfo.Name} Materials");
+			SetExportStatus($"{ActorName} Materials");
 
 			ExportedPartMaterials = FObjectMaterials.LoadPartMaterials(this, SwPartDoc, swDisplayStateOpts_e.swThisDisplayState, null);
 
 			Exporter.ExportMaterials(ExportedMaterialsMap);
 
-			SetExportStatus($"{ExportInfo.Name} Meshes");
+			SetExportStatus($"{ActorName} Meshes");
 
 			ConcurrentBag<FBody> Bodies = FBody.FetchBodies(SwPartDoc);
 			FMeshData MeshData = FStripGeometry.CreateMeshData(Bodies, ExportedPartMaterials);
 
+			string MeshName = null;
+
 			if (MeshData != null)
 			{
 				Tuple<FDatasmithFacadeMeshElement, FDatasmithFacadeMesh> NewMesh = null;
-				Exporter.ExportMesh($"{ExportInfo.Name}_Mesh", MeshData, ExportInfo.Name, out NewMesh);
+				MeshName = Exporter.ExportMesh($"{ActorName}_Mesh", MeshData, ActorName, out NewMesh);
 
 				if (NewMesh != null)
 				{
 					DatasmithScene.AddMesh(NewMesh.Item1);
 				}
+				else
+				{
+					MeshName = null;
+				}
 			}
+
+			FDatasmithActorExportInfo ExportInfo = new FDatasmithActorExportInfo();
+			ExportInfo.Name = ActorName;
+			ExportInfo.bVisible = true;
+			ExportInfo.Label = ActorName;
+			ExportInfo.MeshName = MeshName;
+			ExportInfo.Type = MeshName != null ? EActorType.MeshActor : EActorType.SimpleActor;
+
+			Exporter.ExportOrUpdateActor(ExportInfo);
 		}
 
 		public override bool HasMaterialUpdates()
