SketchUp: 2022 support

- Win64 and macOS(plugin and installer)

diff --git a/Engine/Source/Programs/Enterprise/Datasmith/DatasmithSketchUpRubyExporter/Build/build_mac.sh b/Engine/Source/Programs/Enterprise/Datasmith/DatasmithSketchUpRubyExporter/Build/build_mac.sh
--- a/Engine/Source/Programs/Enterprise/Datasmith/DatasmithSketchUpRubyExporter/Build/build_mac.sh
+++ b/Engine/Source/Programs/Enterprise/Datasmith/DatasmithSketchUpRubyExporter/Build/build_mac.sh
@@ -78,6 +78,7 @@ BuildSketchUpPlugin() {
 BuildSketchUpPlugin 2019 SDK_Mac_2019-3-252
 BuildSketchUpPlugin 2020 SDK_Mac_2020-2-171
 BuildSketchUpPlugin 2021 SDK_Mac_2021-0-338
+BuildSketchUpPlugin 2022 SDK_Mac_2022-0-315
 
 # install_name_tool -change @rpath/DatasmithSDK.dylib @loader_path/Dylibs/DatasmithSDK.dylib DatasmithSketchUp.bundle 
 # install_name_tool -change @loader_path/libfreeimage-3.18.0.dylib @loader_path/Dylibs/libfreeimage-3.18.0.dylib DatasmithSketchUp.bundle
diff --git a/Engine/Source/Programs/Enterprise/Datasmith/DatasmithSketchUpRubyExporter/DatasmithSketchUpRuby2022.Build.cs b/Engine/Source/Programs/Enterprise/Datasmith/DatasmithSketchUpRubyExporter/DatasmithSketchUpRuby2022.Build.cs
new file mode 100644
--- /dev/null
+++ b/Engine/Source/Programs/Enterprise/Datasmith/DatasmithSketchUpRubyExporter/DatasmithSketchUpRuby2022.Build.cs
@@ -0,0 +1,31 @@
+// Copyright Epic Games, Inc. All Rights Reserved.
+
+using System.IO;
+
+namespace UnrealBuildTool.Rules
+{
+
+	[SupportedPlatforms("Win64", "Mac")]
+	public class DatasmithSketchUpRuby2022 : DatasmithSketchUpRubyBase
+	{
+		public DatasmithSketchUpRuby2022(ReadOnlyTargetRules Target)
+			: base(Target)
+		{
+			PrivateDefinitions.Add("SKP_SDK_2022");
+		}
+
+		public override string GetSketchUpSDKFolder()
+		{
+			return "SDK_WIN_x64_2022-0-316";
+		}
+
+		public override string GetSketchUpEnvVar()
+		{
+			return "SKP_SDK_2022";
+		}
+		public override string GetRubyLibName()
+		{
+			return "x64-msvcrt-ruby270.lib";
+		}
+	}
+}
diff --git a/Engine/Source/Programs/Enterprise/Datasmith/DatasmithSketchUpRubyExporter/DatasmithSketchUpRuby2022.Target.cs b/Engine/Source/Programs/Enterprise/Datasmith/DatasmithSketchUpRubyExporter/DatasmithSketchUpRuby2022.Target.cs
new file mode 100644
--- /dev/null
+++ b/Engine/Source/Programs/Enterprise/Datasmith/DatasmithSketchUpRubyExporter/DatasmithSketchUpRuby2022.Target.cs
@@ -0,0 +1,16 @@
+// Copyright Epic Games, Inc. All Rights Reserved.
+
+using UnrealBuildTool;
+
+[SupportedPlatforms("Win64")]
+public class DatasmithSketchUpRuby2022Target : DatasmithSketchUpRubyBaseTarget
+{
+	public DatasmithSketchUpRuby2022Target(TargetInfo Target)
+		: base(Target)
+	{
+		LaunchModuleName = "DatasmithSketchUpRuby2022";
+		ExeBinariesSubFolder = @"SketchUpRuby/2022";
+
+		AddCopyPostBuildStep(Target);
+	}
+}
diff --git a/Engine/Source/Programs/Enterprise/Datasmith/DatasmithSketchUpRubyExporter/Installer/Mac/SignAndBuild.sh b/Engine/Source/Programs/Enterprise/Datasmith/DatasmithSketchUpRubyExporter/Installer/Mac/SignAndBuild.sh
--- a/Engine/Source/Programs/Enterprise/Datasmith/DatasmithSketchUpRubyExporter/Installer/Mac/SignAndBuild.sh
+++ b/Engine/Source/Programs/Enterprise/Datasmith/DatasmithSketchUpRubyExporter/Installer/Mac/SignAndBuild.sh
@@ -33,7 +33,7 @@ fi
 DevIdApplication="Developer ID Application: $DevAppId"
 DevIdInstaller="Developer ID Installer: $DevInsId"
 
-for i in {2019..2021}
+for i in {2019..2022}
 do
 	/usr/bin/codesign --force --deep --sign "$DevIdApplication" --timestamp=none "../../.build/${i}/Plugin/UnrealDatasmithSketchUp/DatasmithSketchUpRuby.bundle"
 done
