Fixed crash on exit in ArchiCAD when the Datasmith plugin was loaded

diff --git a/Engine/Source/Programs/Enterprise/Datasmith/DatasmithARCHICADExporter/Private/LoadDatasmithDlls.cpp b/Engine/Source/Programs/Enterprise/Datasmith/DatasmithARCHICADExporter/Private/LoadDatasmithDlls.cpp
--- a/Engine/Source/Programs/Enterprise/Datasmith/DatasmithARCHICADExporter/Private/LoadDatasmithDlls.cpp
+++ b/Engine/Source/Programs/Enterprise/Datasmith/DatasmithARCHICADExporter/Private/LoadDatasmithDlls.cpp
@@ -124,6 +124,17 @@ void UnloadDatasmithDlls(bool bForce)
 			FDatasmithDirectLink::Shutdown();
 			UE_AC_TraceF("UnloadDatasmithDlls - FDatasmithExporterManager::Shutdown\n");
 			FDatasmithExporterManager::Shutdown();
+
+			// the call to FDatasmithExporterManager::Shutdown() is not sufficient.
+			// We have to make sure IsEngineExitRequested() is true before the dll is unloaded as some static destructors rely on that flag to behave properly.
+			//
+			// The specific case handled here:
+			//     FSparseDelegateStorage::SparseDelegateObjectListener dtr uses a maybe-invalid critical-section,
+			//     and rely on IsEngineExitRequested() to behave correctly. We enforce that when the dll unloads
+			if (!IsEngineExitRequested())
+			{
+				RequestEngineExit(TEXT("DLL_PROCESS_DETACH received"));
+			}
 			bLoadSucceed = false;
 		}
 	}
