3ds Max DirectLink: fixed RailClone overwriting own meshes

diff --git a/Engine/Source/Programs/Enterprise/Datasmith/DatasmithMaxExporter/Private/DatasmithMaxDirectLinkItoo.cpp b/Engine/Source/Programs/Enterprise/Datasmith/DatasmithMaxExporter/Private/DatasmithMaxDirectLinkItoo.cpp
--- a/Engine/Source/Programs/Enterprise/Datasmith/DatasmithMaxExporter/Private/DatasmithMaxDirectLinkItoo.cpp
+++ b/Engine/Source/Programs/Enterprise/Datasmith/DatasmithMaxExporter/Private/DatasmithMaxDirectLinkItoo.cpp
@@ -104,6 +104,7 @@ bool ConvertRailClone(ISceneTracker& SceneTracker, FNodeTracker& NodeTracker, Ob
 		for (FHismInstances& Instances: InstancesForMesh)
 		{
 			SceneTracker.SetupDatasmithHISMForNode(NodeTracker, NodeTracker.Node, FRenderMeshForConversion(NodeTracker.Node, Instances.MaxMesh.Get(), false), NodeTracker.Node->GetMtl(), MeshIndex, Instances.Transforms);
+			MeshIndex++;
 		}
 	}
 
