[3ds Max Exporter] Add support for 3ds Max 2023

diff --git a/Engine/Source/Programs/Enterprise/Datasmith/DatasmithMaxExporter/DatasmithMax2023.Build.cs b/Engine/Source/Programs/Enterprise/Datasmith/DatasmithMaxExporter/DatasmithMax2023.Build.cs
new file mode 100644
--- /dev/null
+++ b/Engine/Source/Programs/Enterprise/Datasmith/DatasmithMaxExporter/DatasmithMax2023.Build.cs
@@ -0,0 +1,16 @@
+// Copyright Epic Games, Inc. All Rights Reserved.
+
+using System.Diagnostics;
+
+namespace UnrealBuildTool.Rules
+{
+	[SupportedPlatforms("Win64")]
+	public class DatasmithMax2023 : DatasmithMaxBase
+	{
+		public DatasmithMax2023(ReadOnlyTargetRules Target)
+			: base(Target)
+		{
+		}
+		public override string GetMaxVersion() { return "2023"; }
+	}
+}
\ No newline at end of file
diff --git a/Engine/Source/Programs/Enterprise/Datasmith/DatasmithMaxExporter/DatasmithMax2023.Target.cs b/Engine/Source/Programs/Enterprise/Datasmith/DatasmithMaxExporter/DatasmithMax2023.Target.cs
new file mode 100644
--- /dev/null
+++ b/Engine/Source/Programs/Enterprise/Datasmith/DatasmithMaxExporter/DatasmithMax2023.Target.cs
@@ -0,0 +1,15 @@
+// Copyright Epic Games, Inc. All Rights Reserved.
+
+using UnrealBuildTool;
+using System.Collections.Generic;
+
+public class DatasmithMax2023Target : DatasmithMaxBaseTarget
+{
+	public DatasmithMax2023Target(TargetInfo Target) : base(Target)
+	{
+		LaunchModuleName = "DatasmithMax2023";
+		ExeBinariesSubFolder = @"3DSMax\2023";
+
+		AddCopyPostBuildStep(Target);
+	}
+}
diff --git a/Engine/Source/Programs/Enterprise/Datasmith/DatasmithMaxExporter/Private/DatasmithMaxMeshExporter.cpp b/Engine/Source/Programs/Enterprise/Datasmith/DatasmithMaxExporter/Private/DatasmithMaxMeshExporter.cpp
--- a/Engine/Source/Programs/Enterprise/Datasmith/DatasmithMaxExporter/Private/DatasmithMaxMeshExporter.cpp
+++ b/Engine/Source/Programs/Enterprise/Datasmith/DatasmithMaxExporter/Private/DatasmithMaxMeshExporter.cpp
@@ -26,7 +26,11 @@ MAX_INCLUDES_START
 	#include "mesh.h"
 	#include "meshnormalspec.h"
 	#include "modstack.h"
-	#include "point3.h"
+#if MAX_PRODUCT_YEAR_NUMBER > 2022
+#include "geom/point3.h"
+#else
+#include "point3.h"
+#endif
 MAX_INCLUDES_END
 #pragma warning(pop)
 
