DatasmithMax - Fixing DirectLink serialization corruption caused by copy of a PbrMaterialExpression.

diff --git a/Engine/Source/Programs/Enterprise/Datasmith/DatasmithMaxExporter/Private/MaxMaterialsToUEPbr/DatasmithMaxPhysicalMaterialToUEPbr.cpp b/Engine/Source/Programs/Enterprise/Datasmith/DatasmithMaxExporter/Private/MaxMaterialsToUEPbr/DatasmithMaxPhysicalMaterialToUEPbr.cpp
--- a/Engine/Source/Programs/Enterprise/Datasmith/DatasmithMaxExporter/Private/MaxMaterialsToUEPbr/DatasmithMaxPhysicalMaterialToUEPbr.cpp
+++ b/Engine/Source/Programs/Enterprise/Datasmith/DatasmithMaxExporter/Private/MaxMaterialsToUEPbr/DatasmithMaxPhysicalMaterialToUEPbr.cpp
@@ -182,8 +182,8 @@ IDatasmithMaterialExpression& FDatasmithMaxMaterialsToUEPbrExpressions::CalcIORC
 	IDatasmithMaterialExpression& Rp = Divide(RpNum, RpDen);
 
 	IDatasmithMaterialExpression& ResPre = Add(Rp, Rs);
-	IDatasmithMaterialExpression& Res = Multiply(ResPre, Scalar(0.5));
-	Res = Power(Res, Scalar(0.5));
+	IDatasmithMaterialExpression& ResMul = Multiply(ResPre, Scalar(0.5));
+	IDatasmithMaterialExpression& Res = Power(ResMul, Scalar(0.5));
 
 	return Lerp(ToBeConnected0, ToBeConnected90, Res);
 }
