3ds Max: Corona Physical Material to UEPbr

- Scope: Pure Metallic, Thin(as ThinTranslucent), Transparent/Refractive material(including volumetric transmittance for PathTracer), simple opaque, emissive materials

- Roughness(with Glossiness mode) converted linearly. So that it's invertible same as in Max - e.g. 0.5 Roughness is same as 0.5 Gloss

- Specular/Ior conversion suported for both Ior and Disney modes. Ior is properly converted(to Specular). And Disney Specular to refraction Ior.

diff --git a/Engine/Source/Programs/Enterprise/Datasmith/DatasmithMaxExporter/Private/MaxMaterialsToUEPbr/DatasmithMaxCoronaMaterialsToUEPbr.cpp b/Engine/Source/Programs/Enterprise/Datasmith/DatasmithMaxExporter/Private/MaxMaterialsToUEPbr/DatasmithMaxCoronaMaterialsToUEPbr.cpp
--- a/Engine/Source/Programs/Enterprise/Datasmith/DatasmithMaxExporter/Private/MaxMaterialsToUEPbr/DatasmithMaxCoronaMaterialsToUEPbr.cpp
+++ b/Engine/Source/Programs/Enterprise/Datasmith/DatasmithMaxExporter/Private/MaxMaterialsToUEPbr/DatasmithMaxCoronaMaterialsToUEPbr.cpp
@@ -4,6 +4,7 @@
 
 #include "DatasmithMaxClassIDs.h"
 #include "DatasmithMaxTexmapParser.h"
+#include "DatasmithMaxHelper.h"
 #include "DatasmithMaxWriter.h"
 #include "DatasmithSceneFactory.h"
 #include "MaxMaterialsToUEPbr/DatasmithMaxTexmapToUEPbr.h"
@@ -352,6 +353,366 @@ namespace DatasmithMaxCoronaMaterialsToUEPbrImpl
 			}
 		}
 	};
+
+	struct FCoronaPhysicalMaterial
+	{
+		FLinearColor BaseColor;
+		float BaseLevel;
+		DatasmithMaxTexmapParser::FMapParameter BaseTexmap;
+		int MetalnessMode;
+		FLinearColor OpacityColor;
+		float OpacityLevel;
+		DatasmithMaxTexmapParser::FMapParameter OpacityTexmap;
+		bool bOpacityCutout;
+		float BaseRoughness;
+		DatasmithMaxTexmapParser::FMapParameter BaseRoughnessTexmap;
+		float BaseIor;
+		DatasmithMaxTexmapParser::FMapParameter BaseIorTexmap;
+		float RefractionAmount;
+		DatasmithMaxTexmapParser::FMapParameter RefractionAmountTexmap;
+		bool bUseThinMode;
+		float ClearcoatAmount;
+		DatasmithMaxTexmapParser::FMapParameter ClearcoatAmountTexmap;
+		float ClearcoatIor;
+		DatasmithMaxTexmapParser::FMapParameter ClearcoatIorTexmap;
+		float ClearcoatRoughness;
+		DatasmithMaxTexmapParser::FMapParameter ClearcoatRoughnessTexmap;
+		FLinearColor VolumetricAbsorptionColor;
+		DatasmithMaxTexmapParser::FMapParameter VolumetricAbsorptionTexmap;
+		FLinearColor VolumetricScatteringColor;
+		DatasmithMaxTexmapParser::FMapParameter VolumetricScatteringTexmap;
+		float AttenuationDistance;
+		float ScatterDirectionality;
+		bool bScatterSingleBounce;
+		FLinearColor SelfIllumColor;
+		float SelfIllumLevel;
+		DatasmithMaxTexmapParser::FMapParameter SelfIllumTexmap;
+		DatasmithMaxTexmapParser::FMapParameter BaseBumpTexmap;
+		float TranslucencyFraction;
+		DatasmithMaxTexmapParser::FMapParameter TranslucencyFractionTexmap;
+		FLinearColor ThinAbsorptionColor;
+		DatasmithMaxTexmapParser::FMapParameter ThinAbsorptionTexmap;
+		DatasmithMaxTexmapParser::FMapParameter MetalnessTexmap;
+		int RoughnessMode;
+		FLinearColor TranslucencyColor;
+		DatasmithMaxTexmapParser::FMapParameter TranslucencyColorTexmap;
+		int IorMode;
+
+		void Parse(Mtl& Material)
+		{
+
+			const TimeValue CurrentTime = GetCOREInterface()->GetTime();
+			int NumParamBlocks = Material.NumParamBlocks();
+
+			for (int j = 0; j < NumParamBlocks; j++)
+			{
+				IParamBlock2* ParamBlock2 = Material.GetParamBlockByID((short)j);
+				// The the descriptor to 'decode'
+				ParamBlockDesc2* ParamBlockDesc = ParamBlock2->GetDesc();
+				// Loop through all the defined parameters therein
+				for (int i = 0; i < ParamBlockDesc->count; i++)
+				{
+					const ParamDef& ParamDefinition = ParamBlockDesc->paramdefs[i];
+
+					if (FCString::Stricmp(ParamDefinition.int_name, TEXT("baseColor")) == 0)
+					{
+						BaseColor = FDatasmithMaxMatHelper::MaxLinearColorToFLinearColor(ParamBlock2->GetColor(ParamDefinition.ID, GetCOREInterface()->GetTime()));
+					}
+					else if (FCString::Stricmp(ParamDefinition.int_name, TEXT("baseLevel")) == 0)
+					{
+						BaseLevel = ParamBlock2->GetFloat(ParamDefinition.ID, GetCOREInterface()->GetTime());
+					}
+					else if (FCString::Stricmp(ParamDefinition.int_name, TEXT("baseTexmap")) == 0)
+					{
+						BaseTexmap.Map = ParamBlock2->GetTexmap(ParamDefinition.ID, GetCOREInterface()->GetTime());
+					}
+					else if (FCString::Stricmp(ParamDefinition.int_name, TEXT("baseTexmapOn")) == 0)
+					{
+						BaseTexmap.bEnabled = ParamBlock2->GetInt(ParamDefinition.ID, CurrentTime) != 0;
+					}
+					else if (FCString::Stricmp(ParamDefinition.int_name, TEXT("baseMapAmount")) == 0)
+					{
+						BaseTexmap.Weight = ParamBlock2->GetFloat(ParamDefinition.ID, CurrentTime);
+					}
+					else if (FCString::Stricmp(ParamDefinition.int_name, TEXT("metalnessMode")) == 0)
+					{
+						MetalnessMode = ParamBlock2->GetInt(ParamDefinition.ID, GetCOREInterface()->GetTime());
+					}
+					else if (FCString::Stricmp(ParamDefinition.int_name, TEXT("opacityColor")) == 0)
+					{
+						OpacityColor = FDatasmithMaxMatHelper::MaxLinearColorToFLinearColor(ParamBlock2->GetColor(ParamDefinition.ID, GetCOREInterface()->GetTime()));
+					}
+					else if (FCString::Stricmp(ParamDefinition.int_name, TEXT("opacityLevel")) == 0)
+					{
+						OpacityLevel = ParamBlock2->GetFloat(ParamDefinition.ID, GetCOREInterface()->GetTime());
+					}
+					else if (FCString::Stricmp(ParamDefinition.int_name, TEXT("opacityTexmap")) == 0)
+					{
+						OpacityTexmap.Map = ParamBlock2->GetTexmap(ParamDefinition.ID, GetCOREInterface()->GetTime());
+					}
+					else if (FCString::Stricmp(ParamDefinition.int_name, TEXT("opacityTexmapOn")) == 0)
+					{
+						OpacityTexmap.bEnabled = ParamBlock2->GetInt(ParamDefinition.ID, CurrentTime) != 0;
+					}
+					else if (FCString::Stricmp(ParamDefinition.int_name, TEXT("opacityMapAmount")) == 0)
+					{
+						OpacityTexmap.Weight = ParamBlock2->GetFloat(ParamDefinition.ID, CurrentTime);
+					}
+					else if (FCString::Stricmp(ParamDefinition.int_name, TEXT("opacityCutout")) == 0)
+					{
+						bOpacityCutout = ParamBlock2->GetInt(ParamDefinition.ID, GetCOREInterface()->GetTime()) != 0;
+					}
+					else if (FCString::Stricmp(ParamDefinition.int_name, TEXT("baseRoughness")) == 0)
+					{
+						BaseRoughness = ParamBlock2->GetFloat(ParamDefinition.ID, GetCOREInterface()->GetTime());
+					}
+					else if (FCString::Stricmp(ParamDefinition.int_name, TEXT("baseRoughnessTexmap")) == 0)
+					{
+						BaseRoughnessTexmap.Map = ParamBlock2->GetTexmap(ParamDefinition.ID, GetCOREInterface()->GetTime());
+					}
+					else if (FCString::Stricmp(ParamDefinition.int_name, TEXT("baseRoughnessTexmapOn")) == 0)
+					{
+						BaseRoughnessTexmap.bEnabled = ParamBlock2->GetInt(ParamDefinition.ID, CurrentTime) != 0;
+					}
+					else if (FCString::Stricmp(ParamDefinition.int_name, TEXT("baseRoughnessMapAmount")) == 0)
+					{
+						BaseRoughnessTexmap.Weight = ParamBlock2->GetFloat(ParamDefinition.ID, CurrentTime);
+					}
+					else if (FCString::Stricmp(ParamDefinition.int_name, TEXT("baseIor")) == 0)
+					{
+						BaseIor = ParamBlock2->GetFloat(ParamDefinition.ID, GetCOREInterface()->GetTime());
+					}
+					else if (FCString::Stricmp(ParamDefinition.int_name, TEXT("baseIorTexmap")) == 0)
+					{
+						BaseIorTexmap.Map = ParamBlock2->GetTexmap(ParamDefinition.ID, GetCOREInterface()->GetTime());
+					}
+					else if (FCString::Stricmp(ParamDefinition.int_name, TEXT("baseIorTexmapOn")) == 0)
+					{
+						BaseIorTexmap.bEnabled = ParamBlock2->GetInt(ParamDefinition.ID, CurrentTime) != 0;
+					}
+					else if (FCString::Stricmp(ParamDefinition.int_name, TEXT("baseIorMapAmount")) == 0)
+					{
+						BaseIorTexmap.Weight = ParamBlock2->GetFloat(ParamDefinition.ID, CurrentTime);
+					}
+					else if (FCString::Stricmp(ParamDefinition.int_name, TEXT("refractionAmount")) == 0)
+					{
+						RefractionAmount = ParamBlock2->GetFloat(ParamDefinition.ID, GetCOREInterface()->GetTime());
+					}
+					else if (FCString::Stricmp(ParamDefinition.int_name, TEXT("refractionAmountTexmap")) == 0)
+					{
+						RefractionAmountTexmap.Map = ParamBlock2->GetTexmap(ParamDefinition.ID, GetCOREInterface()->GetTime());
+					}
+					else if (FCString::Stricmp(ParamDefinition.int_name, TEXT("refractionAmountTexmapOn")) == 0)
+					{
+						RefractionAmountTexmap.bEnabled = ParamBlock2->GetInt(ParamDefinition.ID, CurrentTime) != 0;
+					}
+					else if (FCString::Stricmp(ParamDefinition.int_name, TEXT("refractionAmountMapAmount")) == 0)
+					{
+						RefractionAmountTexmap.Weight = ParamBlock2->GetFloat(ParamDefinition.ID, CurrentTime);
+					}
+					else if (FCString::Stricmp(ParamDefinition.int_name, TEXT("useThinMode")) == 0)
+					{
+						bUseThinMode = ParamBlock2->GetInt(ParamDefinition.ID, GetCOREInterface()->GetTime()) != 0;
+					}
+					else if (FCString::Stricmp(ParamDefinition.int_name, TEXT("clearcoatAmount")) == 0)
+					{
+						ClearcoatAmount = ParamBlock2->GetFloat(ParamDefinition.ID, GetCOREInterface()->GetTime());
+					}
+					else if (FCString::Stricmp(ParamDefinition.int_name, TEXT("clearcoatAmountTexmap")) == 0)
+					{
+						ClearcoatAmountTexmap.Map = ParamBlock2->GetTexmap(ParamDefinition.ID, GetCOREInterface()->GetTime());
+					}
+					else if (FCString::Stricmp(ParamDefinition.int_name, TEXT("clearcoatAmountTexmapOn")) == 0)
+					{
+						ClearcoatAmountTexmap.bEnabled = ParamBlock2->GetInt(ParamDefinition.ID, CurrentTime) != 0;
+					}
+					else if (FCString::Stricmp(ParamDefinition.int_name, TEXT("clearcoatAmountMapAmount")) == 0)
+					{
+						ClearcoatAmountTexmap.Weight = ParamBlock2->GetFloat(ParamDefinition.ID, CurrentTime);
+					}
+					else if (FCString::Stricmp(ParamDefinition.int_name, TEXT("clearcoatIor")) == 0)
+					{
+						ClearcoatIor = ParamBlock2->GetFloat(ParamDefinition.ID, GetCOREInterface()->GetTime());
+					}
+					else if (FCString::Stricmp(ParamDefinition.int_name, TEXT("clearcoatIorTexmap")) == 0)
+					{
+						ClearcoatIorTexmap.Map = ParamBlock2->GetTexmap(ParamDefinition.ID, GetCOREInterface()->GetTime());
+					}
+					else if (FCString::Stricmp(ParamDefinition.int_name, TEXT("clearcoatIorTexmapOn")) == 0)
+					{
+						ClearcoatIorTexmap.bEnabled = ParamBlock2->GetInt(ParamDefinition.ID, CurrentTime) != 0;
+					}
+					else if (FCString::Stricmp(ParamDefinition.int_name, TEXT("clearcoatIorMapAmount")) == 0)
+					{
+						ClearcoatIorTexmap.Weight = ParamBlock2->GetFloat(ParamDefinition.ID, CurrentTime);
+					}
+					else if (FCString::Stricmp(ParamDefinition.int_name, TEXT("clearcoatRoughness")) == 0)
+					{
+						ClearcoatRoughness = ParamBlock2->GetFloat(ParamDefinition.ID, GetCOREInterface()->GetTime());
+					}
+					else if (FCString::Stricmp(ParamDefinition.int_name, TEXT("clearcoatRoughnessTexmap")) == 0)
+					{
+						ClearcoatRoughnessTexmap.Map = ParamBlock2->GetTexmap(ParamDefinition.ID, GetCOREInterface()->GetTime());
+					}
+					else if (FCString::Stricmp(ParamDefinition.int_name, TEXT("clearcoatRoughnessTexmapOn")) == 0)
+					{
+						ClearcoatRoughnessTexmap.bEnabled = ParamBlock2->GetInt(ParamDefinition.ID, CurrentTime) != 0;
+					}
+					else if (FCString::Stricmp(ParamDefinition.int_name, TEXT("clearcoatRoughnessMapAmount")) == 0)
+					{
+						ClearcoatRoughnessTexmap.Weight = ParamBlock2->GetFloat(ParamDefinition.ID, CurrentTime);
+					}
+					else if (FCString::Stricmp(ParamDefinition.int_name, TEXT("volumetricAbsorptionColor")) == 0)
+					{
+						VolumetricAbsorptionColor = FDatasmithMaxMatHelper::MaxLinearColorToFLinearColor(ParamBlock2->GetColor(ParamDefinition.ID, GetCOREInterface()->GetTime()));
+					}
+					else if (FCString::Stricmp(ParamDefinition.int_name, TEXT("volumetricAbsorptionTexmap")) == 0)
+					{
+						VolumetricAbsorptionTexmap.Map = ParamBlock2->GetTexmap(ParamDefinition.ID, GetCOREInterface()->GetTime());
+					}
+					else if (FCString::Stricmp(ParamDefinition.int_name, TEXT("volumetricAbsorptionTexmapOn")) == 0)
+					{
+						VolumetricAbsorptionTexmap.bEnabled = ParamBlock2->GetInt(ParamDefinition.ID, CurrentTime) != 0;
+					}
+					else if (FCString::Stricmp(ParamDefinition.int_name, TEXT("volumetricAbsorptionMapAmount")) == 0)
+					{
+						VolumetricAbsorptionTexmap.Weight = ParamBlock2->GetFloat(ParamDefinition.ID, CurrentTime);
+					}
+					else if (FCString::Stricmp(ParamDefinition.int_name, TEXT("volumetricScatteringColor")) == 0)
+					{
+						VolumetricScatteringColor = FDatasmithMaxMatHelper::MaxLinearColorToFLinearColor(ParamBlock2->GetColor(ParamDefinition.ID, GetCOREInterface()->GetTime()));
+					}
+					else if (FCString::Stricmp(ParamDefinition.int_name, TEXT("volumetricScatteringTexmap")) == 0)
+					{
+						VolumetricScatteringTexmap.Map = ParamBlock2->GetTexmap(ParamDefinition.ID, GetCOREInterface()->GetTime());
+					}
+					else if (FCString::Stricmp(ParamDefinition.int_name, TEXT("volumetricScatteringTexmapOn")) == 0)
+					{
+						VolumetricScatteringTexmap.bEnabled = ParamBlock2->GetInt(ParamDefinition.ID, CurrentTime) != 0;
+					}
+					else if (FCString::Stricmp(ParamDefinition.int_name, TEXT("volumetricScatteringMapAmount")) == 0)
+					{
+						VolumetricScatteringTexmap.Weight = ParamBlock2->GetFloat(ParamDefinition.ID, CurrentTime);
+					}
+					else if (FCString::Stricmp(ParamDefinition.int_name, TEXT("attenuationDistance")) == 0)
+					{
+						AttenuationDistance = ParamBlock2->GetFloat(ParamDefinition.ID, GetCOREInterface()->GetTime()) * (float)GetSystemUnitScale(UNITS_CENTIMETERS);
+					}
+					else if (FCString::Stricmp(ParamDefinition.int_name, TEXT("scatterDirectionality")) == 0)
+					{
+						ScatterDirectionality = ParamBlock2->GetFloat(ParamDefinition.ID, GetCOREInterface()->GetTime());
+					}
+					else if (FCString::Stricmp(ParamDefinition.int_name, TEXT("scatterSingleBounce")) == 0)
+					{
+						bScatterSingleBounce = ParamBlock2->GetInt(ParamDefinition.ID, GetCOREInterface()->GetTime()) != 0;
+					}
+					else if (FCString::Stricmp(ParamDefinition.int_name, TEXT("selfIllumColor")) == 0)
+					{
+						SelfIllumColor = FDatasmithMaxMatHelper::MaxLinearColorToFLinearColor(ParamBlock2->GetColor(ParamDefinition.ID, GetCOREInterface()->GetTime()));
+					}
+					else if (FCString::Stricmp(ParamDefinition.int_name, TEXT("selfIllumLevel")) == 0)
+					{
+						SelfIllumLevel = ParamBlock2->GetFloat(ParamDefinition.ID, GetCOREInterface()->GetTime());
+					}
+					else if (FCString::Stricmp(ParamDefinition.int_name, TEXT("selfIllumTexmap")) == 0)
+					{
+						SelfIllumTexmap.Map = ParamBlock2->GetTexmap(ParamDefinition.ID, GetCOREInterface()->GetTime());
+					}
+					else if (FCString::Stricmp(ParamDefinition.int_name, TEXT("selfIllumTexmapOn")) == 0)
+					{
+						SelfIllumTexmap.bEnabled = ParamBlock2->GetInt(ParamDefinition.ID, CurrentTime) != 0;
+					}
+					else if (FCString::Stricmp(ParamDefinition.int_name, TEXT("selfIllumMapAmount")) == 0)
+					{
+						SelfIllumTexmap.Weight = ParamBlock2->GetFloat(ParamDefinition.ID, CurrentTime);
+					}
+					else if (FCString::Stricmp(ParamDefinition.int_name, TEXT("baseBumpTexmap")) == 0)
+					{
+						BaseBumpTexmap.Map = ParamBlock2->GetTexmap(ParamDefinition.ID, GetCOREInterface()->GetTime());
+					}
+					else if (FCString::Stricmp(ParamDefinition.int_name, TEXT("baseBumpTexmapOn")) == 0)
+					{
+						BaseBumpTexmap.bEnabled = ParamBlock2->GetInt(ParamDefinition.ID, CurrentTime) != 0;
+					}
+					else if (FCString::Stricmp(ParamDefinition.int_name, TEXT("baseBumpMapAmount")) == 0)
+					{
+						BaseBumpTexmap.Weight = ParamBlock2->GetFloat(ParamDefinition.ID, CurrentTime);
+					}
+					else if (FCString::Stricmp(ParamDefinition.int_name, TEXT("translucencyFraction")) == 0)
+					{
+						TranslucencyFraction = ParamBlock2->GetFloat(ParamDefinition.ID, GetCOREInterface()->GetTime());
+					}
+					else if (FCString::Stricmp(ParamDefinition.int_name, TEXT("translucencyFractionTexmap")) == 0)
+					{
+						TranslucencyFractionTexmap.Map = ParamBlock2->GetTexmap(ParamDefinition.ID, GetCOREInterface()->GetTime());
+					}
+					else if (FCString::Stricmp(ParamDefinition.int_name, TEXT("translucencyFractionTexmapOn")) == 0)
+					{
+						TranslucencyFractionTexmap.bEnabled = ParamBlock2->GetInt(ParamDefinition.ID, CurrentTime) != 0;
+					}
+					else if (FCString::Stricmp(ParamDefinition.int_name, TEXT("translucencyFractionMapAmount")) == 0)
+					{
+						TranslucencyFractionTexmap.Weight = ParamBlock2->GetFloat(ParamDefinition.ID, CurrentTime);
+					}
+					else if (FCString::Stricmp(ParamDefinition.int_name, TEXT("thinAbsorptionColor")) == 0)
+					{
+						ThinAbsorptionColor = FDatasmithMaxMatHelper::MaxLinearColorToFLinearColor(ParamBlock2->GetColor(ParamDefinition.ID, GetCOREInterface()->GetTime()));
+					}
+					else if (FCString::Stricmp(ParamDefinition.int_name, TEXT("thinAbsorptionTexmap")) == 0)
+					{
+						ThinAbsorptionTexmap.Map = ParamBlock2->GetTexmap(ParamDefinition.ID, GetCOREInterface()->GetTime());
+					}
+					else if (FCString::Stricmp(ParamDefinition.int_name, TEXT("thinAbsorptionTexmapOn")) == 0)
+					{
+						ThinAbsorptionTexmap.bEnabled = ParamBlock2->GetInt(ParamDefinition.ID, CurrentTime) != 0;
+					}
+					else if (FCString::Stricmp(ParamDefinition.int_name, TEXT("thinAbsorptionMapAmount")) == 0)
+					{
+						ThinAbsorptionTexmap.Weight = ParamBlock2->GetFloat(ParamDefinition.ID, CurrentTime);
+					}
+					else if (FCString::Stricmp(ParamDefinition.int_name, TEXT("metalnessTexmap")) == 0)
+					{
+						MetalnessTexmap.Map = ParamBlock2->GetTexmap(ParamDefinition.ID, GetCOREInterface()->GetTime());
+					}
+					else if (FCString::Stricmp(ParamDefinition.int_name, TEXT("metalnessTexmapOn")) == 0)
+					{
+						MetalnessTexmap.bEnabled = ParamBlock2->GetInt(ParamDefinition.ID, CurrentTime) != 0;
+					}
+					else if (FCString::Stricmp(ParamDefinition.int_name, TEXT("metalnessMapAmount")) == 0)
+					{
+						MetalnessTexmap.Weight = ParamBlock2->GetFloat(ParamDefinition.ID, CurrentTime);
+					}
+					else if (FCString::Stricmp(ParamDefinition.int_name, TEXT("roughnessMode")) == 0)
+					{
+						RoughnessMode = ParamBlock2->GetInt(ParamDefinition.ID, GetCOREInterface()->GetTime());
+					}
+					else if (FCString::Stricmp(ParamDefinition.int_name, TEXT("translucencyColor")) == 0)
+					{
+						TranslucencyColor = FDatasmithMaxMatHelper::MaxLinearColorToFLinearColor(ParamBlock2->GetColor(ParamDefinition.ID, GetCOREInterface()->GetTime()));
+					}
+					else if (FCString::Stricmp(ParamDefinition.int_name, TEXT("translucencyColorTexmap")) == 0)
+					{
+						TranslucencyColorTexmap.Map = ParamBlock2->GetTexmap(ParamDefinition.ID, GetCOREInterface()->GetTime());
+					}
+					else if (FCString::Stricmp(ParamDefinition.int_name, TEXT("translucencyColorTexmapOn")) == 0)
+					{
+						TranslucencyColorTexmap.bEnabled = ParamBlock2->GetInt(ParamDefinition.ID, CurrentTime) != 0;
+					}
+					else if (FCString::Stricmp(ParamDefinition.int_name, TEXT("translucencyColorMapAmount")) == 0)
+					{
+						TranslucencyColorTexmap.Weight = ParamBlock2->GetFloat(ParamDefinition.ID, CurrentTime);
+					}
+					else if (FCString::Stricmp(ParamDefinition.int_name, TEXT("iorMode")) == 0)
+					{
+						IorMode = ParamBlock2->GetInt(ParamDefinition.ID, GetCOREInterface()->GetTime());
+					}
+				}
+				ParamBlock2->ReleaseDesc();
+			}
+		}
+	};
+
+  
 }
 
 FDatasmithMaxCoronaMaterialsToUEPbr::FDatasmithMaxCoronaMaterialsToUEPbr()
@@ -820,3 +1181,269 @@ void FDatasmithMaxCoronaLightMaterialToUEPbr::Convert(TSharedRef<IDatasmithScene
 
 	MaterialElement = PbrMaterialElement;
 }
+
+bool FDatasmithMaxCoronaPhysicalMaterialToUEPbr::IsSupported(Mtl* Material)
+{
+	return true;
+}
+
+void FDatasmithMaxCoronaPhysicalMaterialToUEPbr::Convert(TSharedRef<IDatasmithScene> DatasmithScene,
+	TSharedPtr<IDatasmithBaseMaterialElement>& MaterialElement, Mtl* Material, const TCHAR* AssetsPath)
+{
+	if ( !Material )
+	{
+		return;
+	}
+
+	TSharedRef< IDatasmithUEPbrMaterialElement > PbrMaterialElement = FDatasmithSceneFactory::CreateUEPbrMaterial( Material->GetName().data() );
+	FScopedConvertState ScopedConvertState(ConvertState);
+	ConvertState.DatasmithScene = DatasmithScene;
+	ConvertState.MaterialElement = PbrMaterialElement;
+	ConvertState.AssetsPath = AssetsPath;
+
+	DatasmithMaxCoronaMaterialsToUEPbrImpl::FCoronaPhysicalMaterial MaterialProperties;
+	MaterialProperties.Parse(*Material);
+
+	// todo: move out?
+	const TCHAR* PARAM_NAME_BASELEVEL = TEXT("Base Level");
+	const TCHAR* PARAM_NAME_BASECOLOR = TEXT("Base Color");
+	const TCHAR* PARAM_NAME_BASECOLORMAP_WEIGHT = TEXT("Base Color Texture Amount");
+
+	const TCHAR* PARAM_NAME_EMISSIVELEVEL = TEXT("Emissive Level");
+	const TCHAR* PARAM_NAME_EMISSIVECOLOR = TEXT("Emissive Color");
+	const TCHAR* PARAM_NAME_EMISSIVECOLORMAP_WEIGHT = TEXT("Emissive Color Texture Amount");
+
+	const TCHAR* PARAM_NAME_OPACITY_LEVEL = TEXT("Opacity Level");
+	const TCHAR* PARAM_NAME_OPACITY_MAP_WEIGHT = TEXT("Opacity Texture Amount");
+
+	const TCHAR* PARAM_NAME_ROUGHNESS = TEXT("Roughness");
+	const TCHAR* PARAM_NAME_ROUGHNESSMAP_WEIGHT = TEXT("Roughness Map Amount");
+
+	const TCHAR* PARAM_NAME_BASEIOR = TEXT("Base IOR");
+	const TCHAR* PARAM_NAME_BASEIORMAP_WEIGHT = TEXT("Base IOR Map Weight");
+
+	const TCHAR* PARAM_NAME_REFRACTIONAMOUNT = TEXT("Refraction Amount");
+	const TCHAR* PARAM_NAME_REFRACTIONMAP_WEIGHT = TEXT("Refraction Map Weight");
+
+	const TCHAR* PARAM_NAME_BUMPMAP = TEXT("Bump Map");
+
+	const TCHAR* PARAM_NAME_ABSORPTIONDISTANCE = TEXT("Absorption Distance");
+
+
+
+	IDatasmithMaterialExpression* BaseTexmapExpression = BlendTexmapWithColor(MaterialProperties.BaseTexmap, MaterialProperties.BaseColor, PARAM_NAME_BASECOLOR, PARAM_NAME_BASECOLORMAP_WEIGHT);
+	IDatasmithMaterialExpression& DiffuseColorExpression = BaseTexmapExpression
+		                                                       ? Multiply(*BaseTexmapExpression, Scalar(MaterialProperties.BaseLevel, PARAM_NAME_BASELEVEL))
+		                                                       : Color(ScaleRGB(MaterialProperties.BaseColor, MaterialProperties.BaseLevel), PARAM_NAME_BASECOLOR);
+
+	// Opacity Color
+	IDatasmithMaterialExpression* OpacityTexmapExpression = BlendTexmapWithColor(MaterialProperties.OpacityTexmap, MaterialProperties.OpacityColor, nullptr, PARAM_NAME_OPACITY_MAP_WEIGHT);
+	IDatasmithMaterialExpression* OpacityColorExpression = nullptr;
+	if (OpacityTexmapExpression)
+	{
+		OpacityColorExpression = &Multiply(*OpacityTexmapExpression, Scalar(MaterialProperties.OpacityLevel, PARAM_NAME_OPACITY_LEVEL));
+	}
+	else
+	{
+		FLinearColor OpacityColor = MaterialProperties.OpacityColor * MaterialProperties.OpacityLevel;
+		// Don't consider almost full opacity as transparency to prevent creating opacity expression(which would lead to making non-opaque shader)
+		if (!OpacityColor.Equals(FLinearColor::White)) 
+		{
+			OpacityColorExpression = &Color(OpacityColor);
+		}
+	}
+
+	// Emissive
+	if (!FMath::IsNearlyZero(MaterialProperties.SelfIllumLevel))
+	{
+		IDatasmithMaterialExpression* SelfIllumTexmapExpression = BlendTexmapWithColor(MaterialProperties.SelfIllumTexmap, MaterialProperties.SelfIllumColor, PARAM_NAME_EMISSIVECOLOR, PARAM_NAME_EMISSIVECOLORMAP_WEIGHT);
+		Connect(PbrMaterialElement->GetEmissiveColor(), SelfIllumTexmapExpression
+															? Multiply(*SelfIllumTexmapExpression,
+																Scalar(MaterialProperties.SelfIllumLevel, PARAM_NAME_EMISSIVELEVEL))
+															: Color( ScaleRGB(MaterialProperties.SelfIllumColor, MaterialProperties.SelfIllumLevel), PARAM_NAME_EMISSIVECOLOR));
+	}
+
+	// Refraction
+	IDatasmithMaterialExpression* RefractionAmountMapExpression = BlendTexmapWithScalar(MaterialProperties.RefractionAmountTexmap, MaterialProperties.RefractionAmount, PARAM_NAME_REFRACTIONAMOUNT, PARAM_NAME_REFRACTIONMAP_WEIGHT);
+	bool bHasRefraction = MaterialProperties.RefractionAmount > 0.0f || RefractionAmountMapExpression;
+
+	// Roughness
+	IDatasmithMaterialExpression* RoughnessTexmapExpression = BlendTexmapWithScalar(MaterialProperties.BaseRoughnessTexmap, MaterialProperties.BaseRoughness, PARAM_NAME_ROUGHNESS, PARAM_NAME_ROUGHNESSMAP_WEIGHT);
+	IDatasmithMaterialExpression& RoughnessExpression = RoughnessTexmapExpression ? *RoughnessTexmapExpression
+		                                                                          : Scalar(MaterialProperties.BaseRoughness, PARAM_NAME_ROUGHNESS);
+	bool bRoughnessAsGlossiness = MaterialProperties.RoughnessMode == 1;
+	Connect(PbrMaterialElement->GetRoughness(), bRoughnessAsGlossiness ? OneMinus(RoughnessExpression) : RoughnessExpression);
+
+	// Specular
+	// Specular in Unreal is F0 divided by 0.08 (i.e. F0=S*0.08 calculated in Unreal, where S is UE PBR Specular input )
+	// and F0 = (ior-1)^2/(ior+1)^2
+	// baseIor in 'normal' Corona mode(iorMode=0) is actual index of refraction, but in 'Disney' mode (iorMode=1) baseIor means the same "F0/0.08" as in Unreal
+	// then, in Disney mode, (actual)Ior is (sqrt(Specular*0.08) + 1)/(1 - sqrt(Specular*0.08)
+	IDatasmithMaterialExpression* BaseIorMapExpression = BlendTexmapWithScalar(MaterialProperties.BaseIorTexmap, MaterialProperties.BaseIor, PARAM_NAME_BASEIOR, PARAM_NAME_BASEIORMAP_WEIGHT);
+
+	// Compute Ior (even for Disney mode) we'll need it for refraction anyway
+	float Ior = (MaterialProperties.IorMode == 0)
+		? MaterialProperties.BaseIor
+		: ((MaterialProperties.BaseIor < 1.0f ) ? (FMath::Sqrt(MaterialProperties.BaseIor*0.08) + 1) / (1 - FMath::Sqrt(MaterialProperties.BaseIor*0.08)) : 1.8f);
+
+	IDatasmithMaterialExpression& SpecularExpression = (MaterialProperties.IorMode == 0)
+		? (BaseIorMapExpression
+			? Divide(
+				Power(Subtract(*BaseIorMapExpression, 1.0f), 2.0f), 
+				Multiply(Power(Add(*BaseIorMapExpression, 1.0f), 2.0f), 0.08f))
+			: Scalar((Ior > 1.0f ? FMath::Square(Ior-1)/FMath::Square(Ior+1) : 0.0f) / 0.08f))
+		: (BaseIorMapExpression ? *BaseIorMapExpression : Scalar(MaterialProperties.BaseIor));
+
+	IDatasmithMaterialExpression* BaseColorExpression = &DiffuseColorExpression;
+
+	// Build material based on specific configuration
+	if ((MaterialProperties.MetalnessMode == 1) && !MaterialProperties.MetalnessTexmap.IsMapPresentAndEnabled()) // Pure metal
+	{
+		Connect(PbrMaterialElement->GetMetallic(), &Scalar(1.0f));
+
+		if (OpacityColorExpression)
+		{
+			Connect(PbrMaterialElement->GetOpacity(), Desaturate(*OpacityColorExpression));
+			BaseColorExpression = &Multiply(DiffuseColorExpression, *OpacityColorExpression); // Scale base color by opacity color to imitate colored opacity
+		}
+	}
+	else if (OpacityColorExpression || bHasRefraction)  // Transparent non-metal
+	{
+		if (MaterialProperties.bUseThinMode)
+		{
+			// useThinMode is used for 'leaves' which means that opacity controls where a leaf pixel is visible
+
+			PbrMaterialElement->SetShadingModel(EDatasmithShadingModel::ThinTranslucent);
+			IDatasmithMaterialExpressionGeneric* ThinTranslucencyMaterialOutput = PbrMaterialElement->AddMaterialExpression<IDatasmithMaterialExpressionGeneric>();
+			ThinTranslucencyMaterialOutput->SetExpressionName(TEXT("ThinTranslucentMaterialOutput"));
+
+			IDatasmithMaterialExpression* TranslucencyExpression = nullptr;
+
+			// Translucency
+			if (MaterialProperties.TranslucencyFraction > 0.0f || MaterialProperties.TranslucencyFractionTexmap.IsMapPresentAndEnabled())
+			{
+				IDatasmithMaterialExpression* TranslucencyColorTexmapExpression = BlendTexmapWithColor(MaterialProperties.TranslucencyColorTexmap, MaterialProperties.TranslucencyColor);
+				IDatasmithMaterialExpression& TranslucencyColorExpression = TranslucencyColorTexmapExpression ? *TranslucencyColorTexmapExpression : Color(MaterialProperties.TranslucencyColor);
+				IDatasmithMaterialExpression* TranslucencyFractionTexmap = ConvertTexmap(MaterialProperties.TranslucencyFractionTexmap);
+
+				// todo: TranslucencyFractionTexmap
+				TranslucencyExpression = &TranslucencyColorExpression;
+			}
+
+			// Thin Absorption
+			if (bHasRefraction) // Only when there's some refraction thin absorption is affecting
+			{
+				IDatasmithMaterialExpression* ThinAbsorptionTexmapExpression = BlendTexmapWithColor(MaterialProperties.ThinAbsorptionTexmap, MaterialProperties.ThinAbsorptionColor);
+				IDatasmithMaterialExpression& ThinAbsorptionColorExpression = ThinAbsorptionTexmapExpression ? *ThinAbsorptionTexmapExpression : Color(MaterialProperties.ThinAbsorptionColor);
+
+				IDatasmithMaterialExpression& RefractionAmountExpression = RefractionAmountMapExpression ? *RefractionAmountMapExpression : Scalar(MaterialProperties.RefractionAmount);
+
+				TranslucencyExpression = &(TranslucencyExpression ? Lerp(*TranslucencyExpression, ThinAbsorptionColorExpression, RefractionAmountExpression) : Multiply(ThinAbsorptionColorExpression, RefractionAmountExpression));
+			}
+
+			IDatasmithMaterialExpression* OpacityExpression = OpacityColorExpression;
+
+			// Refraction (factor into opacity)
+			if (bHasRefraction)
+			{
+				IDatasmithMaterialExpression& RefractionAmountExpression = RefractionAmountMapExpression ? *RefractionAmountMapExpression : Scalar(MaterialProperties.RefractionAmount);
+
+				IDatasmithMaterialExpression& InvertedRefraction = OneMinus(RefractionAmountExpression); // Invert refaction to consider it Opacity factor for ThinMode (i.e. full refraction is Zero opacity)
+				OpacityExpression = &(OpacityColorExpression ? Multiply(*OpacityColorExpression, InvertedRefraction) : InvertedRefraction);
+			}
+
+			OpacityExpression = COMPOSE_OR_NULL(Desaturate, OpacityExpression);
+			Connect(PbrMaterialElement->GetOpacity(), OpacityExpression);
+
+			// Multiply specular by opacity to override default 0.5 and disable specular on fully transparent pixels(e.g. 'holes' in green leaves)
+			Connect(PbrMaterialElement->GetSpecular(), OpacityColorExpression ? Multiply(SpecularExpression, *OpacityColorExpression) : SpecularExpression);
+
+			// Need to turn transmittance into White color in fully transparent areas to disable absorption
+			// Actually fully transparent  - not transparent because of full refraction but with zero source Corona Opacity 
+			Connect(*ThinTranslucencyMaterialOutput->GetInput(0), (TranslucencyExpression && OpacityColorExpression) 
+				? &Lerp(Color(FLinearColor::White), *TranslucencyExpression, *OpacityColorExpression)
+				: TranslucencyExpression);
+
+			PbrMaterialElement->SetTwoSided(true); // Set two-sided for thin materials(usually used for thin leaves)
+		}
+		else if (bHasRefraction) // Refractive glass
+		{
+			IDatasmithMaterialExpression& InvertedRefraction = OneMinus(RefractionAmountMapExpression ? *RefractionAmountMapExpression : Scalar(MaterialProperties.RefractionAmount));
+
+			IDatasmithMaterialExpression& OpacityExpression = Desaturate(OpacityColorExpression ? Multiply(*OpacityColorExpression, InvertedRefraction) : InvertedRefraction);
+
+			IDatasmithMaterialExpression& RefractionAmountExpression = RefractionAmountMapExpression ? *RefractionAmountMapExpression : Scalar(MaterialProperties.RefractionAmount);
+
+			IDatasmithMaterialExpression& Multiply09 = Multiply(Scalar(0.9), RefractionAmountExpression);
+			IDatasmithMaterialExpression& FresnelOpacity = OneMinus(CalcIORComplex(Ior, 0, Multiply(Scalar(0.5), RefractionAmountExpression), Multiply09));
+
+			Connect(PbrMaterialElement->GetOpacity(), Desaturate(Multiply(FresnelOpacity, OpacityExpression)));
+
+			IDatasmithMaterialExpression& BaseIorExpression = Scalar(Ior, PARAM_NAME_BASEIOR);
+
+			Connect(PbrMaterialElement->GetRefraction(), PathTracingQualitySwitch(Lerp(Scalar(1.0f), BaseIorExpression, Fresnel()), BaseIorExpression));
+			Connect(PbrMaterialElement->GetSpecular(), OpacityColorExpression ? Multiply(SpecularExpression, *OpacityColorExpression) : SpecularExpression);
+
+			BaseColorExpression = &Multiply(Scalar(0.6), Power(Multiply09, Scalar(5.0)));
+			if (OpacityColorExpression)
+			{
+				BaseColorExpression = &Lerp(*BaseColorExpression, DiffuseColorExpression, *OpacityColorExpression);
+			}
+
+			Connect(PbrMaterialElement->GetMetallic(), ConvertTexmap(MaterialProperties.MetalnessTexmap));
+
+			// Compute volumetric properties for PathTracer
+			IDatasmithMaterialExpression* TransmittanceExpression = nullptr;
+			bool bVolumetricAbsorptionEnabled = MaterialProperties.AttenuationDistance > 0.0f;
+			if (bVolumetricAbsorptionEnabled && (!MaterialProperties.VolumetricAbsorptionColor.Equals(FLinearColor::White) || MaterialProperties.VolumetricAbsorptionTexmap.IsMapPresentAndEnabled()))
+			{
+				IDatasmithMaterialExpression* VolumetricAbsorptionTexmapExpression = BlendTexmapWithColor(MaterialProperties.VolumetricAbsorptionTexmap, MaterialProperties.VolumetricAbsorptionColor);
+				IDatasmithMaterialExpression& VolumetricAbsorptionColorExpression = VolumetricAbsorptionTexmapExpression ? *VolumetricAbsorptionTexmapExpression : Color(MaterialProperties.VolumetricAbsorptionColor);
+
+				TransmittanceExpression = &Power(VolumetricAbsorptionColorExpression, Divide(Scalar(100.0f), Scalar(MaterialProperties.AttenuationDistance, PARAM_NAME_ABSORPTIONDISTANCE)));
+			}
+
+			if (TransmittanceExpression)
+			{
+				IDatasmithMaterialExpressionGeneric* AbsorptionMediumMaterialOutput = PbrMaterialElement->AddMaterialExpression<IDatasmithMaterialExpressionGeneric>();
+				AbsorptionMediumMaterialOutput->SetExpressionName(TEXT("AbsorptionMediumMaterialOutput"));
+
+				// Need to turn transmittance into White color in fully transparent areas to disable absorption
+				Connect(*AbsorptionMediumMaterialOutput->GetInput(0),
+					Lerp(Color(FLinearColor::White), 
+						*TransmittanceExpression, 
+						OpacityColorExpression ? *OpacityColorExpression : Scalar(1.0f)));
+			}
+		}
+		else
+		{
+			Connect(PbrMaterialElement->GetMetallic(), ConvertTexmap(MaterialProperties.MetalnessTexmap));
+			Connect(PbrMaterialElement->GetSpecular(), Multiply(*OpacityColorExpression, SpecularExpression));
+			Connect(PbrMaterialElement->GetOpacity(), COMPOSE_OR_NULL(Desaturate, OpacityColorExpression));
+		}
+	}
+	else // Any other Opaque
+	{
+		Connect(PbrMaterialElement->GetMetallic(), ConvertTexmap(MaterialProperties.MetalnessTexmap));
+		Connect(PbrMaterialElement->GetSpecular(), SpecularExpression);
+	}
+	Connect(PbrMaterialElement->GetBaseColor(), BaseColorExpression);
+
+	if (MaterialProperties.BaseBumpTexmap.IsMapPresentAndEnabled())
+	{
+		ConvertState.DefaultTextureMode = EDatasmithTextureMode::Bump; // Will change to normal if we pass through a normal map texmap
+		ConvertState.bCanBake = false; // Current baking fails to produce proper normal maps
+	
+		IDatasmithMaterialExpression* BumpExpression = FDatasmithMaxTexmapToUEPbrUtils::MapOrValue(this, MaterialProperties.BaseBumpTexmap, PARAM_NAME_BUMPMAP, TOptional<FLinearColor>(), TOptional<float>());
+	
+		if ( BumpExpression )
+		{
+			BumpExpression->ConnectExpression(PbrMaterialElement->GetNormal());
+			BumpExpression->SetName(PARAM_NAME_BUMPMAP);
+		}
+
+		ConvertState.bCanBake = true;
+	}
+	
+	MaterialElement = PbrMaterialElement;
+}
diff --git a/Engine/Source/Programs/Enterprise/Datasmith/DatasmithMaxExporter/Private/MaxMaterialsToUEPbr/DatasmithMaxCoronaMaterialsToUEPbr.h b/Engine/Source/Programs/Enterprise/Datasmith/DatasmithMaxExporter/Private/MaxMaterialsToUEPbr/DatasmithMaxCoronaMaterialsToUEPbr.h
--- a/Engine/Source/Programs/Enterprise/Datasmith/DatasmithMaxExporter/Private/MaxMaterialsToUEPbr/DatasmithMaxCoronaMaterialsToUEPbr.h
+++ b/Engine/Source/Programs/Enterprise/Datasmith/DatasmithMaxExporter/Private/MaxMaterialsToUEPbr/DatasmithMaxCoronaMaterialsToUEPbr.h
@@ -4,7 +4,7 @@
 
 #include "DatasmithMaxMaterialsToUEPbr.h"
 
-class FDatasmithMaxCoronaMaterialsToUEPbr : public FDatasmithMaxMaterialsToUEPbr
+class FDatasmithMaxCoronaMaterialsToUEPbr : public FDatasmithMaxMaterialsToUEPbrExpressions
 {
 public:
 	FDatasmithMaxCoronaMaterialsToUEPbr();
@@ -20,7 +20,14 @@ public:
 	virtual void Convert( TSharedRef< IDatasmithScene > DatasmithScene, TSharedPtr< IDatasmithBaseMaterialElement >& MaterialElement, Mtl* Material, const TCHAR* AssetsPath ) override;
 };
 
-class FDatasmithMaxCoronaLightMaterialToUEPbr : public FDatasmithMaxMaterialsToUEPbrExpressions
+class FDatasmithMaxCoronaLightMaterialToUEPbr : public FDatasmithMaxCoronaMaterialsToUEPbr
+{
+public:
+	virtual bool IsSupported( Mtl* Material ) override;
+	virtual void Convert( TSharedRef< IDatasmithScene > DatasmithScene, TSharedPtr< IDatasmithBaseMaterialElement >& MaterialElement, Mtl* Material, const TCHAR* AssetsPath ) override;
+};
+
+class FDatasmithMaxCoronaPhysicalMaterialToUEPbr : public FDatasmithMaxCoronaMaterialsToUEPbr
 {
 public:
 	virtual bool IsSupported( Mtl* Material ) override;
diff --git a/Engine/Source/Programs/Enterprise/Datasmith/DatasmithMaxExporter/Private/MaxMaterialsToUEPbr/DatasmithMaxMaterialsToUEPbr.cpp b/Engine/Source/Programs/Enterprise/Datasmith/DatasmithMaxExporter/Private/MaxMaterialsToUEPbr/DatasmithMaxMaterialsToUEPbr.cpp
--- a/Engine/Source/Programs/Enterprise/Datasmith/DatasmithMaxExporter/Private/MaxMaterialsToUEPbr/DatasmithMaxMaterialsToUEPbr.cpp
+++ b/Engine/Source/Programs/Enterprise/Datasmith/DatasmithMaxExporter/Private/MaxMaterialsToUEPbr/DatasmithMaxMaterialsToUEPbr.cpp
@@ -76,6 +76,11 @@ FDatasmithMaxMaterialsToUEPbr* FDatasmithMaxMaterialsToUEPbrManager::GetMaterial
 		static FDatasmithMaxCoronaLightMaterialToUEPbr CoronaConverter = FDatasmithMaxCoronaLightMaterialToUEPbr();
 		MaterialConverter = &CoronaConverter;
 	}
+	else if ( MaterialClassID == CORONAPHYSICALMTLCLASS )
+	{
+		static FDatasmithMaxCoronaPhysicalMaterialToUEPbr CoronaConverter = FDatasmithMaxCoronaPhysicalMaterialToUEPbr();
+		MaterialConverter = &CoronaConverter;
+	}
 	else if (MaterialClassID == PHYSICALMATCLASS)
 	{
 		static FDatasmithMaxPhysicalMaterialToUEPbr PhysicalConverter = FDatasmithMaxPhysicalMaterialToUEPbr();
@@ -201,4 +206,4 @@ IDatasmithMaterialExpression* FDatasmithMaxMaterialsToUEPbr::ConvertTexmap( cons
 	}
 
 	return FDatasmithMaxTexmapToUEPbrUtils::ConvertTextureOutput( this, MaterialExpression, TexOutput );
-}
\ No newline at end of file
+}
diff --git a/Engine/Source/Programs/Enterprise/Datasmith/DatasmithMaxExporter/Private/MaxMaterialsToUEPbr/DatasmithMaxMaterialsToUEPbr.h b/Engine/Source/Programs/Enterprise/Datasmith/DatasmithMaxExporter/Private/MaxMaterialsToUEPbr/DatasmithMaxMaterialsToUEPbr.h
--- a/Engine/Source/Programs/Enterprise/Datasmith/DatasmithMaxExporter/Private/MaxMaterialsToUEPbr/DatasmithMaxMaterialsToUEPbr.h
+++ b/Engine/Source/Programs/Enterprise/Datasmith/DatasmithMaxExporter/Private/MaxMaterialsToUEPbr/DatasmithMaxMaterialsToUEPbr.h
@@ -101,28 +101,39 @@ class FDatasmithMaxMaterialsToUEPbrExpressions: public FDatasmithMaxMaterialsToU
 public:
 	TSharedPtr<IDatasmithUEPbrMaterialElement> GetMaterialElement();
 
-	IDatasmithMaterialExpressionScalar& Scalar(float Value);
-
-	IDatasmithMaterialExpressionColor& Color(const FLinearColor& Value);
+	// Specify ParameterName to have constant exposed as material parameter
+	IDatasmithMaterialExpression& Scalar(float Value, const TCHAR* ParameterName=nullptr);
+	IDatasmithMaterialExpression& Color(const FLinearColor& Value, const TCHAR* ParameterName=nullptr);
 
 	IDatasmithMaterialExpression* WeightTextureOrScalar(const DatasmithMaxTexmapParser::FMapParameter& TextureWeight, float Weight);
 
-	IDatasmithMaterialExpressionGeneric& Add(IDatasmithMaterialExpression& A, IDatasmithMaterialExpression& B);
+	IDatasmithMaterialExpression& Add(IDatasmithMaterialExpression& A, IDatasmithMaterialExpression& B);
+	IDatasmithMaterialExpression& Add(IDatasmithMaterialExpression& A, float B);
+
+	IDatasmithMaterialExpression& Subtract(IDatasmithMaterialExpression& A, IDatasmithMaterialExpression& B);
+	IDatasmithMaterialExpression& Subtract(IDatasmithMaterialExpression& A, float B);
+
+	IDatasmithMaterialExpression& Multiply(IDatasmithMaterialExpression& A, IDatasmithMaterialExpression& B);
+	IDatasmithMaterialExpression& Multiply(IDatasmithMaterialExpression& A, float B);
 
-	IDatasmithMaterialExpressionGeneric& Subtract(IDatasmithMaterialExpression& A, IDatasmithMaterialExpression& B);
+	IDatasmithMaterialExpression& Divide(IDatasmithMaterialExpression& A, IDatasmithMaterialExpression& B);
 
-	IDatasmithMaterialExpressionGeneric& Multiply(IDatasmithMaterialExpression& A, IDatasmithMaterialExpression& B);
+	IDatasmithMaterialExpression& Desaturate(IDatasmithMaterialExpression& A);
 
-	IDatasmithMaterialExpressionGeneric& Divide(IDatasmithMaterialExpression& A, IDatasmithMaterialExpression& B);
+	IDatasmithMaterialExpression& Power(IDatasmithMaterialExpression& A, IDatasmithMaterialExpression& B);
+	IDatasmithMaterialExpression& Power(IDatasmithMaterialExpression& A, float Exp);
 
-	IDatasmithMaterialExpressionGeneric& Desaturate(IDatasmithMaterialExpression& A);
+	IDatasmithMaterialExpression& Lerp(IDatasmithMaterialExpression& A, IDatasmithMaterialExpression& B, IDatasmithMaterialExpression& Alpha);
 
-	IDatasmithMaterialExpressionGeneric& Power(IDatasmithMaterialExpression& A, IDatasmithMaterialExpression& B);
+	IDatasmithMaterialExpression& Fresnel(IDatasmithMaterialExpression* Exponent=nullptr, IDatasmithMaterialExpression* BaseReflectFraction=nullptr); // Any input can be null
 
-	IDatasmithMaterialExpressionGeneric& Lerp(IDatasmithMaterialExpression& A, IDatasmithMaterialExpression& B, IDatasmithMaterialExpression& Alpha);
+	IDatasmithMaterialExpression& PathTracingQualitySwitch(IDatasmithMaterialExpression& Normal, IDatasmithMaterialExpression& PathTraced);
 
-	IDatasmithMaterialExpressionGeneric& Fresnel(IDatasmithMaterialExpression* Exponent=nullptr, IDatasmithMaterialExpression* BaseReflectFraction=nullptr); // Any input can be null
+	IDatasmithMaterialExpression* BlendTexmapWithColor(const DatasmithMaxTexmapParser::FMapParameter& Texmap, const FLinearColor& InColor, const TCHAR* ColorName = nullptr, const TCHAR* WeightName = nullptr);
+	IDatasmithMaterialExpression* BlendTexmapWithScalar(const DatasmithMaxTexmapParser::FMapParameter& Texmap, float InScalar, const TCHAR* ScalarName = nullptr, const TCHAR* WeightName = nullptr);
 
+	// Multiply RGB without multiplying A
+	FLinearColor ScaleRGB(const FLinearColor& C, const float X) { return C * FLinearColor(X, X, X);}
 
 	IDatasmithMaterialExpression* ApplyWeightExpression(IDatasmithMaterialExpression* ValueExpression, IDatasmithMaterialExpression* WeightExpression);
 
@@ -138,7 +149,7 @@ public:
 	IDatasmithMaterialExpression* TextureOrScalar(const TCHAR* Name,
 	                                              const DatasmithMaxTexmapParser::FMapParameter& Map, float Value);
 
-	IDatasmithMaterialExpressionGeneric& OneMinus(IDatasmithMaterialExpression& Expression);
+	IDatasmithMaterialExpression& OneMinus(IDatasmithMaterialExpression& Expression);
 };
 
 // Utility macros to simplify material expression composition:
diff --git a/Engine/Source/Programs/Enterprise/Datasmith/DatasmithMaxExporter/Private/MaxMaterialsToUEPbr/DatasmithMaxPhysicalMaterialToUEPbr.cpp b/Engine/Source/Programs/Enterprise/Datasmith/DatasmithMaxExporter/Private/MaxMaterialsToUEPbr/DatasmithMaxPhysicalMaterialToUEPbr.cpp
--- a/Engine/Source/Programs/Enterprise/Datasmith/DatasmithMaxExporter/Private/MaxMaterialsToUEPbr/DatasmithMaxPhysicalMaterialToUEPbr.cpp
+++ b/Engine/Source/Programs/Enterprise/Datasmith/DatasmithMaxExporter/Private/MaxMaterialsToUEPbr/DatasmithMaxPhysicalMaterialToUEPbr.cpp
@@ -21,14 +21,18 @@ TSharedPtr<IDatasmithUEPbrMaterialElement> FDatasmithMaxMaterialsToUEPbrExpressi
 	return ConvertState.MaterialElement;
 }
 
-IDatasmithMaterialExpressionScalar& FDatasmithMaxMaterialsToUEPbrExpressions::Scalar(float Value)
+IDatasmithMaterialExpression& FDatasmithMaxMaterialsToUEPbrExpressions::Scalar(float Value, const TCHAR* ParameterName)
 {
 	IDatasmithMaterialExpressionScalar* ScalarExpression = GetMaterialElement()->AddMaterialExpression<IDatasmithMaterialExpressionScalar>();
 	ScalarExpression->GetScalar() = Value;
+	if (ParameterName)
+	{
+		ScalarExpression->SetName(ParameterName);
+	}
 	return *ScalarExpression;
 }
 
-IDatasmithMaterialExpressionColor& FDatasmithMaxMaterialsToUEPbrExpressions::Color(const FLinearColor& Value)
+IDatasmithMaterialExpression& FDatasmithMaxMaterialsToUEPbrExpressions::Color(const FLinearColor& Value, const TCHAR* ParameterName)
 {
 	IDatasmithMaterialExpressionColor* Result = GetMaterialElement()->AddMaterialExpression<IDatasmithMaterialExpressionColor>();
 	Result->GetColor() = Value;
@@ -51,7 +55,7 @@ IDatasmithMaterialExpression* FDatasmithMaxMaterialsToUEPbrExpressions::WeightTe
 	return nullptr;
 }
 
-IDatasmithMaterialExpressionGeneric& FDatasmithMaxMaterialsToUEPbrExpressions::Add(IDatasmithMaterialExpression& A,
+IDatasmithMaterialExpression& FDatasmithMaxMaterialsToUEPbrExpressions::Add(IDatasmithMaterialExpression& A,
 	IDatasmithMaterialExpression& B)
 {
 	IDatasmithMaterialExpressionGeneric* Result = GetMaterialElement()->AddMaterialExpression<IDatasmithMaterialExpressionGeneric>();
@@ -62,7 +66,12 @@ IDatasmithMaterialExpressionGeneric& FDatasmithMaxMaterialsToUEPbrExpressions::A
 	return *Result;
 }
 
-IDatasmithMaterialExpressionGeneric& FDatasmithMaxMaterialsToUEPbrExpressions::Subtract(IDatasmithMaterialExpression& A,
+IDatasmithMaterialExpression& FDatasmithMaxMaterialsToUEPbrExpressions::Add(IDatasmithMaterialExpression& A, float B)
+{
+	return Add(A, Scalar(B));
+}
+
+IDatasmithMaterialExpression& FDatasmithMaxMaterialsToUEPbrExpressions::Subtract(IDatasmithMaterialExpression& A,
 	IDatasmithMaterialExpression& B)
 {
 	IDatasmithMaterialExpressionGeneric* Result = GetMaterialElement()->AddMaterialExpression<IDatasmithMaterialExpressionGeneric>();
@@ -73,7 +82,12 @@ IDatasmithMaterialExpressionGeneric& FDatasmithMaxMaterialsToUEPbrExpressions::S
 	return *Result;
 }
 
-IDatasmithMaterialExpressionGeneric& FDatasmithMaxMaterialsToUEPbrExpressions::Multiply(IDatasmithMaterialExpression& A, IDatasmithMaterialExpression& B)
+IDatasmithMaterialExpression& FDatasmithMaxMaterialsToUEPbrExpressions::Subtract(IDatasmithMaterialExpression& A, float B)
+{
+	return Subtract(A, Scalar(B));
+}
+
+IDatasmithMaterialExpression& FDatasmithMaxMaterialsToUEPbrExpressions::Multiply(IDatasmithMaterialExpression& A, IDatasmithMaterialExpression& B)
 {
 	IDatasmithMaterialExpressionGeneric* Result = GetMaterialElement()->AddMaterialExpression<IDatasmithMaterialExpressionGeneric>();
 	Result->SetExpressionName(TEXT("Multiply"));
@@ -83,7 +97,12 @@ IDatasmithMaterialExpressionGeneric& FDatasmithMaxMaterialsToUEPbrExpressions::M
 	return *Result;
 }
 
-IDatasmithMaterialExpressionGeneric& FDatasmithMaxMaterialsToUEPbrExpressions::Divide(IDatasmithMaterialExpression& A,
+IDatasmithMaterialExpression& FDatasmithMaxMaterialsToUEPbrExpressions::Multiply(IDatasmithMaterialExpression& A, float B)
+{
+	return Multiply(A, Scalar(B));
+}
+
+IDatasmithMaterialExpression& FDatasmithMaxMaterialsToUEPbrExpressions::Divide(IDatasmithMaterialExpression& A,
 	IDatasmithMaterialExpression& B)
 {
 	IDatasmithMaterialExpressionGeneric* Result = GetMaterialElement()->AddMaterialExpression<IDatasmithMaterialExpressionGeneric>();
@@ -94,7 +113,7 @@ IDatasmithMaterialExpressionGeneric& FDatasmithMaxMaterialsToUEPbrExpressions::D
 	return *Result;
 }
 
-IDatasmithMaterialExpressionGeneric& FDatasmithMaxMaterialsToUEPbrExpressions::Desaturate(IDatasmithMaterialExpression& A)
+IDatasmithMaterialExpression& FDatasmithMaxMaterialsToUEPbrExpressions::Desaturate(IDatasmithMaterialExpression& A)
 {
 	IDatasmithMaterialExpressionGeneric* Result = GetMaterialElement()->AddMaterialExpression<IDatasmithMaterialExpressionGeneric>();
 	Result->SetExpressionName(TEXT("Desaturation"));
@@ -103,7 +122,7 @@ IDatasmithMaterialExpressionGeneric& FDatasmithMaxMaterialsToUEPbrExpressions::D
 	return *Result;
 }
 
-IDatasmithMaterialExpressionGeneric& FDatasmithMaxMaterialsToUEPbrExpressions::Power(IDatasmithMaterialExpression& A,
+IDatasmithMaterialExpression& FDatasmithMaxMaterialsToUEPbrExpressions::Power(IDatasmithMaterialExpression& A,
 	IDatasmithMaterialExpression& B)
 {
 	IDatasmithMaterialExpressionGeneric* Result = GetMaterialElement()->AddMaterialExpression<IDatasmithMaterialExpressionGeneric>();
@@ -114,7 +133,22 @@ IDatasmithMaterialExpressionGeneric& FDatasmithMaxMaterialsToUEPbrExpressions::P
 	return *Result;
 }
 
-IDatasmithMaterialExpressionGeneric& FDatasmithMaxMaterialsToUEPbrExpressions::Lerp(IDatasmithMaterialExpression& A,
+IDatasmithMaterialExpression& FDatasmithMaxMaterialsToUEPbrExpressions::Power(IDatasmithMaterialExpression& A, float Exp)
+{
+	TSharedRef<IDatasmithKeyValueProperty> ExpProperty = FDatasmithSceneFactory::CreateKeyValueProperty(TEXT("ConstExponent"));
+	ExpProperty->SetPropertyType(EDatasmithKeyValuePropertyType::Float);
+	ExpProperty->SetValue(*LexToString(Exp));
+
+	IDatasmithMaterialExpressionGeneric* Result = GetMaterialElement()->AddMaterialExpression<IDatasmithMaterialExpressionGeneric>();
+	Result->SetExpressionName(TEXT("Power"));
+
+	A.ConnectExpression(*Result->GetInput(0));
+	Result->AddProperty(ExpProperty);
+	return *Result;
+}
+
+
+IDatasmithMaterialExpression& FDatasmithMaxMaterialsToUEPbrExpressions::Lerp(IDatasmithMaterialExpression& A,
 	IDatasmithMaterialExpression& B, IDatasmithMaterialExpression& Alpha)
 {
 	IDatasmithMaterialExpressionGeneric* Result = GetMaterialElement()->AddMaterialExpression<IDatasmithMaterialExpressionGeneric>();
@@ -126,7 +160,7 @@ IDatasmithMaterialExpressionGeneric& FDatasmithMaxMaterialsToUEPbrExpressions::L
 	return *Result;
 }
 
-IDatasmithMaterialExpressionGeneric& FDatasmithMaxMaterialsToUEPbrExpressions::Fresnel(IDatasmithMaterialExpression* Exponent,
+IDatasmithMaterialExpression& FDatasmithMaxMaterialsToUEPbrExpressions::Fresnel(IDatasmithMaterialExpression* Exponent,
 	IDatasmithMaterialExpression* BaseReflectFraction)
 {
 	IDatasmithMaterialExpressionGeneric* Fresnel = GetMaterialElement()->AddMaterialExpression<IDatasmithMaterialExpressionGeneric>();
@@ -137,6 +171,30 @@ IDatasmithMaterialExpressionGeneric& FDatasmithMaxMaterialsToUEPbrExpressions::F
 	return *Fresnel;
 }
 
+
+IDatasmithMaterialExpression& FDatasmithMaxMaterialsToUEPbrExpressions::PathTracingQualitySwitch(
+	IDatasmithMaterialExpression& Normal, IDatasmithMaterialExpression& PathTraced)
+{
+	IDatasmithMaterialExpressionGeneric* Result = GetMaterialElement()->AddMaterialExpression<IDatasmithMaterialExpressionGeneric>();
+	Result->SetExpressionName(TEXT("PathTracingQualitySwitch"));
+
+	Connect(*Result->GetInput(0), Normal);
+	Connect(*Result->GetInput(1), PathTraced);
+	return *Result;
+}
+
+IDatasmithMaterialExpression* FDatasmithMaxMaterialsToUEPbrExpressions::BlendTexmapWithColor(const DatasmithMaxTexmapParser::FMapParameter& Texmap, const FLinearColor& InColor, const TCHAR* ColorName, const TCHAR* WeightName)
+{
+	IDatasmithMaterialExpression* TexmapExpression = ConvertTexmap(Texmap);
+	return TexmapExpression ? &Lerp(Color(InColor, ColorName), *TexmapExpression, Scalar(Texmap.Weight, WeightName)) : nullptr;
+};
+
+IDatasmithMaterialExpression* FDatasmithMaxMaterialsToUEPbrExpressions::BlendTexmapWithScalar(const DatasmithMaxTexmapParser::FMapParameter& Texmap, float InScalar, const TCHAR* ScalarName, const TCHAR* WeightName)
+{
+	IDatasmithMaterialExpression* TexmapExpression = ConvertTexmap(Texmap);
+	return TexmapExpression ? &Lerp(Scalar(InScalar, ScalarName), *TexmapExpression, Scalar(Texmap.Weight, WeightName)) : nullptr;
+};
+
 IDatasmithMaterialExpression* FDatasmithMaxMaterialsToUEPbrExpressions::ApplyWeightExpression(
 	IDatasmithMaterialExpression* ValueExpression, IDatasmithMaterialExpression* WeightExpression)
 {
@@ -181,11 +239,9 @@ IDatasmithMaterialExpression& FDatasmithMaxMaterialsToUEPbrExpressions::CalcIORC
 	IDatasmithMaterialExpression& RpDen = Add(RpDenPre, Scalar(1.0));
 	IDatasmithMaterialExpression& Rp = Divide(RpNum, RpDen);
 
-	IDatasmithMaterialExpression& ResPre = Add(Rp, Rs);
-	IDatasmithMaterialExpression& ResMul = Multiply(ResPre, Scalar(0.5));
-	IDatasmithMaterialExpression& Res = Power(ResMul, Scalar(0.5));
+	IDatasmithMaterialExpression& Res = Multiply(Add(Rp, Rs), Scalar(0.5));
 
-	return Lerp(ToBeConnected0, ToBeConnected90, Res);
+	return Lerp(ToBeConnected0, ToBeConnected90, Power(Res, Scalar(0.5)));
 }
 
 void FDatasmithMaxMaterialsToUEPbrExpressions::Connect(IDatasmithExpressionInput& Input, IDatasmithMaterialExpression& ValueExpression)
@@ -215,7 +271,7 @@ IDatasmithMaterialExpression* FDatasmithMaxMaterialsToUEPbrExpressions::TextureO
 	return FDatasmithMaxTexmapToUEPbrUtils::MapOrValue(this, Map, Name, TOptional<FLinearColor>(), Value);
 }
 
-IDatasmithMaterialExpressionGeneric& FDatasmithMaxMaterialsToUEPbrExpressions::OneMinus(
+IDatasmithMaterialExpression& FDatasmithMaxMaterialsToUEPbrExpressions::OneMinus(
 	IDatasmithMaterialExpression& Expression)
 {
 	IDatasmithMaterialExpressionGeneric* Result = GetMaterialElement()->AddMaterialExpression< IDatasmithMaterialExpressionGeneric >();
diff --git a/Engine/Source/Programs/Enterprise/Datasmith/DatasmithMaxExporter/Private/MaxMaterialsToUEPbr/DatasmithMaxTexmapToUEPbr.cpp b/Engine/Source/Programs/Enterprise/Datasmith/DatasmithMaxExporter/Private/MaxMaterialsToUEPbr/DatasmithMaxTexmapToUEPbr.cpp
--- a/Engine/Source/Programs/Enterprise/Datasmith/DatasmithMaxExporter/Private/MaxMaterialsToUEPbr/DatasmithMaxTexmapToUEPbr.cpp
+++ b/Engine/Source/Programs/Enterprise/Datasmith/DatasmithMaxExporter/Private/MaxMaterialsToUEPbr/DatasmithMaxTexmapToUEPbr.cpp
@@ -701,7 +701,7 @@ IDatasmithMaterialExpression* FDatasmithMaxTexmapToUEPbrUtils::MapOrValue( FData
 				MapWeightLerp->SetExpressionName( TEXT("LinearInterpolate") );
 
 				IDatasmithMaterialExpressionScalar* MapWeight = MaxMaterialToUEPbr->ConvertState.MaterialElement->AddMaterialExpression< IDatasmithMaterialExpressionScalar >();
-				MapWeight->SetName( TEXT("Map Weight") );
+				MapWeight->SetName( *(FString(ParameterName) + TEXT("Map Weight")) );
 				MapWeight->GetScalar() = MapParameter.Weight;
 
 				if ( ValueExpression )
diff --git a/Engine/Source/Programs/Enterprise/Datasmith/DatasmithMaxExporter/Public/DatasmithMaxClassIDs.h b/Engine/Source/Programs/Enterprise/Datasmith/DatasmithMaxExporter/Public/DatasmithMaxClassIDs.h
--- a/Engine/Source/Programs/Enterprise/Datasmith/DatasmithMaxExporter/Public/DatasmithMaxClassIDs.h
+++ b/Engine/Source/Programs/Enterprise/Datasmith/DatasmithMaxExporter/Public/DatasmithMaxClassIDs.h
@@ -66,8 +66,11 @@
 #define  CORONAMATCLASS			Class_ID(0x70be6506,0x448931dd)
 #define  CORONALAYERMATCLASS	Class_ID(0x65486584,0x8425554e)
 #define  CORONALIGHTMATCLASS	Class_ID(0x53b7717a,0x500d5093)
+#define  CORONAPHYSICALMTLCLASS	Class_ID(0x6912ab89,0x87151720) 
 #define	 PHYSICALMATCLASS		Class_ID(0x3d6b1cec,0xdeadc001)
 
+
+
 #define  VRAYMATCLASS			Class_ID(0x37bf3f2f, 0x7034695c)
 #define  VRAYBLENDMATCLASS		Class_ID(0x3da0041b, 0xf436e31)
 #define  VRAYLIGHTMATCLASS		Class_ID(0x7ccf263e, 0x3f5b39b9)
