Datasmith Revit: Re-enable autosync

diff --git a/Engine/Source/Programs/Enterprise/Datasmith/DatasmithRevitExporter/Private/DatasmithRevitApplication.cs b/Engine/Source/Programs/Enterprise/Datasmith/DatasmithRevitExporter/Private/DatasmithRevitApplication.cs
--- a/Engine/Source/Programs/Enterprise/Datasmith/DatasmithRevitExporter/Private/DatasmithRevitApplication.cs
+++ b/Engine/Source/Programs/Enterprise/Datasmith/DatasmithRevitExporter/Private/DatasmithRevitApplication.cs
@@ -23,19 +23,18 @@ namespace DatasmithRevitExporter
 		private EventHandler<DocumentClosingEventArgs> DocumentClosingHandler;
 		private EventHandler<ViewActivatedEventArgs> ViewActivatedHandler;
 
-		//PushButton AutoSyncPushButton;
+		PushButton AutoSyncPushButton;
 		PushButton SyncPushButton;
 
-		//BitmapImage AutoSyncIconOn_Small;
-		//BitmapImage AutoSyncIconOn_Large;
-		//BitmapImage AutoSyncIconOff_Small;
-		//BitmapImage AutoSyncIconOff_Large;
+		BitmapImage AutoSyncIconOn_Small;
+		BitmapImage AutoSyncIconOn_Large;
+		BitmapImage AutoSyncIconOff_Small;
+		BitmapImage AutoSyncIconOff_Large;
 
 		public object Properties { get; private set; }
 
 		public static DatasmithRevitApplication Instance { get; private set; }
 
-#if false
 		public void SetAutoSyncButtonToggled(bool bToggled)
 		{
 			if (bToggled)
@@ -50,7 +49,7 @@ namespace DatasmithRevitExporter
 			}
 			SyncPushButton.Enabled = !bToggled;
 		}
-#endif
+
 		// Implement the interface to execute some tasks when Revit starts.
 		public Result OnStartup(
 			UIControlledApplication InApplication // handle to the application being started
@@ -70,13 +69,13 @@ namespace DatasmithRevitExporter
 			string AssemblyPath = Assembly.GetExecutingAssembly().Location;
 			PushButtonData ExportButtonData = new PushButtonData("Export3DView", DatasmithRevitResources.Strings.ButtonExport3DView, AssemblyPath, "DatasmithRevitExporter.DatasmithExportRevitCommand");
 			PushButtonData SyncButtonData = new PushButtonData("Sync3DView", DatasmithRevitResources.Strings.ButtonSync, AssemblyPath, "DatasmithRevitExporter.DatasmithSyncRevitCommand");
-			//PushButtonData AutoSyncButtonData = new PushButtonData("AutoSync3DView", DatasmithRevitResources.Strings.ButtonAutoSync, AssemblyPath, "DatasmithRevitExporter.DatasmithAutoSyncRevitCommand");
+			PushButtonData AutoSyncButtonData = new PushButtonData("AutoSync3DView", DatasmithRevitResources.Strings.ButtonAutoSync, AssemblyPath, "DatasmithRevitExporter.DatasmithAutoSyncRevitCommand");
 			PushButtonData ManageConnectionsButtonData = new PushButtonData("Connections", DatasmithRevitResources.Strings.ButtonConnections, AssemblyPath, "DatasmithRevitExporter.DatasmithManageConnectionsRevitCommand");
 			PushButtonData SettingsButtonData = new PushButtonData("Settings", DatasmithRevitResources.Strings.ButtonSettings, AssemblyPath, "DatasmithRevitExporter.DatasmithShowSettingsRevitCommand");
 			PushButtonData LogButtonData = new PushButtonData("Messages", DatasmithRevitResources.Strings.ButtonMessages, AssemblyPath, "DatasmithRevitExporter.DatasmithShowMessagesRevitCommand");
 
 			SyncPushButton = DirectLinkRibbonPanel.AddItem(SyncButtonData) as PushButton;
-			//AutoSyncPushButton = DirectLinkRibbonPanel.AddItem(AutoSyncButtonData) as PushButton;
+			AutoSyncPushButton = DirectLinkRibbonPanel.AddItem(AutoSyncButtonData) as PushButton;
 			PushButton ManageConnectionsButton = DirectLinkRibbonPanel.AddItem(ManageConnectionsButtonData) as PushButton;
 			PushButton ExportPushButton = FileExportRibbonPanel.AddItem(ExportButtonData) as PushButton;
 			PushButton ShowLogButton = DatasmithRibbonPanel.AddItem(LogButtonData) as PushButton;
@@ -94,14 +93,14 @@ namespace DatasmithRevitExporter
 
 			DatasmithIconBase = Path.Combine(Path.GetDirectoryName(AssemblyPath), "DatasmithAutoSyncIcon");
 
-			//AutoSyncIconOn_Small = new BitmapImage(new Uri(DatasmithIconBase + "On16.png"));
-			//AutoSyncIconOn_Large = new BitmapImage(new Uri(DatasmithIconBase + "On32.png"));
-			//AutoSyncIconOff_Small = new BitmapImage(new Uri(DatasmithIconBase + "Off16.png"));
-			//AutoSyncIconOff_Large = new BitmapImage(new Uri(DatasmithIconBase + "Off32.png"));
+			AutoSyncIconOn_Small = new BitmapImage(new Uri(DatasmithIconBase + "On16.png"));
+			AutoSyncIconOn_Large = new BitmapImage(new Uri(DatasmithIconBase + "On32.png"));
+			AutoSyncIconOff_Small = new BitmapImage(new Uri(DatasmithIconBase + "Off16.png"));
+			AutoSyncIconOff_Large = new BitmapImage(new Uri(DatasmithIconBase + "Off32.png"));
 
-			//AutoSyncPushButton.Image = AutoSyncIconOn_Small;
-			//AutoSyncPushButton.LargeImage = AutoSyncIconOn_Large;
-			//AutoSyncPushButton.ToolTip = DatasmithRevitResources.Strings.ButtonAutoSyncHint;
+			AutoSyncPushButton.Image = AutoSyncIconOn_Small;
+			AutoSyncPushButton.LargeImage = AutoSyncIconOn_Large;
+			AutoSyncPushButton.ToolTip = DatasmithRevitResources.Strings.ButtonAutoSyncHint;
 
 			DatasmithIconBase = Path.Combine(Path.GetDirectoryName(AssemblyPath), "DatasmithManageConnectionsIcon");
 			ManageConnectionsButton.Image = new BitmapImage(new Uri(DatasmithIconBase + "16.png"));
diff --git a/Engine/Source/Programs/Enterprise/Datasmith/DatasmithRevitExporter/Private/DatasmithRevitCommand.cs b/Engine/Source/Programs/Enterprise/Datasmith/DatasmithRevitExporter/Private/DatasmithRevitCommand.cs
--- a/Engine/Source/Programs/Enterprise/Datasmith/DatasmithRevitExporter/Private/DatasmithRevitCommand.cs
+++ b/Engine/Source/Programs/Enterprise/Datasmith/DatasmithRevitExporter/Private/DatasmithRevitCommand.cs
@@ -126,7 +126,6 @@ namespace DatasmithRevitExporter
 		}
 	}
 
-#if false // AutoSync is temporary disabled
 	[Transaction(TransactionMode.Manual)]
 	public class DatasmithAutoSyncRevitCommand : DatasmithRevitCommand
 	{
@@ -139,7 +138,6 @@ namespace DatasmithRevitExporter
 			return Result.Succeeded;
 		}
 	}
-#endif
 
 	// Add-in external command Export to Unreal Datasmith. 
 	[Transaction(TransactionMode.Manual)]
diff --git a/Engine/Source/Programs/Enterprise/Datasmith/DatasmithRevitExporter/Private/DatasmithRevitUtils.cs b/Engine/Source/Programs/Enterprise/Datasmith/DatasmithRevitExporter/Private/DatasmithRevitUtils.cs
--- a/Engine/Source/Programs/Enterprise/Datasmith/DatasmithRevitExporter/Private/DatasmithRevitUtils.cs
+++ b/Engine/Source/Programs/Enterprise/Datasmith/DatasmithRevitExporter/Private/DatasmithRevitUtils.cs
@@ -66,40 +66,44 @@ namespace DatasmithRevitExporter
 			{
 				foreach (Parameter Parameter in Parameters)
 				{
-					if (Parameter.HasValue)
+					try
 					{
-						if (Settings != null && !Settings.MatchParameterByMetadata(Parameter))
+						if (Parameter.HasValue)
 						{
-							continue; // Skip export of this param
-						}
+							if (Settings != null && !Settings.MatchParameterByMetadata(Parameter))
+							{
+								continue; // Skip export of this param
+							}
 
-						string ParameterValue = Parameter.AsValueString();
+							string ParameterValue = Parameter.AsValueString();
 
-						if (string.IsNullOrEmpty(ParameterValue))
-						{
-							switch (Parameter.StorageType)
+							if (string.IsNullOrEmpty(ParameterValue))
 							{
-								case StorageType.Integer:
-								ParameterValue = Parameter.AsInteger().ToString();
-								break;
-								case StorageType.Double:
-								ParameterValue = Parameter.AsDouble().ToString();
-								break;
-								case StorageType.String:
-								ParameterValue = Parameter.AsString();
-								break;
-								case StorageType.ElementId:
-								ParameterValue = Parameter.AsElementId().ToString();
-								break;
+								switch (Parameter.StorageType)
+								{
+									case StorageType.Integer:
+									ParameterValue = Parameter.AsInteger().ToString();
+									break;
+									case StorageType.Double:
+									ParameterValue = Parameter.AsDouble().ToString();
+									break;
+									case StorageType.String:
+									ParameterValue = Parameter.AsString();
+									break;
+									case StorageType.ElementId:
+									ParameterValue = Parameter.AsElementId().ToString();
+									break;
+								}
 							}
-						}
 
-						if (!string.IsNullOrEmpty(ParameterValue))
-						{
-							string MetadataKey = InMetadataPrefix + Parameter.Definition.Name;
-							ElementMetaData.AddPropertyString(MetadataKey, ParameterValue);
+							if (!string.IsNullOrEmpty(ParameterValue))
+							{
+								string MetadataKey = InMetadataPrefix + Parameter.Definition.Name;
+								ElementMetaData.AddPropertyString(MetadataKey, ParameterValue);
+							}
 						}
 					}
+					catch { }
 				}
 			}
 		}
