Add logic in DatasmithCADWorker to kill itself when parent process does not exist anymore

diff --git a/Engine/Source/Programs/Enterprise/Datasmith/DatasmithCADWorker/Private/DatasmithCADWorker.cpp b/Engine/Source/Programs/Enterprise/Datasmith/DatasmithCADWorker/Private/DatasmithCADWorker.cpp
--- a/Engine/Source/Programs/Enterprise/Datasmith/DatasmithCADWorker/Private/DatasmithCADWorker.cpp
+++ b/Engine/Source/Programs/Enterprise/Datasmith/DatasmithCADWorker/Private/DatasmithCADWorker.cpp
@@ -109,7 +109,7 @@ INT32_MAIN_INT32_ARGC_TCHAR_ARGV()
 
 	FCoreDelegates::OnExit.Broadcast();
 
-	FPlatformMisc::RequestExit(false);
+	FPlatformMisc::RequestExit(true);
 
 	return Result;
 }
