DatasmtihCADWorker still initialize Old CADLibrary but not the new one

diff --git a/Engine/Source/Programs/Enterprise/Datasmith/DatasmithCADWorker/DatasmithCADWorker.Build.cs b/Engine/Source/Programs/Enterprise/Datasmith/DatasmithCADWorker/DatasmithCADWorker.Build.cs
--- a/Engine/Source/Programs/Enterprise/Datasmith/DatasmithCADWorker/DatasmithCADWorker.Build.cs
+++ b/Engine/Source/Programs/Enterprise/Datasmith/DatasmithCADWorker/DatasmithCADWorker.Build.cs
@@ -23,7 +23,5 @@ public class DatasmithCADWorker : ModuleRules
 				"CADTools",
 			}
 		);
-
-		PublicDelayLoadDLLs.Add("kernel_io.dll");
 	}
 }
diff --git a/Engine/Source/Programs/Enterprise/Datasmith/DatasmithCADWorker/Private/DatasmithCADWorker.cpp b/Engine/Source/Programs/Enterprise/Datasmith/DatasmithCADWorker/Private/DatasmithCADWorker.cpp
--- a/Engine/Source/Programs/Enterprise/Datasmith/DatasmithCADWorker/Private/DatasmithCADWorker.cpp
+++ b/Engine/Source/Programs/Enterprise/Datasmith/DatasmithCADWorker/Private/DatasmithCADWorker.cpp
@@ -6,15 +6,20 @@
 #include "RequiredProgramMainCPPInclude.h"
 #include "CADInterfacesModule.h"
 #include "CADToolsModule.h"
+
+#ifdef USE_KERNEL_IO_SDK
 #include "CoreTechTypes.h"
+#endif
+#ifdef USE_TECHSOFT_SDK
+#include "TechSoftInterface.h"
+#endif
 
 #include "Modules/ModuleManager.h"
 
 IMPLEMENT_APPLICATION(DatasmithCADWorker, "DatasmithCADWorker");
 DEFINE_LOG_CATEGORY(LogDatasmithCADWorker);
 
-#define EXIT_MISSING_CORETECH_MODULE 2
-
+#define EXIT_MISSING_CAD_MODULE 2
 
 void GetParameter(int32 Argc, TCHAR* Argv[], const FString& InParam, FString& OutValue)
 {
@@ -47,10 +52,11 @@ int32 Main(int32 Argc, TCHAR * Argv[])
 {
 	UE_SET_LOG_VERBOSITY(LogDatasmithCADWorker, Verbose);
 
-#ifndef USE_KERNEL_IO_SDK
-	UE_LOG(LogDatasmithCADWorker, Error, TEXT("Missing CoreTech module. DatasmithCADWorker is not functional."));
-	return EXIT_MISSING_CORETECH_MODULE;
+#if !defined(USE_KERNEL_IO_SDK) && !defined(USE_TECHSOFT_SDK)
+	UE_LOG(LogDatasmithCADWorker, Error, TEXT("Missing CAD module. DatasmithCADWorker is not functional."));
+	return EXIT_MISSING_CAD_MODULE;
 #else
+
 	FString ServerPID;
 	FString ServerPort;
 	FString CacheDirectory;
@@ -62,7 +68,17 @@ int32 Main(int32 Argc, TCHAR * Argv[])
 	GetParameter(Argc, Argv, "-CacheVersion", CacheVersion);
 	GetParameter(Argc, Argv, "-EnginePluginsDir", EnginePluginsPath);
 
+#ifdef USE_KERNEL_IO_SDK
 	CADLibrary::InitializeCoreTechInterface();
+#endif
+
+#ifdef USE_TECHSOFT_SDK
+	if (!CADLibrary::TechSoftInterface::TECHSOFT_InitializeKernel(*EnginePluginsPath))
+	{
+		UE_LOG(LogDatasmithCADWorker, Error, TEXT("TechSoft interface cannot be initialized. CADInterfaces module is not available."));
+		return EXIT_FAILURE;
+	}
+#endif
 
 	if (ICADInterfacesModule::Get().GetAvailability() != ECADInterfaceAvailability::Available)
 	{
