[3ds Max] DirectLink: Containers - preserve hierarchy

- fixed scene reset/new/load DirectLink name
- (internal)flag to use viewport mesh for Sync/AutoSync to prevent GeometryChanged events appearing. Testcase - use Optimize/MeshSmooth/TurboSmooth modifier and have render/viewport mesh different quality. This would cause infinite update before.

diff --git a/Engine/Source/Programs/Enterprise/Datasmith/DatasmithMaxExporter/Private/DatasmithMaxDirectLinkMaxScript.cpp b/Engine/Source/Programs/Enterprise/Datasmith/DatasmithMaxExporter/Private/DatasmithMaxDirectLinkMaxScript.cpp
--- a/Engine/Source/Programs/Enterprise/Datasmith/DatasmithMaxExporter/Private/DatasmithMaxDirectLinkMaxScript.cpp
+++ b/Engine/Source/Programs/Enterprise/Datasmith/DatasmithMaxExporter/Private/DatasmithMaxDirectLinkMaxScript.cpp
@@ -147,7 +147,7 @@ Value* Reset_cf(Value** arg_list, int count)
 		return bool_result(false);
 	}
 
-	GetExporter()->Reset();
+	GetExporter()->ResetSceneTracking();
 	return bool_result(true);
 }
 
diff --git a/Engine/Source/Programs/Enterprise/Datasmith/DatasmithMaxExporter/Private/DatasmithMaxDirectLinkMesh.cpp b/Engine/Source/Programs/Enterprise/Datasmith/DatasmithMaxExporter/Private/DatasmithMaxDirectLinkMesh.cpp
--- a/Engine/Source/Programs/Enterprise/Datasmith/DatasmithMaxExporter/Private/DatasmithMaxDirectLinkMesh.cpp
+++ b/Engine/Source/Programs/Enterprise/Datasmith/DatasmithMaxExporter/Private/DatasmithMaxDirectLinkMesh.cpp
@@ -298,6 +298,15 @@ bool FillDatasmithMeshFromBoundingBox(FDatasmithMesh& DatasmithMesh, INode* Expo
 		return false;
 	}
 
+	if (MaxMesh->getNumFaces() == 0)
+	{
+		if (bNeedsDelete)
+		{
+			MaxMesh->DeleteThis();
+		}
+		return false;
+	}
+
 	if (MaxMesh->getBoundingBox().IsEmpty())
 	{
 		MaxMesh->buildBoundingBox();
diff --git a/Engine/Source/Programs/Enterprise/Datasmith/DatasmithMaxExporter/Private/DatasmithMaxDirectLinkNotifications.cpp b/Engine/Source/Programs/Enterprise/Datasmith/DatasmithMaxExporter/Private/DatasmithMaxDirectLinkNotifications.cpp
--- a/Engine/Source/Programs/Enterprise/Datasmith/DatasmithMaxExporter/Private/DatasmithMaxDirectLinkNotifications.cpp
+++ b/Engine/Source/Programs/Enterprise/Datasmith/DatasmithMaxExporter/Private/DatasmithMaxDirectLinkNotifications.cpp
@@ -137,6 +137,10 @@ public:
 	virtual void LinkChanged(NodeKeyTab& nodes) override
 	{
 		LogNodeEvent(L"LinkChanged", nodes);
+		for (int NodeIndex = 0; NodeIndex < nodes.Count(); ++NodeIndex)
+		{
+			SceneTracker.NodeLinkChanged(nodes[NodeIndex]);
+		}
 	}
 
 	virtual void GroupChanged(NodeKeyTab& nodes) override
@@ -492,6 +496,14 @@ FString FNotifications::ConvertNotificationCodeToString(int code)
 	return Str ? *Str : TEXT("<unknown>");
 }
 
+void FNotifications::PrepareForUpdate()
+{
+	if (NodeEventCallback)
+	{
+		GetISceneEventManager()->TriggerMessages(NodeEventCallback->CallbackKey);
+	}
+}
+
 void FNotifications::On3dsMaxNotification(void* param, NotifyInfo* info)
 {
 	FNotifications& NotificationsHandler = *static_cast<FNotifications*>(param);
@@ -547,20 +559,48 @@ void FNotifications::On3dsMaxNotification(void* param, NotifyInfo* info)
 			SceneTracker.NodeDeleted(reinterpret_cast<INode*>(info->callParam));
 			break;
 		}
+		case NOTIFY_NODE_UNLINKED:
+		{
+			INode* Node = reinterpret_cast<INode*>(info->callParam);
+			LogDebugNode(NotificationsHandler.ConvertNotificationCodeToString(info->intcode), Node);
+			break;
+		}
+		case NOTIFY_SCENE_POST_DELETED_NODE:
+		{
+			INode* Node = reinterpret_cast<INode*>(info->callParam);
+			LogDebugNode(NotificationsHandler.ConvertNotificationCodeToString(info->intcode), Node);
+			break;
+		}
+		case NOTIFY_SCENE_XREF_POST_MERGE:
+		{
+			INode* Node = reinterpret_cast<INode*>(info->callParam);
+			LogDebugNode(NotificationsHandler.ConvertNotificationCodeToString(info->intcode), Node);
+
+			SceneTracker.NodeXRefMerged(Node);
+			break;
+		}
 		}
 	}
 
-	// Handle events not related to tracking scene changes
+	// Handle events for scene load/new/reset
+	// Note - different UI functions which destroy current scene send different notifications
+	// Reset - NOTIFY_SYSTEM_PRE_RESET NOTIFY_POST_SCENE_RESET
+	// Open -  NOTIFY_FILE_PRE_OPEN_PROCESS NOTIFY_FILE_PRE_OPEN NOTIFY_POST_SCENE_RESET NOTIFY_FILE_POST_OPEN NOTIFY_FILE_POST_OPEN_PROCESS NOTIFY_FILE_POST_OPEN_PROCESS_FINALIZED
+	// File>New>New All -  NOTIFY_SYSTEM_PRE_NEW NOTIFY_POST_SCENE_RESET NOTIFY_SYSTEM_POST_NEW
+	// File>New>New From Template -  NOTIFY_SYSTEM_PRE_RESET NOTIFY_POST_SCENE_RESET NOTIFY_SYSTEM_POST_RESET
 	switch (info->intcode)
 	{
-	// Handle New/Reset events - reset tracking immediately when "Pre events are received - after this point all nodes are invalid, don't wait for "Post" event
+	// Handle New/Reset events sent before scene reset - stop tracking immediately when "Pre" events are received - after this point all nodes are invalid, don't wait for "Post" event
 	case NOTIFY_SYSTEM_PRE_NEW:  // Sent when File>New>New All is selected 
 	case NOTIFY_SYSTEM_PRE_RESET:  // Sent when Reset OR File>New>New From Template is selected
-		Exporter->Reset();
+	case NOTIFY_FILE_PRE_OPEN:  // Sent when File>Open is used
+		Exporter->ResetSceneTracking();
 		break;
 
+	case NOTIFY_SYSTEM_POST_NEW:
+	case NOTIFY_SYSTEM_POST_RESET:
 	case NOTIFY_FILE_POST_OPEN:
-		Exporter->Reset();
+		Exporter->InitializeDirectLinkForScene();
 		break;
 
 	}
diff --git a/Engine/Source/Programs/Enterprise/Datasmith/DatasmithMaxExporter/Private/DatasmithMaxExporterWithDirectLink.cpp b/Engine/Source/Programs/Enterprise/Datasmith/DatasmithMaxExporter/Private/DatasmithMaxExporterWithDirectLink.cpp
--- a/Engine/Source/Programs/Enterprise/Datasmith/DatasmithMaxExporter/Private/DatasmithMaxExporterWithDirectLink.cpp
+++ b/Engine/Source/Programs/Enterprise/Datasmith/DatasmithMaxExporter/Private/DatasmithMaxExporterWithDirectLink.cpp
@@ -69,21 +69,21 @@ typedef Texmap* FTexmapKey;
 class FDatasmith3dsMaxScene
 {
 public:
-	TSharedPtr<IDatasmithScene> DatasmithSceneRef;
-	TSharedPtr<FDatasmithSceneExporter> SceneExporterRef;
-
 	FDatasmith3dsMaxScene() 
 	{
-		Reset();
+		ResetScene();
 	}
 
-	void Reset()
+	void ResetScene()
 	{
 		DatasmithSceneRef.Reset();
-		DatasmithSceneRef = FDatasmithSceneFactory::CreateScene(TEXT(""));
 		SceneExporterRef.Reset();
-		SceneExporterRef = MakeShared<FDatasmithSceneExporter>();
+	}
 
+	void SetupScene()
+	{
+		DatasmithSceneRef = FDatasmithSceneFactory::CreateScene(TEXT(""));
+		SceneExporterRef = MakeShared<FDatasmithSceneExporter>();
 
 		MSTR Renderer;
 		FString Host;
@@ -98,20 +98,6 @@ public:
 
 		FString Version = FString::FromInt(MAX_VERSION_MAJOR) + TEXT(".") + FString::FromInt(MAX_VERSION_MINOR) + TEXT(".") + FString::FromInt(MAX_VERSION_POINT);
 		DatasmithSceneRef->SetProductVersion(*Version);
-
-		// XXX: PreExport needs to be called before DirectLink instance is constructed - 
-		// Reason - it calls initialization of FTaskGraphInterface. Callstack:
-		// PreExport:
-		//  - FDatasmithExporterManager::Initialize 
-		//	-- DatasmithGameThread::InitializeInCurrentThread
-		//  --- GEngineLoop.PreInit
-		//  ---- PreInitPreStartupScreen
-		//  ----- FTaskGraphInterface::Startup
-		PreExport();
-
-		SetOutputPath(GetDirectlinkCacheDirectory());
-		FString SceneName = FPaths::GetCleanFilename(GetCOREInterface()->GetCurFileName().data());
-		SetName(*SceneName);
 	}
 
 	TSharedRef<IDatasmithScene> GetDatasmithScene()
@@ -146,6 +132,9 @@ public:
 		// Start measuring the time taken to export the scene.
 		SceneExporterRef->PreExport();
 	}
+
+	TSharedPtr<IDatasmithScene> DatasmithSceneRef;
+	TSharedPtr<FDatasmithSceneExporter> SceneExporterRef;
 };
 
 class FNodeTrackerHandle
@@ -340,42 +329,54 @@ public:
 		TimeValue Time;
 	};
 
-	void Start(TimeValue Time)
+	void Start(TimeValue Time, bool bInRenderQuality)
 	{
+		bRenderQuality = bInRenderQuality;
+
 		BeginProc.SetTime(Time);
 		EndProc.SetTime(Time);
 
-		BeginProc.BeginEnumeration();
+		if (bRenderQuality)
+		{
+			BeginProc.BeginEnumeration();
+		}
 	}
 
 	void Finish()
 	{
-		BeginProc.EndEnumeration();
-
-		// Call RenderEnd on every node that had RenderBegin called
-		EndProc.BeginEnumeration();
-		for(INode* Node: NodesPrepared)
+		if (bRenderQuality)
 		{
-			Node->EnumRefHierarchy(EndProc);
-		}
-		EndProc.EndEnumeration();
+			BeginProc.EndEnumeration();
 
-		NodesPrepared.Reset();
+			// Call RenderEnd on every node that had RenderBegin called
+			EndProc.BeginEnumeration();
+			for(INode* Node: NodesPrepared)
+			{
+				Node->EnumRefHierarchy(EndProc);
+			}
+			EndProc.EndEnumeration();
+			NodesPrepared.Reset();
+		}
 	}
 
 	void PrepareNode(INode* Node)
 	{
-		// Skip if node was already Prepared
-		bool bIsAlreadyPrepared;
-		NodesPrepared.FindOrAdd(Node, &bIsAlreadyPrepared);
-		if (bIsAlreadyPrepared)
+		if (bRenderQuality)
 		{
-			return;
-		}
+			// Skip if node was already Prepared
+			bool bIsAlreadyPrepared;
+			NodesPrepared.FindOrAdd(Node, &bIsAlreadyPrepared);
+			if (bIsAlreadyPrepared)
+			{
+				return;
+			}
 
-		Node->EnumRefHierarchy(BeginProc);
+			Node->EnumRefHierarchy(BeginProc);
+		}
 	}
 
+	bool bRenderQuality = false; // If need to call RenderBegin on all nodes to make them return Render-quality mesh
+
 	FBeginRefEnumProc BeginProc;
 	FEndRefEnumProc EndProc;
 	
@@ -464,7 +465,7 @@ public:
 class FSceneTracker: public ISceneTracker
 {
 public:
-	FSceneTracker(const FExportOptions& InOptions, FDatasmith3dsMaxScene& InExportedScene, FNotifications& InNotificationsHandler)
+	FSceneTracker(const FExportOptions& InOptions, FDatasmith3dsMaxScene& InExportedScene, FNotifications* InNotificationsHandler)
 		: Options(InOptions)
 		, ExportedScene(InExportedScene)
 		, NotificationsHandler(InNotificationsHandler)
@@ -473,13 +474,14 @@ public:
 	bool ParseScene()
 	{
 		INode* Node = GetCOREInterface()->GetRootNode();
-		bSceneParsed = ParseScene(Node, nullptr);
+		bSceneParsed = ParseScene(Node);
 		return bSceneParsed;
 	}
 
-	// Parset scene or XRef scene(in this case attach to parent datasmith actor)
-	bool ParseScene(INode* SceneRootNode, IDatasmithActorElement* ParentElement)
+	// Parse scene or XRef scene(in this case attach to parent datasmith actor)
+	bool ParseScene(INode* SceneRootNode, FXRefScene XRefScene=FXRefScene())
 	{
+		LogDebugNode(TEXT("ParseScene"), SceneRootNode);
 		// todo: do we need Root Datasmith node of scene/XRefScene in the hierarchy?
 		// is there anything we need to handle for main file root node?
 		// for XRefScene? Maybe addition/removal? Do we need one node to consolidate XRefScene under?
@@ -511,20 +513,25 @@ public:
 			}
 			else
 			{
-				ParseScene(SceneRootNode->GetXRefTree(XRefChild), ParentElement);
+				ParseScene(SceneRootNode->GetXRefTree(XRefChild), FXRefScene{SceneRootNode, XRefChild});
 			}
 		}
 
 		int32 ChildNum = SceneRootNode->NumberOfChildren();
 		for (int32 ChildIndex = 0; ChildIndex < ChildNum; ++ChildIndex)
 		{
-			ParseNode(SceneRootNode->GetChildNode(ChildIndex));
+			if (FNodeTracker* NodeTracker = ParseNode(SceneRootNode->GetChildNode(ChildIndex)))
+			{
+				NodeTracker->SetXRefIndex(XRefScene);
+			}
 		}
 		return true;
 	}
 
-	void ParseNode(INode* Node)
+	FNodeTracker* ParseNode(INode* Node)
 	{
+		LogDebugNode(TEXT("ParseNode"), Node);
+
 		BOOL bIsNodeHidden = Node->IsNodeHidden(TRUE);
 
 		if (bIsNodeHidden == false)
@@ -557,6 +564,7 @@ public:
 			ensure(NodeTracker->bDeleted);
 			NodeTracker->bDeleted = false;
 			InvalidateNode(*NodeTracker);
+			return NodeTracker;
 		}
 		else
 		{
@@ -568,6 +576,7 @@ public:
 			{
 				ParseNode(Node->GetChildNode(ChildIndex));
 			}
+			return NodeTracker.GetNodeTracker();
 		}
 	}
 
@@ -633,19 +642,48 @@ public:
 	}
 
 	// Applies all recorded changes to Datasmith scene
-	bool Update(bool bQuiet)
+	bool Update(bool bQuiet, bool bRenderQuality)
 	{
+		// Disable Undo, editing, redraw, messages during export/sync so that nothing changes the scene
+		GetCOREInterface()->EnableUndo(false);
+		GetCOREInterface()->DisableSceneRedraw();
+		SuspendAll UberSuspend(TRUE, TRUE, TRUE, TRUE, TRUE, TRUE);
+
+		// Flush all updates for SceneEventManager - so they are not received in mid of Update
+		// When ProgressBar is updated it calls internal max event loop which can send unprocessed events to the callback
+		if (NotificationsHandler)
+		{
+			NotificationsHandler->PrepareForUpdate();
+		}
+
 		DatasmithMaxLogger::Get().Purge();
-		TimeValue Time = GetCOREInterface()->GetTime();
-		NodesPreparer.Start(Time);
+
+		NodesPreparer.Start(GetCOREInterface()->GetTime(), bRenderQuality);
+
+		bUpdateInProgress = true;
+		bool bResult = UpdateInternalSafe(bQuiet);
+		bUpdateInProgress = false;
+
+		NodesPreparer.Finish();
+
+		UberSuspend.Resume();
+		GetCOREInterface()->EnableSceneRedraw();
+		GetCOREInterface()->EnableUndo(true);
+
+		return bResult;
+	}
+
+	bool UpdateInternalSafe(bool bQuiet)
+	{
 		__try
 		{
 			return UpdateInternal(bQuiet);
 		}
-		__finally // Make sure that finalize called no matter what happened in Vegas
+		__except(EXCEPTION_EXECUTE_HANDLER)
 		{
-			NodesPreparer.Finish();
+			LogInfo(TEXT("Update finished with exception"));
 		}
+		return false;
 	}
 
 	bool UpdateInternal(bool bQuiet)
@@ -854,6 +892,7 @@ public:
 	FORCENOINLINE
 	FNodeTrackerHandle& AddNode(FNodeKey NodeKey, INode* Node)
 	{
+		LogDebugNode(TEXT("AddNode"), Node);
 		FNodeTrackerHandle& NodeTracker = NodeTrackers.Emplace(NodeKey, FNodeTrackerHandle(NodeKey, Node));
 		
 		NodeTrackersNames.FindOrAdd(NodeTracker.GetNodeTracker()->Name).Add(NodeTracker.GetNodeTracker());
@@ -869,18 +908,49 @@ public:
 
 	void InvalidateNode(FNodeTracker& NodeTracker)
 	{
+		ensure(!bUpdateInProgress);
+
 		NodeTracker.Invalidate();
 		InvalidatedNodeTrackers.Add(&NodeTracker);
+
+		// Invalidate whole sub-hierarchy of nodes that were previously recorded as children(this might have changed)
+		for (FNodeTracker* Child : NodeTracker.Children)
+		{
+			InvalidateNode(*Child);
+		}
+
+		// Invalidate whole sub-hierarchy of nodes are now children.
+		// E.g. a node could have been hidden so its children were attached to grandparent(parent of hidden node)
+		// Need to invalidate those to reattach
+		int32 ChildNum = NodeTracker.Node->NumberOfChildren();
+		for (int32 ChildIndex = 0; ChildIndex < ChildNum; ++ChildIndex)
+		{
+			InvalidateNode(NodeEventNamespace::GetKeyByNode(NodeTracker.Node->GetChildNode(ChildIndex)));
+		}
+
+		NodeTracker.Children.Reset();
 	}
 
 	// todo: make fine invalidates - full only something like geometry change, but finer for transform, name change and more
-	void InvalidateNode(FNodeKey NodeKey)
+	FNodeTracker* InvalidateNode(FNodeKey NodeKey)
 	{
 		if (FNodeTrackerHandle* NodeTrackerHandle = NodeTrackers.Find(NodeKey))
 		{
 			FNodeTracker* NodeTracker = NodeTrackerHandle->GetNodeTracker();
-			InvalidateNode(*NodeTracker);
+			if (NodeEventNamespace::GetNodeByKey(NodeKey))
+			{
+				InvalidateNode(*NodeTracker);
+				return NodeTracker;
+			}
+			else
+			{
+				// Sometimes note update received without node Delete event
+				// Test case: create container, add node to it. Close it, open it, close again, then sync
+				InvalidatedNodeTrackers.Add(NodeTracker);
+				NodeTracker->bDeleted = true;
+			}
 		}
+		return nullptr;
 	}
 
 	bool IsNodeInvalidated(const FNodeTrackerHandle& NodeTracker)
@@ -913,9 +983,18 @@ public:
 			}
 			else
 			{
-				ExportedScene.DatasmithSceneRef->RemoveActor(NodeTracker.DatasmithActorElement, EDatasmithActorRemovalRule::KeepChildrenAndKeepRelativeTransform);
+				// Detach all children(so they won't be reattached automatically to root when actor is detached from parent below)
+				// Children reattachment will happen later in Update
+				int32 ChildCount = NodeTracker.DatasmithActorElement->GetChildrenCount();
+				// Remove last child each time to optimize array elements relocation
+				for(int32 ChildIndex = ChildCount-1; ChildIndex >= 0; --ChildIndex)
+				{
+					NodeTracker.DatasmithActorElement->RemoveChild(NodeTracker.DatasmithActorElement->GetChild(ChildIndex));
+				}
+				ExportedScene.DatasmithSceneRef->RemoveActor(NodeTracker.DatasmithActorElement, EDatasmithActorRemovalRule::RemoveChildren);
 			}
 			NodeTracker.DatasmithActorElement.Reset();
+			NodeTracker.Children.Reset();
 		}
 	}
 
@@ -1236,7 +1315,8 @@ public:
 	// Find first ansestor node which has DatasmithActor created for it
 	FNodeTracker* GetParentNodeTracker(FNodeTracker& NodeTracker)
 	{
-		FNodeKey ParentNodeKey = NodeEventNamespace::GetKeyByNode(NodeTracker.Node->GetParentNode());
+		INode* XRefParent = NodeTracker.GetXRefParent();
+		FNodeKey ParentNodeKey = NodeEventNamespace::GetKeyByNode(XRefParent ? XRefParent : NodeTracker.Node->GetParentNode());
 		FNodeTrackerHandle* ParentNodeTrackerHandle = NodeTrackers.Find(ParentNodeKey);
 		if (!ParentNodeTrackerHandle)
 		{
@@ -1270,7 +1350,8 @@ public:
 
 		if (FNodeTracker* ParentNodeTracker = GetAncestorNodeTrackerWithDatasmithActor(NodeTracker))
 		{
-			ParentNodeTracker->DatasmithActorElement->AddChild(NodeTracker.DatasmithActorElement);
+			ParentNodeTracker->DatasmithActorElement->AddChild(NodeTracker.DatasmithActorElement, EDatasmithActorAttachmentRule::KeepWorldTransform);
+			ParentNodeTracker->Children.Add(&NodeTracker);
 		}
 		else
 		{
@@ -1625,10 +1706,10 @@ public:
 			TSharedRef<IDatasmithActorElement> HismActorElement = FDatasmithMaxSceneExporter::ExportHierarchicalInstanceStaticMeshActor( 
 				ExportedScene.GetDatasmithScene(), NodeTracker.Node, GeometryNode, *MeshLabel, SupportedChannels,
 				Material, &Transforms, *MeshName, Converter.UnitToCentimeter, EStaticMeshExportMode::Default, InversedHISMActor);
-			NodeTracker.DatasmithActorElement->AddChild(HismActorElement);
+			NodeTracker.DatasmithActorElement->AddChild(HismActorElement, EDatasmithActorAttachmentRule::KeepWorldTransform);
 			if (InversedHISMActor)
 			{
-				NodeTracker.DatasmithActorElement->AddChild(InversedHISMActor);
+				NodeTracker.DatasmithActorElement->AddChild(InversedHISMActor, EDatasmithActorAttachmentRule::KeepWorldTransform);
 			}
 			MeshIndex++;
 		}
@@ -1707,15 +1788,41 @@ public:
 			return;
 		}
 
-		NotificationsHandler.AddNode(Node);
+		if (NotificationsHandler)
+		{
+			NotificationsHandler->AddNode(Node);
+		}
 
 		ParseNode(Node);
 	}
 
+	virtual void NodeXRefMerged(INode* Node) override
+	{
+		if (!Node)
+		{
+			return;
+		}
+
+		int32 XRefIndex = -1; // Node that has this xref scene attached to(e.g. to place in hierarchy and to transform)
+		INode* SceneRootNode = GetCOREInterface()->GetRootNode();
+		for (int XRefChild = 0; XRefChild < SceneRootNode->GetXRefFileCount(); ++XRefChild)
+		{
+			if (Node == SceneRootNode->GetXRefTree(XRefChild))
+			{
+				XRefIndex = XRefChild;
+			}
+		}
+
+		FNodeKey NodeKey = NodeEventNamespace::GetKeyByNode(Node);
+		InvalidateNode(NodeKey);
+
+		ParseScene(Node, FXRefScene{SceneRootNode, XRefIndex}); // Parse xref hierarchy - it won't add itself! Or will it?
+	}
+
 	virtual void NodeDeleted(INode* Node) override
 	{
+		LogDebugNode(TEXT("NodeDeleted"), Node);
 		// todo: check for null
-
 		FNodeKey NodeKey = NodeEventNamespace::GetKeyByNode(Node);
 
 		if (FNodeTrackerHandle* NodeTrackerHandle = NodeTrackers.Find(NodeKey))
@@ -1800,14 +1907,21 @@ public:
 
 		InvalidateNode(NodeKey);
 	}
+
+	virtual void NodeLinkChanged(FNodeKey NodeKey) override
+	{
+		InvalidateNode(NodeKey);
+	}
 	
 	///////////////////////////////////////////////
 
 	const FExportOptions& Options;
 	FDatasmith3dsMaxScene& ExportedScene;
-	FNotifications& NotificationsHandler;
+	FNotifications* NotificationsHandler;
 
 	bool bSceneParsed = false;
+	bool bUpdateInProgress = false;
+
 
 	TMap<FNodeKey, FNodeTrackerHandle> NodeTrackers; // All scene nodes
 	TMap<FString, TSet<FNodeTracker*>> NodeTrackersNames; // Nodes grouped by name
@@ -1845,9 +1959,10 @@ public:
 class FExporter: public IExporter
 {
 public:
-	FExporter(FExportOptions& InOptions): Options(InOptions), NotificationsHandler(*this), SceneTracker(Options, ExportedScene, NotificationsHandler)
+	FExporter(FExportOptions& InOptions): Options(InOptions), NotificationsHandler(*this), SceneTracker(Options, ExportedScene, &NotificationsHandler)
 	{
-		Reset();
+		ResetSceneTracking();
+		InitializeDirectLinkForScene(); // Setup DL connection immediately when plugin loaded
 	}
 
 	virtual void Shutdown() override;
@@ -1863,31 +1978,53 @@ public:
 		ExportedScene.SetName(Name);
 	}
 
+	virtual void InitializeScene() override
+	{
+		ExportedScene.SetupScene();
+	}
+
 	virtual void ParseScene() override
 	{
 		SceneTracker.ParseScene();
 	}
 
-	// Just export, parsing scene from scratch
-	bool Export(bool bQuiet)
+	virtual void InitializeDirectLinkForScene() override
 	{
-		SceneTracker.Update(bQuiet);
-		if (Options.bAnimatedTransforms)
+		if (DirectLinkImpl) 
 		{
-			SceneTracker.ExportAnimations();
+			return;
 		}
-		ExportedScene.GetSceneExporter().Export(ExportedScene.GetDatasmithScene(), false);
-		return true;
-	}
 
-	virtual void InitializeDirectLinkForScene() override
-	{
+		InitializeScene();
+
+		// XXX: PreExport needs to be called before DirectLink instance is constructed - 
+		// Reason - it calls initialization of FTaskGraphInterface. Callstack:
+		// PreExport:
+		//  - FDatasmithExporterManager::Initialize 
+		//	-- DatasmithGameThread::InitializeInCurrentThread
+		//  --- GEngineLoop.PreInit
+		//  ---- PreInitPreStartupScreen
+		//  ----- FTaskGraphInterface::Startup
+		ExportedScene.PreExport();
+
+		SetOutputPath(GetDirectlinkCacheDirectory());
+		FString SceneName = FPaths::GetCleanFilename(GetCOREInterface()->GetCurFileName().data());
+		SetName(*SceneName);
+
 		DirectLinkImpl.Reset(new FDatasmithDirectLink);
 		DirectLinkImpl->InitializeForScene(ExportedScene.GetDatasmithScene());
 	}
 
 	virtual void UpdateDirectLinkScene() override
 	{
+		if (!DirectLinkImpl) 
+		{
+			// InitializeDirectLinkForScene wasn't called yet. This rarely happens when Sync is pressed right before event like PostSceneReset(for New All UI command) was handled
+			// Very quickly! Unfortunately code needs to wait for PostSceneReset to get proper scene name there(no earlier event signals that name is available)
+			InitializeDirectLinkForScene();
+		}
+
+		LogDebug(TEXT("UpdateDirectLinkScene"));
 		DirectLinkImpl->UpdateScene(ExportedScene.GetDatasmithScene());
 		StartSceneChangeTracking(); // Always track scene changes if it's synced with DirectLink
 	}
@@ -1965,10 +2102,10 @@ public:
 
 	virtual bool UpdateScene(bool bQuiet) override
 	{
-		return SceneTracker.Update(bQuiet);
+		return SceneTracker.Update(bQuiet, false);
 	}
 
-	virtual void Reset() override
+	virtual void ResetSceneTracking() override
 	{
 		NotificationsHandler.StopSceneChangeTracking();
 		if (IsAutoSyncEnabled())
@@ -1976,12 +2113,11 @@ public:
 			ToggleAutoSync();
 		}
 
-		ExportedScene.Reset();
+		ExportedScene.ResetScene();
 
 		SceneTracker.Reset();
 
 		DirectLinkImpl.Reset();
-		InitializeDirectLinkForScene();
 	}
 
 	virtual ISceneTracker& GetSceneTracker() override
@@ -2001,6 +2137,7 @@ public:
 	bool bAutoSyncEnabled = false;
 	float AutoSyncDelaySeconds = 0.5f; // AutoSync is attempted periodically using this interval
 	float AutoSyncIdleDelaySeconds = 0.5f; // Time period user should be idle to run AutoSync
+
 };
 
 FPersistentExportOptions PersistentExportOptions;
@@ -2060,11 +2197,22 @@ void FExporter::Shutdown()
 
 bool Export(const TCHAR* Name, const TCHAR* OutputPath, bool bQuiet)
 {
-	FExporter TempExporter(PersistentExportOptions.Options);
-	TempExporter.ExportedScene.SetName(Name);
-	TempExporter.SetOutputPath(OutputPath);
+	FDatasmith3dsMaxScene ExportedScene;
+	ExportedScene.SetupScene();
+	ExportedScene.SetName(Name);
+	ExportedScene.SetOutputPath(OutputPath);
+
+	FSceneTracker SceneTracker(PersistentExportOptions.Options, ExportedScene, nullptr);
+	SceneTracker.Update(bQuiet, true);
+
+	if (PersistentExportOptions.Options.bAnimatedTransforms)
+	{
+		SceneTracker.ExportAnimations();
+	}
 
-	return TempExporter.Export(bQuiet);
+	ExportedScene.GetSceneExporter().Export(ExportedScene.GetDatasmithScene(), false);
+
+	return true;
 }
 
 bool OpenDirectLinkUI()
diff --git a/Engine/Source/Programs/Enterprise/Datasmith/DatasmithMaxExporter/Public/DatasmithMaxDirectLink.h b/Engine/Source/Programs/Enterprise/Datasmith/DatasmithMaxExporter/Public/DatasmithMaxDirectLink.h
--- a/Engine/Source/Programs/Enterprise/Datasmith/DatasmithMaxExporter/Public/DatasmithMaxDirectLink.h
+++ b/Engine/Source/Programs/Enterprise/Datasmith/DatasmithMaxExporter/Public/DatasmithMaxDirectLink.h
@@ -62,10 +62,13 @@ class IExporter
 
 	virtual ISceneTracker& GetSceneTracker() = 0;
 
+	virtual void InitializeScene() = 0;
+
 	// Scene update
 	virtual void ParseScene() = 0;
 	virtual bool UpdateScene(bool bQuiet) = 0;
-	virtual void Reset() = 0;
+
+	virtual void ResetSceneTracking() = 0;
 
 	// ChangeTracking
 	virtual void StartSceneChangeTracking() = 0;
@@ -92,6 +95,15 @@ IPersistentExportOptions& GetPersistentExportOptions();
 typedef NodeEventNamespace::NodeKey FNodeKey;
 typedef MtlBase* FMaterialKey;
 
+struct FXRefScene
+{
+	INode* Tree = nullptr;
+	int32 XRefFileIndex = -1;
+
+	operator bool(){ return XRefFileIndex != -1; }
+};
+
+
 // Identifies Max node to track its changes
 class FNodeTracker: FNoncopyable
 {
@@ -113,6 +125,16 @@ public:
 		return InstanceHandle != 0;
 	}
 
+	void SetXRefIndex(FXRefScene InXRefScene)
+	{
+		XRefScene = InXRefScene;
+	}
+
+	INode* GetXRefParent()
+	{
+		return XRefScene ? XRefScene.Tree->GetXRefParent(XRefScene.XRefFileIndex): nullptr;
+	}
+
 	FNodeKey NodeKey;
 	INode* const Node;
 	FNodeTracker* Collision = nullptr;
@@ -128,6 +150,10 @@ public:
 
 	TSharedPtr<IDatasmithMeshActorElement> DatasmithMeshActor;
 
+	FXRefScene XRefScene;
+
+	TArray<FNodeTracker*> Children;
+
 };
 
 class FRenderMeshForConversion: FNoncopyable
@@ -190,6 +216,7 @@ public:
 	virtual void NodeGeometryChanged(FNodeKey NodeKey) = 0;
 	virtual void NodeHideChanged(FNodeKey NodeKey) = 0;
 	virtual void NodePropertiesChanged(FNodeKey NodeKey) = 0;
+	virtual void NodeLinkChanged(FNodeKey NodeKey) = 0;
 	virtual void NodeTransformChanged(FNodeKey NodeKey) = 0;
 	virtual void NodeMaterialAssignmentChanged(FNodeKey NodeKey) = 0;
 	virtual void NodeMaterialGraphModified(FNodeKey NodeKey) = 0;
@@ -200,6 +227,7 @@ public:
 	virtual void SetupActor(FNodeTracker& NodeTracker) = 0;
 	virtual void SetupDatasmithHISMForNode(FNodeTracker& NodeTracker, INode* GeometryNode, const FRenderMeshForConversion& RenderMesh, Mtl* Material, int32 MeshIndex, const TArray<Matrix3>& Transforms) = 0;
 	virtual void RemoveMaterial(const TSharedPtr<IDatasmithBaseMaterialElement>& DatasmithMaterial) = 0;
+	virtual void NodeXRefMerged(INode* Node) = 0;
 };
 
 //---- Geometry utility function
@@ -241,6 +269,7 @@ public:
 	void AddNode(INode*);
 
 	FString ConvertNotificationCodeToString(int code);
+	void PrepareForUpdate();
 
 	static void On3dsMaxNotification(void* param, NotifyInfo* info);
 
