DatasmithRhinoExporter - Adding support for locally hidden object.

This is an oversight on our part. Up to now, the visibility was only tested using layers, but there are "hide/show" commands that can be used to control the visibility on a per-object basis.

diff --git a/Engine/Source/Programs/Enterprise/Datasmith/DatasmithRhinoExporter/Private/DatasmithRhinoExportContext.cs b/Engine/Source/Programs/Enterprise/Datasmith/DatasmithRhinoExporter/Private/DatasmithRhinoExportContext.cs
--- a/Engine/Source/Programs/Enterprise/Datasmith/DatasmithRhinoExporter/Private/DatasmithRhinoExportContext.cs
+++ b/Engine/Source/Programs/Enterprise/Datasmith/DatasmithRhinoExporter/Private/DatasmithRhinoExportContext.cs
@@ -221,6 +221,7 @@ namespace DatasmithRhino
 		{
 			get
 			{
+				bool bVisibleLocally = true;
 				if (RhinoCommonObject is Layer RhinoLayer)
 				{
 					// This recursion ensure that only layer actors with an actual exported object under them are considered visible.
@@ -228,12 +229,17 @@ namespace DatasmithRhino
 					return RhinoLayer.IsVisible
 						&& ChildrenInternal.Any(Child => Child.bIsVisible);
 				}
+				else if (RhinoCommonObject is RhinoObject CurrentRhinoObject)
+				{
+					bVisibleLocally = CurrentRhinoObject.Visible;
+				}
 				else if (bIsRoot)
 				{
 					return true;
 				}
 
-				return (VisibilityLayer == null || VisibilityLayer.IsVisible)
+				return bVisibleLocally
+					&& (VisibilityLayer == null || VisibilityLayer.IsVisible)
 					&& (DefinitionNode == null || DefinitionNode.bIsVisible);
 			}
 		}
