[3ds Max] DirectLink: Lights fixes

- CoronaLight Cylinder orientation
- Photometric light attenuation
- Arnold light intensity
- CoronaSun color/temperature
- VRaySun color/temperature

diff --git a/Engine/Source/Programs/Enterprise/Datasmith/DatasmithMaxExporter/Private/DatasmithMaxExporterWithDirectLink.cpp b/Engine/Source/Programs/Enterprise/Datasmith/DatasmithMaxExporter/Private/DatasmithMaxExporterWithDirectLink.cpp
--- a/Engine/Source/Programs/Enterprise/Datasmith/DatasmithMaxExporter/Private/DatasmithMaxExporterWithDirectLink.cpp
+++ b/Engine/Source/Programs/Enterprise/Datasmith/DatasmithMaxExporter/Private/DatasmithMaxExporterWithDirectLink.cpp
@@ -1638,12 +1638,13 @@ public:
 
 		Lights.Add(&NodeTracker);
 
+		TSharedPtr<IDatasmithLightActorElement> LightElement;
 		if(!NodeTracker.DatasmithActorElement)
 		{
 			// note: this is how baseline exporter derives names
 			FString UniqueName = FString::FromInt(NodeTracker.Node->GetHandle());
 
-			TSharedPtr<IDatasmithLightActorElement> LightElement = FDatasmithMaxSceneExporter::CreateLightElementForNode(NodeTracker.Node, *UniqueName);
+			LightElement = FDatasmithMaxSceneExporter::CreateLightElementForNode(NodeTracker.Node, *UniqueName);
 
 			if (!LightElement)
 			{
@@ -1669,6 +1670,31 @@ public:
 		}
 		SetupActor(NodeTracker);
 
+		//Cylinder shaped lights don't have the same default orientations, so we recalculate their transform and add the shape information.
+		if (LightElement->IsA(EDatasmithElementType::AreaLight) 
+			&& StaticCastSharedPtr<IDatasmithAreaLightElement>(LightElement)->GetLightShape() == EDatasmithLightShape::Cylinder)
+		{
+			FVector Translation, Scale;
+			FQuat Rotation;
+
+			const float UnitMultiplier = (float)GetSystemUnitScale(UNITS_CENTIMETERS);
+			const FMaxLightCoordinateConversionParams LightParams = FMaxLightCoordinateConversionParams(NodeTracker.Node, EDatasmithLightShape::Cylinder);
+			if (NodeTracker.Node->GetWSMDerivedObject() != nullptr)
+			{
+				FDatasmithMaxSceneExporter::MaxToUnrealCoordinates(NodeTracker.Node->GetObjTMAfterWSM(GetCOREInterface()->GetTime()), Translation, Rotation, Scale, UnitMultiplier, LightParams);
+			}
+			else
+			{
+				FDatasmithMaxSceneExporter::MaxToUnrealCoordinates(NodeTracker.Node->GetObjectTM(GetCOREInterface()->GetTime()), Translation, Rotation, Scale, UnitMultiplier, LightParams);
+			}
+
+			Rotation.Normalize();
+			LightElement->SetTranslation(Translation);
+			LightElement->SetScale(Scale);
+			LightElement->SetRotation(Rotation);
+		}
+
+
 		NodeTracker.bInvalidated = false;
 
 		return true;
diff --git a/Engine/Source/Programs/Enterprise/Datasmith/DatasmithMaxExporter/Private/DatasmithMaxSceneExporter.cpp b/Engine/Source/Programs/Enterprise/Datasmith/DatasmithMaxExporter/Private/DatasmithMaxSceneExporter.cpp
--- a/Engine/Source/Programs/Enterprise/Datasmith/DatasmithMaxExporter/Private/DatasmithMaxSceneExporter.cpp
+++ b/Engine/Source/Programs/Enterprise/Datasmith/DatasmithMaxExporter/Private/DatasmithMaxSceneExporter.cpp
@@ -956,12 +956,7 @@ bool FDatasmithMaxSceneExporter::ParseLight(INode* Node, TSharedRef< IDatasmithL
 	}
 	else if (LightClass == EMaxLightClass::SunEquivalent)
 	{
-		LightElement->SetIntensity( 10.f );
-		LightElement->SetUseIes( false );
-		LightElement->SetUseTemperature( true );
-		LightElement->SetTemperature( 5780.f );
-		LightElement->SetColor( FLinearColor::White );
-
+		ParseSun(Node, LightElement);
 		return true;
 	}
 
@@ -1067,14 +1062,15 @@ bool FDatasmithMaxSceneExporter::ParseLightObject(LightObject& Light, TSharedRef
 {
 	LightState LState;
 	Interval ValidInterval = FOREVER;
-	Light.EvalLightState( 0, ValidInterval, &LState );
+	RefResult EvalLightStateResult = Light.EvalLightState( 0, ValidInterval, &LState );
 
 	LightElement->SetEnabled( Light.GetUseLight() != 0 );
 
 	Point3 LightColor = LState.color;
 	LightElement->SetColor( FLinearColor( LightColor.x, LightColor.y, LightColor.z ) );
 
-	LightElement->SetIntensity( Light.GetIntensity( GetCOREInterface()->GetTime() ) );
+	float Intensity = Light.GetIntensity( GetCOREInterface()->GetTime() );
+	LightElement->SetIntensity( Intensity );
 
 	if ( LightElement->IsA( EDatasmithElementType::PointLight ) )
 	{
@@ -1367,6 +1363,16 @@ bool FDatasmithMaxSceneExporter::ParsePhotometricLight(LightObject& Light, TShar
 {
 	LightscapeLight& PhotometricLight = static_cast< LightscapeLight& >( Light );
 
+	// Photometric lights don't expose attenuation in EvalLightState
+	IParamBlock2* ParamBlock2 = PhotometricLight.GetParamBlockByID(LightscapeLight::PB_GENERAL);
+	bool bUseAtten = ParamBlock2->GetInt((short) LightscapeLight::PB_USE_FARATTENUATION) != 0;
+	float Atten = ParamBlock2->GetFloat((short) LightscapeLight::PB_END_FARATTENUATION);
+
+	if (bUseAtten)
+	{
+		PointLightElement->SetAttenuationRadius( Atten * (float)GetSystemUnitScale( UNITS_CENTIMETERS ) );
+	}
+
 	LightscapeLight::IntensityType IntensityUnits = PhotometricLight.GetIntensityType();
 
 	switch ( IntensityUnits )
@@ -1642,6 +1648,103 @@ bool FDatasmithMaxSceneExporter::ParseVRayLight(LightObject& Light, TSharedRef<
 	return true;
 }
 
+void FDatasmithMaxSceneExporter::ParseSun(INode* Node, TSharedRef<IDatasmithLightActorElement> LightElement)
+{
+	float Intensity = 10.0f;
+	FLinearColor Color = FLinearColor::White;
+	bool bUseTemperature = true;
+	float Temperature = 5780.f;
+
+	ObjectState ObjState = Node->EvalWorldState(0);
+	LightObject& Light = *(LightObject*)ObjState.obj;
+	Class_ID ClassID = Light.ClassID();
+
+	if(ClassID == VRAYSUNCLASS)
+	{
+		IParamBlock2* ParamBlock2 = Light.GetParamBlockByID( (short)EVRayLightParamBlocks::Params );
+
+		ParamBlockDesc2* ParamBlockDesc = ParamBlock2->GetDesc();
+
+		// Loop through all the defined parameters therein
+		for (int i = 0; i < ParamBlockDesc->count; i++)
+		{
+			const ParamDef& ParamDefinition = ParamBlockDesc->paramdefs[i];
+				
+			if (FCString::Stricmp(ParamDefinition.int_name, TEXT("intensity_multiplier")) == 0)
+			{
+				Intensity = ParamBlock2->GetFloat(ParamDefinition.ID, GetCOREInterface()->GetTime());
+			}
+			else if (FCString::Stricmp(ParamDefinition.int_name, TEXT("filter_color")) == 0)
+			{
+				Color = FDatasmithMaxMatHelper::MaxLinearColorToFLinearColor((BMM_Color_fl)ParamBlock2->GetColor(ParamDefinition.ID, GetCOREInterface()->GetTime()));
+				bUseTemperature = false;
+			}
+		}
+	}
+	else if(ClassID == CORONASUNCLASS || ClassID == CORONASUNCLASSB)
+	{
+		int32 CoronaSunColorMode = 0;
+		float CoronaSunTemperature = 0;
+		FLinearColor CoronaSunColor;
+
+		const int NumParamBlocks = Light.NumParamBlocks();
+
+		for ( int j = 0; j < NumParamBlocks; j++ )
+		{
+			IParamBlock2* ParamBlock2 = Light.GetParamBlockByID( (short)j );
+			ParamBlockDesc2* ParamBlockDesc = ParamBlock2->GetDesc();
+
+			for (int i = 0; i < ParamBlockDesc->count; i++)
+			{
+				const ParamDef& ParamDefinition = ParamBlockDesc->paramdefs[i];
+
+				if (FCString::Stricmp(ParamDefinition.int_name, TEXT("intensity")) == 0)
+				{
+					Intensity = ParamBlock2->GetFloat(ParamDefinition.ID, GetCOREInterface()->GetTime());
+				}
+				else if (FCString::Stricmp(ParamDefinition.int_name, TEXT("colorMode")) == 0)
+				{
+					CoronaSunColorMode = ParamBlock2->GetInt(ParamDefinition.ID, GetCOREInterface()->GetTime());
+				}
+				else if (FCString::Stricmp(ParamDefinition.int_name, TEXT("blackbodyTemperature")) == 0)
+				{
+					CoronaSunTemperature = ParamBlock2->GetFloat(ParamDefinition.ID, GetCOREInterface()->GetTime());
+				}
+				else if (FCString::Stricmp(ParamDefinition.int_name, TEXT("colorDirect")) == 0)
+				{
+					CoronaSunColor = FDatasmithMaxMatHelper::MaxLinearColorToFLinearColor((BMM_Color_fl)ParamBlock2->GetColor(ParamDefinition.ID, GetCOREInterface()->GetTime()));
+				}
+			}
+		}
+		switch (CoronaSunColorMode)
+		{
+		case 0: // Direct Input
+			{
+				bUseTemperature = false;
+				Color = CoronaSunColor;
+				break;
+			}
+		case 1: // Kelvin temp
+			{
+				bUseTemperature = true;
+				Temperature = CoronaSunTemperature;
+				break;
+			}
+		case 2: // Realistic
+			{
+				// Use defaults
+				break;
+			}
+		}
+	}
+
+	LightElement->SetIntensity( Intensity );
+	LightElement->SetUseIes( false );
+	LightElement->SetUseTemperature( bUseTemperature );
+	LightElement->SetTemperature( Temperature );
+	LightElement->SetColor( Color );
+}
+
 bool FDatasmithMaxSceneExporter::ParseVRayLightPortal(LightObject& Light, TSharedRef< IDatasmithLightmassPortalElement > LightPortalElement, TSharedRef< IDatasmithScene > DatasmithScene)
 {
 	IParamBlock2* ParamBlock2 = Light.GetParamBlockByID( (short)EVRayLightParamBlocks::Params );
@@ -1879,6 +1982,11 @@ bool FDatasmithMaxSceneExporter::ParseLightParameters(EMaxLightClass LightClass,
 					ProcessLightTexture( LightElement, TextureMap, DatasmithScene );
 				}
 			}
+			else if (FCString::Stricmp(ParamDefinition.int_name, TEXT("intensity")) == 0 && LightClass == EMaxLightClass::ArnoldLight)
+			{
+				float Intensity = ParamBlock2->GetFloat(ParamDefinition.ID, GetCOREInterface()->GetTime());
+				LightElement->SetIntensity( Intensity );
+			}
 
 			// using ies file
 			if ( LightElement->GetUseIes() )
diff --git a/Engine/Source/Programs/Enterprise/Datasmith/DatasmithMaxExporter/Public/DatasmithMaxSceneExporter.h b/Engine/Source/Programs/Enterprise/Datasmith/DatasmithMaxExporter/Public/DatasmithMaxSceneExporter.h
--- a/Engine/Source/Programs/Enterprise/Datasmith/DatasmithMaxExporter/Public/DatasmithMaxSceneExporter.h
+++ b/Engine/Source/Programs/Enterprise/Datasmith/DatasmithMaxExporter/Public/DatasmithMaxSceneExporter.h
@@ -101,4 +101,5 @@ private:
 	static bool ParseLightParameters(EMaxLightClass LightClass, LightObject& Light, TSharedRef< IDatasmithLightActorElement > LightElement, TSharedRef< IDatasmithScene > DatasmithScene);
 	static bool ProcessLightTexture(TSharedRef< IDatasmithLightActorElement > LightElement, Texmap* LightTexture, TSharedRef< IDatasmithScene > DatasmithScene);
 	static bool ParseTransformAnimation(INode* Node, TSharedRef< IDatasmithTransformAnimationElement > AnimationElement, float UnitMultiplier, const FMaxLightCoordinateConversionParams& LightParams);
+	static void ParseSun(INode* Node, TSharedRef<IDatasmithLightActorElement> LightElement);
 };
