[3ds Max] DirectLink: fixed crash/sync freeze

- repro: start max with plugin, don't sync!, load scene with relatively deep hierarchy, only then sync

diff --git a/Engine/Source/Programs/Enterprise/Datasmith/DatasmithMaxExporter/Private/DatasmithMaxExporterWithDirectLink.cpp b/Engine/Source/Programs/Enterprise/Datasmith/DatasmithMaxExporter/Private/DatasmithMaxExporterWithDirectLink.cpp
--- a/Engine/Source/Programs/Enterprise/Datasmith/DatasmithMaxExporter/Private/DatasmithMaxExporterWithDirectLink.cpp
+++ b/Engine/Source/Programs/Enterprise/Datasmith/DatasmithMaxExporter/Private/DatasmithMaxExporterWithDirectLink.cpp
@@ -568,7 +568,7 @@ public:
 		}
 		else
 		{
-			FNodeTrackerHandle& NodeTracker = AddNode(NodeKey, Node);
+			FNodeTracker* NodeTracker = AddNode(NodeKey, Node);
 
 			// Parse children
 			int32 ChildNum = Node->NumberOfChildren();
@@ -576,7 +576,7 @@ public:
 			{
 				ParseNode(Node->GetChildNode(ChildIndex));
 			}
-			return NodeTracker.GetNodeTracker();
+			return NodeTracker;
 		}
 	}
 
@@ -890,13 +890,13 @@ public:
 	}
 
 	FORCENOINLINE
-	FNodeTrackerHandle& AddNode(FNodeKey NodeKey, INode* Node)
+	FNodeTracker* AddNode(FNodeKey NodeKey, INode* Node)
 	{
 		LogDebugNode(TEXT("AddNode"), Node);
-		FNodeTrackerHandle& NodeTracker = NodeTrackers.Emplace(NodeKey, FNodeTrackerHandle(NodeKey, Node));
+		FNodeTracker* NodeTracker = NodeTrackers.Emplace(NodeKey, FNodeTrackerHandle(NodeKey, Node)).GetNodeTracker();
 		
-		NodeTrackersNames.FindOrAdd(NodeTracker.GetNodeTracker()->Name).Add(NodeTracker.GetNodeTracker());
-		InvalidatedNodeTrackers.Add(NodeTracker.GetNodeTracker());
+		NodeTrackersNames.FindOrAdd(NodeTracker->Name).Add(NodeTracker);
+		InvalidatedNodeTrackers.Add(NodeTracker);
 
 		return NodeTracker;
 	}
