Fix build break in SketchUp exporter
  - Since CL 19253774, conversion from FVector3f to FVector is explicit.

diff --git a/Engine/Source/Programs/Enterprise/Datasmith/DatasmithSketchUpRubyExporter/Private/DatasmithSketchUpCamera.cpp b/Engine/Source/Programs/Enterprise/Datasmith/DatasmithSketchUpRubyExporter/Private/DatasmithSketchUpCamera.cpp
--- a/Engine/Source/Programs/Enterprise/Datasmith/DatasmithSketchUpRubyExporter/Private/DatasmithSketchUpCamera.cpp
+++ b/Engine/Source/Programs/Enterprise/Datasmith/DatasmithSketchUpRubyExporter/Private/DatasmithSketchUpCamera.cpp
@@ -109,7 +109,7 @@ void FCamera::Update(FExportContext& Context)
 	// To avoid perturbating X, which is forward in Unreal, the handedness conversion is done by flipping the side vector Y.
 	// SketchUp uses inches as internal system unit for all 3D coordinates in the model while Unreal uses centimeters.
 
-	FVector Translation = DatasmithSketchUpUtils::FromSketchUp::ConvertPosition(SourcePosition);
+	FVector Translation = FVector(DatasmithSketchUpUtils::FromSketchUp::ConvertPosition(SourcePosition));
 
 	// Set the world transform of the Datasmith camera actor.
 	DatasmithCamera->SetRotation(Rotation);
@@ -126,7 +126,7 @@ void FCamera::Update(FExportContext& Context)
 	DatasmithCamera->SetFocalLength(float(FocalLength));
 
 	// Set the Datasmith camera focus distance (in centimeters).
-	FVector DistanceVector = DatasmithSketchUpUtils::FromSketchUp::ConvertPosition(SourceTarget.x - SourcePosition.x, SourceTarget.y - SourcePosition.y, SourceTarget.z - SourcePosition.z);
+	FVector DistanceVector = FVector(DatasmithSketchUpUtils::FromSketchUp::ConvertPosition(SourceTarget.x - SourcePosition.x, SourceTarget.y - SourcePosition.y, SourceTarget.z - SourcePosition.z));
 	DatasmithCamera->SetFocusDistance(DistanceVector.Size());
 
 }
diff --git a/Engine/Source/Programs/Enterprise/Datasmith/DatasmithSketchUpRubyExporter/Private/DatasmithSketchUpMesh.cpp b/Engine/Source/Programs/Enterprise/Datasmith/DatasmithSketchUpRubyExporter/Private/DatasmithSketchUpMesh.cpp
--- a/Engine/Source/Programs/Enterprise/Datasmith/DatasmithSketchUpRubyExporter/Private/DatasmithSketchUpMesh.cpp
+++ b/Engine/Source/Programs/Enterprise/Datasmith/DatasmithSketchUpRubyExporter/Private/DatasmithSketchUpMesh.cpp
@@ -306,7 +306,7 @@ namespace DatasmithSketchUp
 		for (int32 VertexNo = 0; VertexNo < VertexCount; VertexNo++)
 		{
 			SUPoint3D const& VertexPoint = MeshVertexPoints[VertexNo];
-			FVector P = DatasmithSketchUpUtils::FromSketchUp::ConvertPosition(VertexPoint);
+			FVector P = FVector(DatasmithSketchUpUtils::FromSketchUp::ConvertPosition(VertexPoint));
 			OutDMesh.SetVertex(VertexNo, P.X, P.Y, P.Z);
 		}
 
