[3ds Max] DirectLink: fixed connections icon size

- was clipped before, changed to large size

diff --git a/Engine/Source/Programs/Enterprise/Datasmith/DatasmithMaxExporter/Extras/3dsMaxScripts/UI/DatasmithUI.xaml b/Engine/Source/Programs/Enterprise/Datasmith/DatasmithMaxExporter/Extras/3dsMaxScripts/UI/DatasmithUI.xaml
--- a/Engine/Source/Programs/Enterprise/Datasmith/DatasmithMaxExporter/Extras/3dsMaxScripts/UI/DatasmithUI.xaml
+++ b/Engine/Source/Programs/Enterprise/Datasmith/DatasmithMaxExporter/Extras/3dsMaxScripts/UI/DatasmithUI.xaml
@@ -66,6 +66,7 @@
             MinWidth="0"
             ResizeStyle="HideText"
             ShowText="true"
+            Size="Large"
             Text="Connections"
             UID="e7499676-fcc7-4d7b-9b23-d71a9beea4e9" />
         </ads:RibbonRowPanel>
