Fixed crash on sync on all exporters.
  - Root cause: In cl 17646035, the value of bBuildWithEditorOnlyData in the different Target.cs failes was changed from true to false. Doing so implies the use of the wrong BulkData on which the Datasmith mesh serialization is relying on and generates the crash.
  - Solution: Move bBuildWithEditorOnlyData  back to true.
Note the reason why the change was made in the first place has not been found yet.

diff --git a/Engine/Source/Programs/Enterprise/Datasmith/DatasmithARCHICADExporter/DatsmithUE4ArchiCAD/DatasmithUE4ArchiCAD.Target.cs b/Engine/Source/Programs/Enterprise/Datasmith/DatasmithARCHICADExporter/DatsmithUE4ArchiCAD/DatasmithUE4ArchiCAD.Target.cs
--- a/Engine/Source/Programs/Enterprise/Datasmith/DatasmithARCHICADExporter/DatsmithUE4ArchiCAD/DatasmithUE4ArchiCAD.Target.cs
+++ b/Engine/Source/Programs/Enterprise/Datasmith/DatasmithARCHICADExporter/DatsmithUE4ArchiCAD/DatasmithUE4ArchiCAD.Target.cs
@@ -26,7 +26,7 @@ public class DatasmithUE4ArchiCADTarget : TargetRules
 
 		bBuildDeveloperTools = false;
 		bUseMallocProfiler = false;
-		bBuildWithEditorOnlyData = false;
+		bBuildWithEditorOnlyData = true;
 		bCompileAgainstEngine = false;
 		bCompileAgainstCoreUObject = true;
 		bCompileICU = false;
diff --git a/Engine/Source/Programs/Enterprise/Datasmith/DatasmithARCHICADExporter/UBTModules/DatasmithArchiCAD.Target.cs b/Engine/Source/Programs/Enterprise/Datasmith/DatasmithARCHICADExporter/UBTModules/DatasmithArchiCAD.Target.cs
--- a/Engine/Source/Programs/Enterprise/Datasmith/DatasmithARCHICADExporter/UBTModules/DatasmithArchiCAD.Target.cs
+++ b/Engine/Source/Programs/Enterprise/Datasmith/DatasmithARCHICADExporter/UBTModules/DatasmithArchiCAD.Target.cs
@@ -21,7 +21,7 @@ public class DatasmithArchiCADTarget : TargetRules
 
 		bBuildDeveloperTools = false;
 		bUseMallocProfiler = false;
-		bBuildWithEditorOnlyData = false;
+		bBuildWithEditorOnlyData = true;
 		bCompileAgainstEngine = false;
 		bCompileAgainstCoreUObject = true;
 		bCompileICU = false;
diff --git a/Engine/Source/Programs/Enterprise/Datasmith/DatasmithFacadeCSharp/DatasmithFacadeCSharp.Target.cs b/Engine/Source/Programs/Enterprise/Datasmith/DatasmithFacadeCSharp/DatasmithFacadeCSharp.Target.cs
--- a/Engine/Source/Programs/Enterprise/Datasmith/DatasmithFacadeCSharp/DatasmithFacadeCSharp.Target.cs
+++ b/Engine/Source/Programs/Enterprise/Datasmith/DatasmithFacadeCSharp/DatasmithFacadeCSharp.Target.cs
@@ -22,7 +22,7 @@ public class DatasmithFacadeCSharpTarget : TargetRules
 
 		bBuildDeveloperTools = false;
 		bUseMallocProfiler = false;
-		bBuildWithEditorOnlyData = false;
+		bBuildWithEditorOnlyData = true;
 		bCompileAgainstEngine = false;
 		bCompileAgainstCoreUObject = true;
 		bCompileICU = false;
diff --git a/Engine/Source/Programs/Enterprise/Datasmith/DatasmithNavisworksExporter/UBTModules/DatasmithNavisworks2020.Target.cs b/Engine/Source/Programs/Enterprise/Datasmith/DatasmithNavisworksExporter/UBTModules/DatasmithNavisworks2020.Target.cs
--- a/Engine/Source/Programs/Enterprise/Datasmith/DatasmithNavisworksExporter/UBTModules/DatasmithNavisworks2020.Target.cs
+++ b/Engine/Source/Programs/Enterprise/Datasmith/DatasmithNavisworksExporter/UBTModules/DatasmithNavisworks2020.Target.cs
@@ -23,7 +23,7 @@ public abstract class DatasmithNavisworksBaseTarget : TargetRules
 
 		bBuildDeveloperTools = false;
 		bUseMallocProfiler = false;
-		bBuildWithEditorOnlyData = false;
+		bBuildWithEditorOnlyData = true;
 		bCompileAgainstEngine = false;
 		bCompileAgainstCoreUObject = false;
 		bCompileICU = false;
diff --git a/Engine/Source/Programs/Enterprise/Datasmith/DatasmithRevitExporter/DatasmithRevit2018.Target.cs b/Engine/Source/Programs/Enterprise/Datasmith/DatasmithRevitExporter/DatasmithRevit2018.Target.cs
--- a/Engine/Source/Programs/Enterprise/Datasmith/DatasmithRevitExporter/DatasmithRevit2018.Target.cs
+++ b/Engine/Source/Programs/Enterprise/Datasmith/DatasmithRevitExporter/DatasmithRevit2018.Target.cs
@@ -25,7 +25,7 @@ public abstract class DatasmithRevitBaseTarget : TargetRules
 
 		bBuildDeveloperTools = false;
 		bUseMallocProfiler = false;
-		bBuildWithEditorOnlyData = false;
+		bBuildWithEditorOnlyData = true;
 		bCompileAgainstEngine = false;
 		bCompileAgainstCoreUObject = true;
 		bCompileICU = false;
diff --git a/Engine/Source/Programs/Enterprise/Datasmith/DatasmithRhinoExporter/DatasmithRhino6.Target.cs b/Engine/Source/Programs/Enterprise/Datasmith/DatasmithRhinoExporter/DatasmithRhino6.Target.cs
--- a/Engine/Source/Programs/Enterprise/Datasmith/DatasmithRhinoExporter/DatasmithRhino6.Target.cs
+++ b/Engine/Source/Programs/Enterprise/Datasmith/DatasmithRhinoExporter/DatasmithRhino6.Target.cs
@@ -34,7 +34,7 @@ public abstract class DatasmithRhinoBaseTarget : TargetRules
 
 		bBuildDeveloperTools = false;
 		bUseMallocProfiler = false;
-		bBuildWithEditorOnlyData = false;
+		bBuildWithEditorOnlyData = true;
 		bCompileAgainstEngine = false;
 		bCompileAgainstCoreUObject = true;
 		bCompileICU = false;
diff --git a/Engine/Source/Programs/Enterprise/Datasmith/DatasmithSDK/DatasmithSDK.Target.cs b/Engine/Source/Programs/Enterprise/Datasmith/DatasmithSDK/DatasmithSDK.Target.cs
--- a/Engine/Source/Programs/Enterprise/Datasmith/DatasmithSDK/DatasmithSDK.Target.cs
+++ b/Engine/Source/Programs/Enterprise/Datasmith/DatasmithSDK/DatasmithSDK.Target.cs
@@ -22,12 +22,12 @@ public class DatasmithSDKTarget : TargetRules
 
 		bBuildDeveloperTools = false;
 		bUseMallocProfiler = false;
-		bBuildWithEditorOnlyData = false;
+		bBuildWithEditorOnlyData = true;
 		bCompileAgainstEngine = false;
 		bCompileAgainstCoreUObject = true;
 		bCompileICU = false;
 		bUsesSlate = false;
-		bDisableDebugInfo = true;
+		bDisableDebugInfo = false;
 		bUsePDBFiles = true;
 		bHasExports = true;
 		bIsBuildingConsoleApplication = true;
diff --git a/Engine/Source/Programs/Enterprise/Datasmith/DatasmithSketchUpRubyExporter/DatasmithSketchUpRuby2020.Target.cs b/Engine/Source/Programs/Enterprise/Datasmith/DatasmithSketchUpRubyExporter/DatasmithSketchUpRuby2020.Target.cs
--- a/Engine/Source/Programs/Enterprise/Datasmith/DatasmithSketchUpRubyExporter/DatasmithSketchUpRuby2020.Target.cs
+++ b/Engine/Source/Programs/Enterprise/Datasmith/DatasmithSketchUpRubyExporter/DatasmithSketchUpRuby2020.Target.cs
@@ -19,7 +19,7 @@ public abstract class DatasmithSketchUpRubyBaseTarget : TargetRules
 
 		bBuildDeveloperTools = false;
 		bUseMallocProfiler = false;
-		bBuildWithEditorOnlyData = false;
+		bBuildWithEditorOnlyData = true;
 		bCompileAgainstEngine = false;
 		bCompileAgainstCoreUObject = true;
 		bCompileICU = false;
diff --git a/Engine/Source/Programs/Enterprise/Datasmith/DatasmithSolidworksExporter/DatasmithSolidworks2021.Target.cs b/Engine/Source/Programs/Enterprise/Datasmith/DatasmithSolidworksExporter/DatasmithSolidworks2021.Target.cs
--- a/Engine/Source/Programs/Enterprise/Datasmith/DatasmithSolidworksExporter/DatasmithSolidworks2021.Target.cs
+++ b/Engine/Source/Programs/Enterprise/Datasmith/DatasmithSolidworksExporter/DatasmithSolidworks2021.Target.cs
@@ -54,7 +54,7 @@ public abstract class DatasmithSolidworksBaseTarget : TargetRules
 		LinkType = TargetLinkType.Monolithic;
 		bBuildDeveloperTools = false;
 		bUseMallocProfiler = false;
-		bBuildWithEditorOnlyData = false;
+		bBuildWithEditorOnlyData = true;
 		bCompileAgainstEngine = false;
 		bCompileAgainstCoreUObject = false;
 		bCompileICU = false;
