[3ds Max] DirectLink: material slots/assignment fixes

- mesh with multiple material ids but single material(i.e. not a multimaterial) is consolidated into single mesh with one material slot. Simple test - create a box(in Max it defaults to 6 matIds) and assign a material then sync/export/observe
- sorted material slots to make them in the same order as in multimaterial
- unused materials removed from Datasmith scene. I.e. assign material - sync - assign different material - first material should go away from udatasmith if not used on other nodes

diff --git a/Engine/Source/Programs/Enterprise/Datasmith/DatasmithMaxExporter/Private/DatasmithMaxDirectLinkMaterial.cpp b/Engine/Source/Programs/Enterprise/Datasmith/DatasmithMaxExporter/Private/DatasmithMaxDirectLinkMaterial.cpp
--- a/Engine/Source/Programs/Enterprise/Datasmith/DatasmithMaxExporter/Private/DatasmithMaxDirectLinkMaterial.cpp
+++ b/Engine/Source/Programs/Enterprise/Datasmith/DatasmithMaxExporter/Private/DatasmithMaxDirectLinkMaterial.cpp
@@ -151,6 +151,7 @@ void FMaterialsCollectionTracker::ReleaseMaterial(FMaterialTracker& MaterialTrac
 {
 	RemoveConvertedMaterial(MaterialTracker);
 	MaterialTrackers.Remove(MaterialTracker.Material);
+	InvalidatedMaterialTrackers.Remove(&MaterialTracker);
 }
 
 void FMaterialsCollectionTracker::RemoveConvertedMaterial(FMaterialTracker& MaterialTracker)
diff --git a/Engine/Source/Programs/Enterprise/Datasmith/DatasmithMaxExporter/Private/DatasmithMaxDirectLinkMesh.cpp b/Engine/Source/Programs/Enterprise/Datasmith/DatasmithMaxExporter/Private/DatasmithMaxDirectLinkMesh.cpp
--- a/Engine/Source/Programs/Enterprise/Datasmith/DatasmithMaxExporter/Private/DatasmithMaxDirectLinkMesh.cpp
+++ b/Engine/Source/Programs/Enterprise/Datasmith/DatasmithMaxExporter/Private/DatasmithMaxDirectLinkMesh.cpp
@@ -390,7 +390,7 @@ bool FillDatasmithMeshFromBoundingBox(FDatasmithMesh& DatasmithMesh, INode* Expo
 }
 
 
-bool CreateDatasmithMeshFromMaxMesh(FDatasmithMesh& DatasmithMesh, INode* Node, const TCHAR* MeshName, const FRenderMeshForConversion& RenderMesh, TSet<uint16>& SupportedChannels, TMap<int32, int32>& UVChannelsMap)
+bool CreateDatasmithMeshFromMaxMesh(FDatasmithMesh& DatasmithMesh, INode* Node, const TCHAR* MeshName, const FRenderMeshForConversion& RenderMesh, bool bConsolidateMaterialIds, TSet<uint16>& SupportedChannels, TMap<int32, int32>& UVChannelsMap)
 {
 	bool bResult = false;
 	if (RenderMesh.GetMesh()->getNumFaces())
@@ -410,7 +410,7 @@ bool CreateDatasmithMeshFromMaxMesh(FDatasmithMesh& DatasmithMesh, INode* Node,
 		if (CachedMesh.getNumFaces() > 0)
 		{
 			// todo: pivot
-			FillDatasmithMeshFromMaxMesh(DatasmithMesh, CachedMesh, Node, false, SupportedChannels, UVChannelsMap, RenderMesh.GetPivot());
+			FillDatasmithMeshFromMaxMesh(DatasmithMesh, CachedMesh, Node, bConsolidateMaterialIds, SupportedChannels, UVChannelsMap, RenderMesh.GetPivot());
 
 			bResult = true; // Set to true, don't care what ExportToUObject does here - we need to move it to a thread anyway
 		}
@@ -420,13 +420,14 @@ bool CreateDatasmithMeshFromMaxMesh(FDatasmithMesh& DatasmithMesh, INode* Node,
 }
 
 // todo: paralelize calls to ExportToUObject 
-bool ConvertMaxMeshToDatasmith(ISceneTracker& Scene, TSharedPtr<IDatasmithMeshElement>& DatasmithMeshElement, INode* Node, const TCHAR* MeshName, const FRenderMeshForConversion& RenderMesh, TSet<uint16>& SupportedChannels, const FRenderMeshForConversion& CollisionMesh)
+bool ConvertMaxMeshToDatasmith(ISceneTracker& Scene, TSharedPtr<IDatasmithMeshElement>& DatasmithMeshElement, INode* Node, const TCHAR* MeshName, const FRenderMeshForConversion& RenderMesh, bool bConsolidateMaterialIds, TSet<uint16>& SupportedChannels, const FRenderMeshForConversion& CollisionMesh)
 {
 	// Reset old mesh
 	if (DatasmithMeshElement)
 	{
 		// todo: potential mesh reuse - when DatasmithMeshElement allows to reset materials(as well as other params)
 		Scene.ReleaseMeshElement(DatasmithMeshElement);
+		SupportedChannels.Reset();
 	}
 
 	TOptional<FDatasmithMaxStaticMeshAttributes> DatasmithAttributes = FDatasmithMaxStaticMeshAttributes::ExtractStaticMeshAttributes(Node);
@@ -450,7 +451,7 @@ bool ConvertMaxMeshToDatasmith(ISceneTracker& Scene, TSharedPtr<IDatasmithMeshEl
 
 	FDatasmithMesh DatasmithMesh;
 	TMap<int32, int32> UVChannelsMap;
-	if (!CreateDatasmithMeshFromMaxMesh(DatasmithMesh, Node, MeshName, RenderMesh, SupportedChannels, UVChannelsMap))
+	if (!CreateDatasmithMeshFromMaxMesh(DatasmithMesh, Node, MeshName, RenderMesh, bConsolidateMaterialIds, SupportedChannels, UVChannelsMap))
 	{
 		LogWarningDialog(FString(TEXT("Invalid object: ")) + Node->GetName());
 		return false;
@@ -483,7 +484,7 @@ bool ConvertMaxMeshToDatasmith(ISceneTracker& Scene, TSharedPtr<IDatasmithMeshEl
 	{
 		TSet<uint16> SupportedChannelsDummy; // ignore map channels for collision mesh
 		TMap<int32, int32> UVChannelsMapDummy;
-		if (CreateDatasmithMeshFromMaxMesh(DatasmithCollisionMesh, CollisionMesh.GetNode(), nullptr, CollisionMesh, SupportedChannelsDummy, UVChannelsMapDummy))
+		if (CreateDatasmithMeshFromMaxMesh(DatasmithCollisionMesh, CollisionMesh.GetNode(), nullptr, CollisionMesh, true, SupportedChannelsDummy, UVChannelsMapDummy))
 		{
 			DatasmithCollisionMeshPtr = &DatasmithCollisionMesh;
 		}
diff --git a/Engine/Source/Programs/Enterprise/Datasmith/DatasmithMaxExporter/Private/DatasmithMaxExporterWithDirectLink.cpp b/Engine/Source/Programs/Enterprise/Datasmith/DatasmithMaxExporter/Private/DatasmithMaxExporterWithDirectLink.cpp
--- a/Engine/Source/Programs/Enterprise/Datasmith/DatasmithMaxExporter/Private/DatasmithMaxExporterWithDirectLink.cpp
+++ b/Engine/Source/Programs/Enterprise/Datasmith/DatasmithMaxExporter/Private/DatasmithMaxExporterWithDirectLink.cpp
@@ -1044,7 +1044,10 @@ public:
 			if (TUniquePtr<FInstances>* InstancesPtr = InstancesForAnimHandle.Find(NodeTracker.InstanceHandle))
 			{
 				FInstances& Instances = **InstancesPtr;
+				UnregisterNodeForMaterial(NodeTracker);
+
 				Instances.NodeTrackers.Remove(&NodeTracker);
+
 				if (!Instances.NodeTrackers.Num())
 				{
 					ReleaseMeshElement(Instances.DatasmithMeshElement);
@@ -1385,24 +1388,14 @@ public:
 		if (!NodeTracker.MaterialTracker || NodeTracker.MaterialTracker->Material != Material)
 		{
 			// Release old material
-			if (NodeTracker.MaterialTracker)
-			{
-				// Release material assignment
-				MaterialsAssignedToNodes[NodeTracker.MaterialTracker].Remove(&NodeTracker);
-
-				// Clean tracker if it's not used aby any node
-				if (MaterialsAssignedToNodes[NodeTracker.MaterialTracker].IsEmpty())
-				{
-					MaterialsCollectionTracker.ReleaseMaterial(*NodeTracker.MaterialTracker);
-				}
-			}
+			UnregisterNodeForMaterial(NodeTracker);
 
 			NodeTracker.MaterialTracker = MaterialsCollectionTracker.AddMaterial(Material);
 			MaterialsAssignedToNodes.FindOrAdd(NodeTracker.MaterialTracker).Add(&NodeTracker);
 		}
 	}
 
-	void UnregisterNodeForMaterial(FNodeTracker& NodeTracker, Mtl* Material)
+	void UnregisterNodeForMaterial(FNodeTracker& NodeTracker)
 	{
 		if (NodeTracker.MaterialTracker)
 		{
@@ -1411,6 +1404,7 @@ public:
 			{
 				MaterialsCollectionTracker.ReleaseMaterial(*NodeTracker.MaterialTracker);
 			}
+			NodeTracker.MaterialTracker = nullptr;
 		}
 	}
 
@@ -1515,8 +1509,7 @@ public:
 			else
 			{
 				// Release old material
-				UnregisterNodeForMaterial(NodeTracker, Material);
-				NodeTracker.MaterialTracker = nullptr;
+				UnregisterNodeForMaterial(NodeTracker);
 				NodeTracker.DatasmithMeshActor->ResetMaterialOverrides();
 			}
 
@@ -1552,7 +1545,18 @@ public:
 
 		if (RenderMesh.GetMesh())
 		{
-			if (ConvertMaxMeshToDatasmith(*this, DatasmithMeshElement, Node, *MeshName, RenderMesh, SupportedChannels, CollisionMesh)) // export might produce anything(e.g. if mesh is empty)
+			bool bHasInstanceWithMultimat = false;
+			for (FNodeTracker* InstanceNodeTracker : Instances.NodeTrackers)
+			{
+				if (Mtl* Material = InstanceNodeTracker->Node->GetMtl())
+				{
+					if (FDatasmithMaxMatHelper::GetMaterialClass(Material) == EDSMaterialType::MultiMat)
+					{
+						bHasInstanceWithMultimat = true;
+					}
+				}
+			}
+			if (ConvertMaxMeshToDatasmith(*this, DatasmithMeshElement, Node, *MeshName, RenderMesh, !bHasInstanceWithMultimat, SupportedChannels, CollisionMesh)) // export might produce anything(e.g. if mesh is empty)
 			{
 				DatasmithMeshElement->SetLabel(Node->GetName());
 				return true;
@@ -1709,7 +1713,7 @@ public:
 		TSharedPtr<IDatasmithMeshElement> DatasmithMeshElement;
 		TSet<uint16> SupportedChannels;
 
-		if (ConvertMaxMeshToDatasmith(*this, DatasmithMeshElement, GeometryNode, *MeshName, RenderMesh, SupportedChannels))
+		if (ConvertMaxMeshToDatasmith(*this, DatasmithMeshElement, GeometryNode, *MeshName, RenderMesh, false, SupportedChannels))
 		{
 			TUniquePtr<FRailClonesConverted>& RailClonesConverted = RailClones.FindOrAdd(&NodeTracker);
 			if (!RailClonesConverted)
diff --git a/Engine/Source/Programs/Enterprise/Datasmith/DatasmithMaxExporter/Private/DatasmithMaxMeshExporter.cpp b/Engine/Source/Programs/Enterprise/Datasmith/DatasmithMaxExporter/Private/DatasmithMaxMeshExporter.cpp
--- a/Engine/Source/Programs/Enterprise/Datasmith/DatasmithMaxExporter/Private/DatasmithMaxMeshExporter.cpp
+++ b/Engine/Source/Programs/Enterprise/Datasmith/DatasmithMaxExporter/Private/DatasmithMaxMeshExporter.cpp
@@ -332,7 +332,9 @@ void AssignMeshMaterials(TSharedPtr<IDatasmithMeshElement>& MeshElement, Mtl* Ma
 	}
 	if (Material != nullptr)
 	{
-		for (uint16 Channel : SupportedChannels)
+		TArray<uint16> ChannelsSorted = SupportedChannels.Array();
+		ChannelsSorted.Sort();
+		for (uint16 Channel : ChannelsSorted)
 		{
 			//Max's channel UI is not zero-based, so we register an incremented ChannelID for better visual consistency after importing in Unreal.
 			uint16 DisplayedChannelID = Channel + 1;
diff --git a/Engine/Source/Programs/Enterprise/Datasmith/DatasmithMaxExporter/Public/DatasmithMaxDirectLink.h b/Engine/Source/Programs/Enterprise/Datasmith/DatasmithMaxExporter/Public/DatasmithMaxDirectLink.h
--- a/Engine/Source/Programs/Enterprise/Datasmith/DatasmithMaxExporter/Public/DatasmithMaxDirectLink.h
+++ b/Engine/Source/Programs/Enterprise/Datasmith/DatasmithMaxExporter/Public/DatasmithMaxDirectLink.h
@@ -243,7 +243,7 @@ void FillDatasmithMeshFromMaxMesh(FDatasmithMesh& DatasmithMesh, Mesh& MaxMesh,
 
 
 // Creates Mesh element and converts max mesh into it
-bool ConvertMaxMeshToDatasmith(ISceneTracker& Scene, TSharedPtr<IDatasmithMeshElement>& DatasmithMeshElement, INode* Node, const TCHAR* MeshName, const FRenderMeshForConversion& RenderMesh, TSet<uint16>& SupportedChannels, const FRenderMeshForConversion& CollisionMesh = FRenderMeshForConversion());
+bool ConvertMaxMeshToDatasmith(ISceneTracker& Scene, TSharedPtr<IDatasmithMeshElement>& DatasmithMeshElement, INode* Node, const TCHAR* MeshName, const FRenderMeshForConversion& RenderMesh, bool bConsolidateMaterialIds, TSet<uint16>& SupportedChannels, const FRenderMeshForConversion& CollisionMesh = FRenderMeshForConversion());
 
 bool OpenDirectLinkUI();
 const TCHAR* GetDirectlinkCacheDirectory();
