[3ds Max] DirectLink: Forest - fixed regression

diff --git a/Engine/Source/Programs/Enterprise/Datasmith/DatasmithMaxExporter/Private/DatasmithMaxDirectLinkItoo.cpp b/Engine/Source/Programs/Enterprise/Datasmith/DatasmithMaxExporter/Private/DatasmithMaxDirectLinkItoo.cpp
--- a/Engine/Source/Programs/Enterprise/Datasmith/DatasmithMaxExporter/Private/DatasmithMaxDirectLinkItoo.cpp
+++ b/Engine/Source/Programs/Enterprise/Datasmith/DatasmithMaxExporter/Private/DatasmithMaxDirectLinkItoo.cpp
@@ -197,6 +197,7 @@ bool ConvertForest(ISceneTracker& Scene, FNodeTracker& NodeTracker, Object* Obj)
 			if (RenderMesh.IsValid())
 			{
 				Scene.SetupDatasmithHISMForNode(NodeTracker, GeometryNode, RenderMesh, NodeTracker.Node->GetMtl(), MeshIndex, Instances.Transforms);
+				MeshIndex ++;
 			}
 		}
 	}
