Datasmith: add support for CastShadow property in Actor Elements

diff --git a/Engine/Plugins/Enterprise/DatasmithContent/Source/DatasmithContent/Private/ObjectTemplates/DatasmithSceneComponentTemplate.cpp b/Engine/Plugins/Enterprise/DatasmithContent/Source/DatasmithContent/Private/ObjectTemplates/DatasmithSceneComponentTemplate.cpp
--- a/Engine/Plugins/Enterprise/DatasmithContent/Source/DatasmithContent/Private/ObjectTemplates/DatasmithSceneComponentTemplate.cpp
+++ b/Engine/Plugins/Enterprise/DatasmithContent/Source/DatasmithContent/Private/ObjectTemplates/DatasmithSceneComponentTemplate.cpp
@@ -44,6 +44,14 @@ UObject* UDatasmithSceneComponentTemplate::UpdateObject( UObject* Destination, b
 		SceneComponent->SetVisibility( bVisible );
 	}
 
+	if ( UPrimitiveComponent* PrimComponent = Cast< UPrimitiveComponent >( SceneComponent ) )
+	{
+		if ( !PreviousTemplate || PreviousTemplate->bCastShadow == PrimComponent->CastShadow )
+		{
+			PrimComponent->SetCastShadow( bCastShadow );
+		}
+	}
+
 	const ULevel* SceneComponentLevel = SceneComponent->GetComponentLevel();
 	bool bCanAttach = ( AttachParent && AttachParent->GetComponentLevel() == SceneComponentLevel );
 	bCanAttach |= ( !AttachParent && ( SceneComponentLevel == nullptr || SceneComponentLevel->OwningWorld == Destination->GetWorld() ) );
@@ -126,6 +134,8 @@ void UDatasmithSceneComponentTemplate::Load( const UObject* Source )
 	AttachParent = SceneComponent->GetAttachParent();
 	Tags = TSet<FName>(SceneComponent->ComponentTags);
 
+	const UPrimitiveComponent* PrimComponent = Cast< const UPrimitiveComponent >( SceneComponent );
+	bCastShadow = PrimComponent ? PrimComponent->CastShadow : true;
 #endif // #if WITH_EDITORONLY_DATA
 }
 
@@ -141,6 +151,7 @@ bool UDatasmithSceneComponentTemplate::Equals( const UDatasmithObjectTemplate* O
 	bool bEquals = AreTransformsEqual( RelativeTransform, TypedOther->RelativeTransform );
 	bEquals = bEquals && ( Mobility == TypedOther->Mobility );
 	bEquals = bEquals && ( bVisible == TypedOther->bVisible );
+	bEquals = bEquals && ( bCastShadow == TypedOther->bCastShadow );
 	bEquals = bEquals && ( AttachParent == TypedOther->AttachParent );
 	bEquals = bEquals && FDatasmithObjectTemplateUtils::SetsEquals(Tags, TypedOther->Tags);
 
diff --git a/Engine/Plugins/Enterprise/DatasmithContent/Source/DatasmithContent/Public/ObjectTemplates/DatasmithSceneComponentTemplate.h b/Engine/Plugins/Enterprise/DatasmithContent/Source/DatasmithContent/Public/ObjectTemplates/DatasmithSceneComponentTemplate.h
--- a/Engine/Plugins/Enterprise/DatasmithContent/Source/DatasmithContent/Public/ObjectTemplates/DatasmithSceneComponentTemplate.h
+++ b/Engine/Plugins/Enterprise/DatasmithContent/Source/DatasmithContent/Public/ObjectTemplates/DatasmithSceneComponentTemplate.h
@@ -26,6 +26,9 @@ public:
 	UPROPERTY()
 	bool bVisible;
 
+	UPROPERTY()
+	bool bCastShadow = true;
+
 	UPROPERTY()
 	TSet<FName> Tags;
 
diff --git a/Engine/Plugins/Enterprise/DatasmithImporter/Source/DatasmithImporter/Private/DatasmithActorImporter.cpp b/Engine/Plugins/Enterprise/DatasmithImporter/Source/DatasmithImporter/Private/DatasmithActorImporter.cpp
--- a/Engine/Plugins/Enterprise/DatasmithImporter/Source/DatasmithImporter/Private/DatasmithActorImporter.cpp
+++ b/Engine/Plugins/Enterprise/DatasmithImporter/Source/DatasmithImporter/Private/DatasmithActorImporter.cpp
@@ -616,6 +616,7 @@ void FDatasmithActorImporter::SetupSceneComponent( USceneComponent* SceneCompone
 	SceneComponentTemplate->RelativeTransform = ActorElement->GetRelativeTransform();
 	SceneComponentTemplate->Mobility = ActorElement->IsA(EDatasmithElementType::Camera) ? EComponentMobility::Movable : EComponentMobility::Static;
 	SceneComponentTemplate->bVisible = ActorElement->GetVisibility();
+	SceneComponentTemplate->bCastShadow = ActorElement->GetCastShadow();
 	SceneComponentTemplate->AttachParent = Parent;
 
 	// Add tags from ActorElement to SceneComponentTemplate
diff --git a/Engine/Source/Developer/Datasmith/DatasmithFacade/Private/DatasmithFacadeActor.cpp b/Engine/Source/Developer/Datasmith/DatasmithFacade/Private/DatasmithFacadeActor.cpp
--- a/Engine/Source/Developer/Datasmith/DatasmithFacade/Private/DatasmithFacadeActor.cpp
+++ b/Engine/Source/Developer/Datasmith/DatasmithFacade/Private/DatasmithFacadeActor.cpp
@@ -191,6 +191,16 @@ bool FDatasmithFacadeActor::GetVisibility() const
 	return GetDatasmithActorElement()->GetVisibility();
 }
 
+void FDatasmithFacadeActor::SetCastShadow(bool bInCastShadow)
+{
+	GetDatasmithActorElement()->SetCastShadow(bInCastShadow);
+}
+
+bool FDatasmithFacadeActor::GetCastShadow() const
+{
+	return GetDatasmithActorElement()->GetCastShadow();
+}
+
 FDatasmithFacadeActor::EActorType FDatasmithFacadeActor::GetActorType() const
 {
 	return FDatasmithFacadeActor::GetActorType(GetDatasmithActorElement());
diff --git a/Engine/Source/Developer/Datasmith/DatasmithFacade/Public/DatasmithFacadeActor.h b/Engine/Source/Developer/Datasmith/DatasmithFacade/Public/DatasmithFacadeActor.h
--- a/Engine/Source/Developer/Datasmith/DatasmithFacade/Public/DatasmithFacadeActor.h
+++ b/Engine/Source/Developer/Datasmith/DatasmithFacade/Public/DatasmithFacadeActor.h
@@ -56,13 +56,13 @@ public:
 		float Yaw,
 		float Roll
 	);
-	
+
 	void GetRotation(
 		float& OutPitch,
 		float& OutYaw,
 		float& OutRoll
 	) const;
-	
+
 	void SetRotation(
 		float X,
 		float Y,
@@ -82,7 +82,7 @@ public:
 		float Y,
 		float Z
 	);
-	
+
 	void GetTranslation(
 		float& OutX,
 		float& OutY,
@@ -129,7 +129,7 @@ public:
 	// Get the number of children on this actor
 	int32 GetChildrenCount() const;
 
-	/** 
+	/**
 	 *	Returns a new FDatasmithFacadeActor pointing to the InIndex-th child of the mesh actor
 	 *	If there is no child at the given index, returned value is nullptr.
 	 *	The caller is responsible of deleting the returned object pointer.
@@ -160,6 +160,12 @@ public:
 	// Set a mesh actor's visibility
 	bool GetVisibility() const;
 
+	void SetCastShadow(
+		bool bInCastShadow
+	);
+
+	bool GetCastShadow() const;
+
 	EActorType GetActorType() const;
 
 #ifdef SWIG_FACADE
diff --git a/Engine/Source/Programs/Enterprise/Datasmith/DatasmithFacadeCSharp/Private/DatasmithFacadeCSharp.cpp b/Engine/Source/Programs/Enterprise/Datasmith/DatasmithFacadeCSharp/Private/DatasmithFacadeCSharp.cpp
--- a/Engine/Source/Programs/Enterprise/Datasmith/DatasmithFacadeCSharp/Private/DatasmithFacadeCSharp.cpp
+++ b/Engine/Source/Programs/Enterprise/Datasmith/DatasmithFacadeCSharp/Private/DatasmithFacadeCSharp.cpp
@@ -3760,6 +3760,28 @@ SWIGEXPORT unsigned int SWIGSTDCALL CSharp_FDatasmithFacadeActor_GetVisibility(v
 }
 
 
+SWIGEXPORT void SWIGSTDCALL CSharp_FDatasmithFacadeActor_SetCastShadow(void * jarg1, unsigned int jarg2) {
+  FDatasmithFacadeActor *arg1 = (FDatasmithFacadeActor *) 0 ;
+  bool arg2 ;
+  
+  arg1 = (FDatasmithFacadeActor *)jarg1; 
+  arg2 = jarg2 ? true : false; 
+  (arg1)->SetCastShadow(arg2);
+}
+
+
+SWIGEXPORT unsigned int SWIGSTDCALL CSharp_FDatasmithFacadeActor_GetCastShadow(void * jarg1) {
+  unsigned int jresult ;
+  FDatasmithFacadeActor *arg1 = (FDatasmithFacadeActor *) 0 ;
+  bool result;
+  
+  arg1 = (FDatasmithFacadeActor *)jarg1; 
+  result = (bool)((FDatasmithFacadeActor const *)arg1)->GetCastShadow();
+  jresult = result; 
+  return jresult;
+}
+
+
 SWIGEXPORT int SWIGSTDCALL CSharp_FDatasmithFacadeActor_GetActorType(void * jarg1) {
   int jresult ;
   FDatasmithFacadeActor *arg1 = (FDatasmithFacadeActor *) 0 ;
@@ -6726,6 +6748,16 @@ SWIGEXPORT void SWIGSTDCALL CSharp_delete_FDatasmithFacadeDirectLink(void * jarg
 }
 
 
+SWIGEXPORT void * SWIGSTDCALL CSharp_new_FDatasmithFacadeUniqueNameProvider() {
+  void * jresult ;
+  FDatasmithFacadeUniqueNameProvider *result = 0 ;
+  
+  result = (FDatasmithFacadeUniqueNameProvider *)new FDatasmithFacadeUniqueNameProvider();
+  jresult = (void *)result; 
+  return jresult;
+}
+
+
 SWIGEXPORT wchar_t * SWIGSTDCALL CSharp_FDatasmithFacadeUniqueNameProvider_GenerateUniqueName(void * jarg1, wchar_t * jarg2) {
   wchar_t * jresult ;
   FDatasmithFacadeUniqueNameProvider *arg1 = (FDatasmithFacadeUniqueNameProvider *) 0 ;
@@ -6770,16 +6802,6 @@ SWIGEXPORT void SWIGSTDCALL CSharp_FDatasmithFacadeUniqueNameProvider_RemoveExis
 }
 
 
-SWIGEXPORT void * SWIGSTDCALL CSharp_new_FDatasmithFacadeUniqueNameProvider() {
-  void * jresult ;
-  FDatasmithFacadeUniqueNameProvider *result = 0 ;
-  
-  result = (FDatasmithFacadeUniqueNameProvider *)new FDatasmithFacadeUniqueNameProvider();
-  jresult = (void *)result; 
-  return jresult;
-}
-
-
 SWIGEXPORT void SWIGSTDCALL CSharp_delete_FDatasmithFacadeUniqueNameProvider(void * jarg1) {
   FDatasmithFacadeUniqueNameProvider *arg1 = (FDatasmithFacadeUniqueNameProvider *) 0 ;
   
diff --git a/Engine/Source/Programs/Enterprise/Datasmith/DatasmithFacadeCSharp/Public/DatasmithFacadeCSharpPINVOKE.cs b/Engine/Source/Programs/Enterprise/Datasmith/DatasmithFacadeCSharp/Public/DatasmithFacadeCSharpPINVOKE.cs
--- a/Engine/Source/Programs/Enterprise/Datasmith/DatasmithFacadeCSharp/Public/DatasmithFacadeCSharpPINVOKE.cs
+++ b/Engine/Source/Programs/Enterprise/Datasmith/DatasmithFacadeCSharp/Public/DatasmithFacadeCSharpPINVOKE.cs
@@ -1082,6 +1082,12 @@ class DatasmithFacadeCSharpPINVOKE {
   [global::System.Runtime.InteropServices.DllImport("DatasmithFacadeCSharp", EntryPoint="CSharp_FDatasmithFacadeActor_GetVisibility")]
   public static extern bool FDatasmithFacadeActor_GetVisibility(global::System.Runtime.InteropServices.HandleRef jarg1);
 
+  [global::System.Runtime.InteropServices.DllImport("DatasmithFacadeCSharp", EntryPoint="CSharp_FDatasmithFacadeActor_SetCastShadow")]
+  public static extern void FDatasmithFacadeActor_SetCastShadow(global::System.Runtime.InteropServices.HandleRef jarg1, bool jarg2);
+
+  [global::System.Runtime.InteropServices.DllImport("DatasmithFacadeCSharp", EntryPoint="CSharp_FDatasmithFacadeActor_GetCastShadow")]
+  public static extern bool FDatasmithFacadeActor_GetCastShadow(global::System.Runtime.InteropServices.HandleRef jarg1);
+
   [global::System.Runtime.InteropServices.DllImport("DatasmithFacadeCSharp", EntryPoint="CSharp_FDatasmithFacadeActor_GetActorType")]
   public static extern int FDatasmithFacadeActor_GetActorType(global::System.Runtime.InteropServices.HandleRef jarg1);
 
@@ -1859,6 +1865,9 @@ class DatasmithFacadeCSharpPINVOKE {
   [global::System.Runtime.InteropServices.DllImport("DatasmithFacadeCSharp", EntryPoint="CSharp_delete_FDatasmithFacadeDirectLink")]
   public static extern void delete_FDatasmithFacadeDirectLink(global::System.Runtime.InteropServices.HandleRef jarg1);
 
+  [global::System.Runtime.InteropServices.DllImport("DatasmithFacadeCSharp", EntryPoint="CSharp_new_FDatasmithFacadeUniqueNameProvider")]
+  public static extern global::System.IntPtr new_FDatasmithFacadeUniqueNameProvider();
+
   [global::System.Runtime.InteropServices.DllImport("DatasmithFacadeCSharp", EntryPoint="CSharp_FDatasmithFacadeUniqueNameProvider_GenerateUniqueName")]
   public static extern global::System.IntPtr FDatasmithFacadeUniqueNameProvider_GenerateUniqueName(global::System.Runtime.InteropServices.HandleRef jarg1, [global::System.Runtime.InteropServices.MarshalAs(global::System.Runtime.InteropServices.UnmanagedType.LPWStr)]string jarg2);
 
@@ -1871,9 +1880,6 @@ class DatasmithFacadeCSharpPINVOKE {
   [global::System.Runtime.InteropServices.DllImport("DatasmithFacadeCSharp", EntryPoint="CSharp_FDatasmithFacadeUniqueNameProvider_RemoveExistingName")]
   public static extern void FDatasmithFacadeUniqueNameProvider_RemoveExistingName(global::System.Runtime.InteropServices.HandleRef jarg1, [global::System.Runtime.InteropServices.MarshalAs(global::System.Runtime.InteropServices.UnmanagedType.LPWStr)]string jarg2);
 
-  [global::System.Runtime.InteropServices.DllImport("DatasmithFacadeCSharp", EntryPoint="CSharp_new_FDatasmithFacadeUniqueNameProvider")]
-  public static extern global::System.IntPtr new_FDatasmithFacadeUniqueNameProvider();
-
   [global::System.Runtime.InteropServices.DllImport("DatasmithFacadeCSharp", EntryPoint="CSharp_delete_FDatasmithFacadeUniqueNameProvider")]
   public static extern void delete_FDatasmithFacadeUniqueNameProvider(global::System.Runtime.InteropServices.HandleRef jarg1);
 
diff --git a/Engine/Source/Programs/Enterprise/Datasmith/DatasmithFacadeCSharp/Public/FDatasmithFacadeActor.cs b/Engine/Source/Programs/Enterprise/Datasmith/DatasmithFacadeCSharp/Public/FDatasmithFacadeActor.cs
--- a/Engine/Source/Programs/Enterprise/Datasmith/DatasmithFacadeCSharp/Public/FDatasmithFacadeActor.cs
+++ b/Engine/Source/Programs/Enterprise/Datasmith/DatasmithFacadeCSharp/Public/FDatasmithFacadeActor.cs
@@ -206,6 +206,15 @@ public class FDatasmithFacadeActor : FDatasmithFacadeElement {
     return ret;
   }
 
+  public void SetCastShadow(bool bInCastShadow) {
+    DatasmithFacadeCSharpPINVOKE.FDatasmithFacadeActor_SetCastShadow(swigCPtr, bInCastShadow);
+  }
+
+  public bool GetCastShadow() {
+    bool ret = DatasmithFacadeCSharpPINVOKE.FDatasmithFacadeActor_GetCastShadow(swigCPtr);
+    return ret;
+  }
+
   public FDatasmithFacadeActor.EActorType GetActorType() {
     FDatasmithFacadeActor.EActorType ret = (FDatasmithFacadeActor.EActorType)DatasmithFacadeCSharpPINVOKE.FDatasmithFacadeActor_GetActorType(swigCPtr);
     return ret;
diff --git a/Engine/Source/Runtime/Datasmith/DatasmithCore/Private/DatasmithSceneElementsImpl.h b/Engine/Source/Runtime/Datasmith/DatasmithCore/Private/DatasmithSceneElementsImpl.h
--- a/Engine/Source/Runtime/Datasmith/DatasmithCore/Private/DatasmithSceneElementsImpl.h
+++ b/Engine/Source/Runtime/Datasmith/DatasmithCore/Private/DatasmithSceneElementsImpl.h
@@ -190,12 +190,15 @@ public:
 		return Parent.View();
 	}
 
-	virtual void SetIsAComponent(bool Value) { UPDATE_BITFLAGS(Flags, Value, EActorFlags::IsAComponent); }
+	virtual void SetIsAComponent(bool Value) override { UPDATE_BITFLAGS(Flags, Value, EActorFlags::IsAComponent); }
 	virtual bool IsAComponent() const override { return !!(Flags & EActorFlags::IsAComponent); }
 
 	virtual void SetVisibility(bool bInVisibility) override { UPDATE_BITFLAGS(Flags, bInVisibility, EActorFlags::IsVisible); }
 	virtual bool GetVisibility() const override { return !!(Flags & EActorFlags::IsVisible); }
 
+	virtual void SetCastShadow(bool bInCastShadow) override { bCastShadow = bInCastShadow; }
+	virtual bool GetCastShadow() const override { return bCastShadow; }
+
 protected:
 	/** Converts all children's transforms to relative */
 	void ConvertChildsToRelative();
@@ -222,7 +225,7 @@ private:
 	TDatasmithReferenceProxy<IDatasmithActorElement> Parent;
 
 	TReflected<EActorFlags, uint8> Flags;
-	TReflected<int32> SelectionIdx;
+	TReflected<bool> bCastShadow;
 };
 
 template< typename InterfaceType >
@@ -235,7 +238,7 @@ inline FDatasmithActorElementImpl<T>::FDatasmithActorElementImpl(const TCHAR* In
 	, Scale(FVector::OneVector)
 	, Rotation(FQuat::Identity)
 	, Flags(EActorFlags::IsVisible)
-	, SelectionIdx(-1)
+	, bCastShadow(true)
 {
 	this->RegisterReferenceProxy(Children, "Children");
 	this->RegisterReferenceProxy(Parent,   "Parent"  );
@@ -245,8 +248,8 @@ inline FDatasmithActorElementImpl<T>::FDatasmithActorElementImpl(const TCHAR* In
 	Store.RegisterParameter(Rotation,     "Rotation"     );
 	Store.RegisterParameter(Layer,        "Layer"        );
 	Store.RegisterParameter(Tags,         "Tags"         ); // reflect as low prio for directlink
-	Store.RegisterParameter(SelectionIdx, "SelectionIdx" );
 	Store.RegisterParameter(Flags,        "Flags"        );
+	Store.RegisterParameter(bCastShadow,  "CastShadow"   );
 }
 
 template< typename T >
diff --git a/Engine/Source/Runtime/Datasmith/DatasmithCore/Private/DatasmithSceneXmlReader.cpp b/Engine/Source/Runtime/Datasmith/DatasmithCore/Private/DatasmithSceneXmlReader.cpp
--- a/Engine/Source/Runtime/Datasmith/DatasmithCore/Private/DatasmithSceneXmlReader.cpp
+++ b/Engine/Source/Runtime/Datasmith/DatasmithCore/Private/DatasmithSceneXmlReader.cpp
@@ -33,10 +33,7 @@ namespace DatasmithSceneXmlReaderImpl
 		DATASMITH_CUSTOMACTORNAME, DATASMITH_LANDSCAPENAME, DATASMITH_POSTPROCESSVOLUME, DATASMITH_ACTORHIERARCHICALINSTANCEDMESHNAME, DATASMITH_DECALACTORNAME };
 
 	template< typename T >
-	T ValueFromString( const FString& InString )
-	{
-		// Invalid
-	}
+	T ValueFromString( const FString& InString ) = delete;
 
 	template<>
 	float ValueFromString< float >( const FString& InString )
@@ -760,6 +757,9 @@ void FDatasmithSceneXmlReader::ParseActor(FXmlNode* InNode, TSharedPtr<IDatasmit
 	FString VisibleAtribute = InNode->GetAttribute(TEXT("visible"));
 	InOutElement->SetVisibility(VisibleAtribute.IsEmpty() ? true : ValueFromString<bool>(VisibleAtribute));
 
+	FString CastShadowAtribute = InNode->GetAttribute(TEXT("castshadow"));
+	InOutElement->SetCastShadow(CastShadowAtribute.IsEmpty() ? true : ValueFromString<bool>(CastShadowAtribute));
+
 	for (FXmlNode* ChildNode : InNode->GetChildrenNodes())
 	{
 		if (ChildNode->GetTag() == TEXT("transform"))
diff --git a/Engine/Source/Runtime/Datasmith/DatasmithCore/Private/DatasmithSceneXmlWriter.cpp b/Engine/Source/Runtime/Datasmith/DatasmithCore/Private/DatasmithSceneXmlWriter.cpp
--- a/Engine/Source/Runtime/Datasmith/DatasmithCore/Private/DatasmithSceneXmlWriter.cpp
+++ b/Engine/Source/Runtime/Datasmith/DatasmithCore/Private/DatasmithSceneXmlWriter.cpp
@@ -122,7 +122,12 @@ FString FDatasmithSceneXmlWriterImpl::GetLabelAndLayer(const TSharedPtr<IDatasmi
 
 	if (!ActorElement->GetVisibility())
 	{
-		LabelAndLayer += TEXT(" visible=\"0\"");
+		LabelAndLayer += TEXT(" visible=\"false\"");
+	}
+
+	if (!ActorElement->GetCastShadow())
+	{
+		LabelAndLayer += TEXT(" castshadow=\"false\"");
 	}
 
 	if (ActorElement->IsAComponent())
diff --git a/Engine/Source/Runtime/Datasmith/DatasmithCore/Public/IDatasmithSceneElements.h b/Engine/Source/Runtime/Datasmith/DatasmithCore/Public/IDatasmithSceneElements.h
--- a/Engine/Source/Runtime/Datasmith/DatasmithCore/Public/IDatasmithSceneElements.h
+++ b/Engine/Source/Runtime/Datasmith/DatasmithCore/Public/IDatasmithSceneElements.h
@@ -140,11 +140,17 @@ public:
 	virtual void SetIsAComponent(bool Value) = 0;
 	virtual bool IsAComponent() const = 0;
 
-	/** Get a mesh actor's visibility */
+	/** Set a mesh actor's visibility */
 	virtual void SetVisibility(bool bInVisibility) = 0;
 
-	/** Set a mesh actor's visibility */
+	/** Get a mesh actor's visibility */
 	virtual bool GetVisibility() const = 0;
+
+	/** Set whether an actor's casts shadow */
+	virtual void SetCastShadow(bool bInCastShadow) = 0;
+
+	/** Get whether an actor's casts shadow */
+	virtual bool GetCastShadow() const = 0;
 };
 
 /**
