[3ds Max] DirectLink: Forest multisubobj materials

- test: add two geometries with different multisubobj materials to the same Forest object
- also fixed crash related to Containers. Test: create container, close it, open, close, previously this whould eventually crash

diff --git a/Engine/Source/Programs/Enterprise/Datasmith/DatasmithMaxExporter/Private/DatasmithMaxDirectLinkItoo.cpp b/Engine/Source/Programs/Enterprise/Datasmith/DatasmithMaxExporter/Private/DatasmithMaxDirectLinkItoo.cpp
--- a/Engine/Source/Programs/Enterprise/Datasmith/DatasmithMaxExporter/Private/DatasmithMaxDirectLinkItoo.cpp
+++ b/Engine/Source/Programs/Enterprise/Datasmith/DatasmithMaxExporter/Private/DatasmithMaxDirectLinkItoo.cpp
@@ -148,9 +148,10 @@ bool ConvertForest(ISceneTracker& Scene, FNodeTracker& NodeTracker, Object* Obj)
 		struct FHismInstances
 		{
 			INode* GeometryNode;
-			TArray<Matrix3> Transforms;
+			TMap<Mtl*, TArray<Matrix3>> InstancesTransformsForMaterial;
 		};
 
+
 		TArray<FHismInstances> InstancesForMesh;
 		InstancesForMesh.Reserve(NumInstances); // Reserve to avoid reallocations
 		TMap<int, int32> RenderableNodeIndicesMap;
@@ -163,7 +164,7 @@ bool ConvertForest(ISceneTracker& Scene, FNodeTracker& NodeTracker, Object* Obj)
 
 				if (int32* RenderableNodeIndexPtr = RenderableNodeIndicesMap.Find(VirtualMaster))
 				{
-					InstancesForMesh[*RenderableNodeIndexPtr].Transforms.Emplace(ForestInstance->tm);
+					InstancesForMesh[*RenderableNodeIndexPtr].InstancesTransformsForMaterial.FindOrAdd(ForestInstance->mtl).Emplace(ForestInstance->tm);
 				}
 				else
 				{
@@ -171,7 +172,7 @@ bool ConvertForest(ISceneTracker& Scene, FNodeTracker& NodeTracker, Object* Obj)
 					FHismInstances& RenderableNode = InstancesForMesh.Emplace_GetRef();
 
 					RenderableNode.GeometryNode = ForestInstance->node;
-					RenderableNode.Transforms.Emplace(ForestInstance->tm);
+					RenderableNode.InstancesTransformsForMaterial.FindOrAdd(ForestInstance->mtl).Emplace(ForestInstance->tm);
 				}
 			}
 		}
@@ -196,7 +197,12 @@ bool ConvertForest(ISceneTracker& Scene, FNodeTracker& NodeTracker, Object* Obj)
 
 			if (RenderMesh.IsValid())
 			{
-				Scene.SetupDatasmithHISMForNode(NodeTracker, GeometryNode, RenderMesh, NodeTracker.Node->GetMtl(), MeshIndex, Instances.Transforms);
+				for(const TPair<Mtl*, TArray<Matrix3>>& MaterialAndTransforms: Instances.InstancesTransformsForMaterial)
+				{
+					Scene.SetupDatasmithHISMForNode(NodeTracker, GeometryNode, RenderMesh, MaterialAndTransforms.Key, MeshIndex, MaterialAndTransforms.Value);
+				}
+
+
 				MeshIndex ++;
 			}
 		}
diff --git a/Engine/Source/Programs/Enterprise/Datasmith/DatasmithMaxExporter/Private/DatasmithMaxExporterWithDirectLink.cpp b/Engine/Source/Programs/Enterprise/Datasmith/DatasmithMaxExporter/Private/DatasmithMaxExporterWithDirectLink.cpp
--- a/Engine/Source/Programs/Enterprise/Datasmith/DatasmithMaxExporter/Private/DatasmithMaxExporterWithDirectLink.cpp
+++ b/Engine/Source/Programs/Enterprise/Datasmith/DatasmithMaxExporter/Private/DatasmithMaxExporterWithDirectLink.cpp
@@ -910,6 +910,15 @@ public:
 	{
 		ensure(!bUpdateInProgress);
 
+		if (NodeTracker.bDeleted)
+		{
+			// Change events sometimes received for nodes that are already deleted 
+			// skipping processing of node subtree because INode pointer may already be invalid
+			// Test case: create container, add node to it. Close it, open it, close again, then sync
+			// Change event will be received for NodeKey that returns NULL from NodeEventNamespace::GetNodeByKey
+			return;
+		}
+
 		NodeTracker.Invalidate();
 		InvalidatedNodeTrackers.Add(&NodeTracker);
 
@@ -944,7 +953,7 @@ public:
 			}
 			else
 			{
-				// Sometimes note update received without node Delete event
+				// Sometimes node update received without node Delete event
 				// Test case: create container, add node to it. Close it, open it, close again, then sync
 				InvalidatedNodeTrackers.Add(NodeTracker);
 				NodeTracker->bDeleted = true;
@@ -1032,6 +1041,8 @@ public:
 			TUniquePtr<FRailClonesConverted> RailClonesConverted;
 			if (RailClones.RemoveAndCopyValue(&NodeTracker, RailClonesConverted))
 			{
+				UnregisterNodeForMaterial(NodeTracker);
+
 				for (TSharedPtr<IDatasmithMeshElement> Mesh : RailClonesConverted->Meshes)
 				{
 					ReleaseMeshElement(Mesh);
@@ -1385,27 +1396,22 @@ public:
 
 	void RegisterNodeForMaterial(FNodeTracker& NodeTracker, Mtl* Material)
 	{
-		if (!NodeTracker.MaterialTracker || NodeTracker.MaterialTracker->Material != Material)
-		{
-			// Release old material
-			UnregisterNodeForMaterial(NodeTracker);
-
-			NodeTracker.MaterialTracker = MaterialsCollectionTracker.AddMaterial(Material);
-			MaterialsAssignedToNodes.FindOrAdd(NodeTracker.MaterialTracker).Add(&NodeTracker);
-		}
+		FMaterialTracker* MaterialTracker = MaterialsCollectionTracker.AddMaterial(Material);
+		NodeTracker.MaterialTrackers.Add(MaterialTracker);
+		MaterialsAssignedToNodes.FindOrAdd(MaterialTracker).Add(&NodeTracker);
 	}
 
 	void UnregisterNodeForMaterial(FNodeTracker& NodeTracker)
 	{
-		if (NodeTracker.MaterialTracker)
+		for (FMaterialTracker* MaterialTracker: NodeTracker.MaterialTrackers)
 		{
-			MaterialsAssignedToNodes[NodeTracker.MaterialTracker].Remove(&NodeTracker);
-			if (MaterialsAssignedToNodes[NodeTracker.MaterialTracker].IsEmpty())
+			MaterialsAssignedToNodes[MaterialTracker].Remove(&NodeTracker);
+			if (MaterialsAssignedToNodes[MaterialTracker].IsEmpty())
 			{
-				MaterialsCollectionTracker.ReleaseMaterial(*NodeTracker.MaterialTracker);
+				MaterialsCollectionTracker.ReleaseMaterial(*MaterialTracker);
 			}
-			NodeTracker.MaterialTracker = nullptr;
 		}
+		NodeTracker.MaterialTrackers.Reset();
 	}
 
 	void UpdateGeometryNode(FNodeTracker& NodeTracker, FInstances& Instances, bool bMaterialsAssignToStaticMesh)
@@ -1489,7 +1495,18 @@ public:
 		{
 			if (Mtl* Material = NodeTracker.Node->GetMtl())
 			{
-				RegisterNodeForMaterial(NodeTracker, Material);
+				bool bMaterialRegistered = false;
+				for(FMaterialTracker* MaterialTracker: NodeTracker.MaterialTrackers)
+				{
+					bMaterialRegistered = bMaterialRegistered || (MaterialTracker->Material == Material);
+				}
+
+				if (!bMaterialRegistered)
+				{
+					// Release old material
+					UnregisterNodeForMaterial(NodeTracker);
+					RegisterNodeForMaterial(NodeTracker, Material);
+				}
 
 				// Assign materials
 				if (bMaterialsAssignToStaticMesh)
diff --git a/Engine/Source/Programs/Enterprise/Datasmith/DatasmithMaxExporter/Public/DatasmithMaxDirectLink.h b/Engine/Source/Programs/Enterprise/Datasmith/DatasmithMaxExporter/Public/DatasmithMaxDirectLink.h
--- a/Engine/Source/Programs/Enterprise/Datasmith/DatasmithMaxExporter/Public/DatasmithMaxDirectLink.h
+++ b/Engine/Source/Programs/Enterprise/Datasmith/DatasmithMaxExporter/Public/DatasmithMaxDirectLink.h
@@ -26,6 +26,7 @@ MAX_INCLUDES_END
 // #define LOG_DEBUG_HEAVY_ENABLE 1
 // #define LOG_DEBUG_ENABLE 1
 
+
 void AssignMeshMaterials(TSharedPtr<IDatasmithMeshElement>&MeshElement, Mtl * Material, const TSet<uint16>&SupportedChannels);
 
 namespace DatasmithMaxDirectLink
@@ -146,7 +147,7 @@ public:
 
 	TSharedPtr<IDatasmithActorElement> DatasmithActorElement;
 
-	class FMaterialTracker* MaterialTracker = nullptr;
+	TSet<class FMaterialTracker*> MaterialTrackers;
 
 	TSharedPtr<IDatasmithMeshActorElement> DatasmithMeshActor;
 
