[3ds Max] DirectLink: CoronaPhysicalMtl ClearCoat

diff --git a/Engine/Source/Programs/Enterprise/Datasmith/DatasmithMaxExporter/Private/MaxMaterialsToUEPbr/DatasmithMaxCoronaMaterialsToUEPbr.cpp b/Engine/Source/Programs/Enterprise/Datasmith/DatasmithMaxExporter/Private/MaxMaterialsToUEPbr/DatasmithMaxCoronaMaterialsToUEPbr.cpp
--- a/Engine/Source/Programs/Enterprise/Datasmith/DatasmithMaxExporter/Private/MaxMaterialsToUEPbr/DatasmithMaxCoronaMaterialsToUEPbr.cpp
+++ b/Engine/Source/Programs/Enterprise/Datasmith/DatasmithMaxExporter/Private/MaxMaterialsToUEPbr/DatasmithMaxCoronaMaterialsToUEPbr.cpp
@@ -1429,6 +1429,28 @@ void FDatasmithMaxCoronaPhysicalMaterialToUEPbr::Convert(TSharedRef<IDatasmithSc
 	}
 	Connect(PbrMaterialElement->GetBaseColor(), BaseColorExpression);
 
+
+	// Clear Coat
+	if (MaterialProperties.ClearcoatAmount > 0 || MaterialProperties.ClearcoatAmountTexmap.IsMapPresentAndEnabled())
+	{
+		IDatasmithMaterialExpression* ClearCoatExpression = TextureOrScalar(TEXT("Clear Coat Amount"), MaterialProperties.ClearcoatAmountTexmap, MaterialProperties.ClearcoatAmount);
+
+		if (ClearCoatExpression)
+		{
+			if (MaterialProperties.ClearcoatIorTexmap.IsMapPresentAndEnabled())
+			{
+				IDatasmithMaterialExpression* ClearCoatIorExpression = TextureOrScalar(TEXT("Clear Coat Ior"), MaterialProperties.ClearcoatIorTexmap, MaterialProperties.ClearcoatIor);
+				Connect(PbrMaterialElement->GetClearCoat(), CalcIORComplex(*ClearCoatIorExpression, Scalar(0), *ClearCoatExpression, Multiply(*ClearCoatExpression, Scalar(0.1))));
+			}
+			else
+			{
+				Connect(PbrMaterialElement->GetClearCoat(), CalcIORComplex(MaterialProperties.ClearcoatIor, 0, *ClearCoatExpression, Multiply(*ClearCoatExpression, Scalar(0.1))));
+			}
+		}
+		Connect(PbrMaterialElement->GetClearCoatRoughness(), TextureOrScalar(TEXT("Clear Coat Roughness"), MaterialProperties.ClearcoatRoughnessTexmap, MaterialProperties.ClearcoatRoughness));
+		PbrMaterialElement->SetShadingModel(EDatasmithShadingModel::ClearCoat);
+	}
+
 	if (MaterialProperties.BaseBumpTexmap.IsMapPresentAndEnabled())
 	{
 		ConvertState.DefaultTextureMode = EDatasmithTextureMode::Bump; // Will change to normal if we pass through a normal map texmap
diff --git a/Engine/Source/Programs/Enterprise/Datasmith/DatasmithMaxExporter/Private/MaxMaterialsToUEPbr/DatasmithMaxMaterialsToUEPbr.h b/Engine/Source/Programs/Enterprise/Datasmith/DatasmithMaxExporter/Private/MaxMaterialsToUEPbr/DatasmithMaxMaterialsToUEPbr.h
--- a/Engine/Source/Programs/Enterprise/Datasmith/DatasmithMaxExporter/Private/MaxMaterialsToUEPbr/DatasmithMaxMaterialsToUEPbr.h
+++ b/Engine/Source/Programs/Enterprise/Datasmith/DatasmithMaxExporter/Private/MaxMaterialsToUEPbr/DatasmithMaxMaterialsToUEPbr.h
@@ -138,7 +138,7 @@ public:
 	IDatasmithMaterialExpression* ApplyWeightExpression(IDatasmithMaterialExpression* ValueExpression, IDatasmithMaterialExpression* WeightExpression);
 
 	IDatasmithMaterialExpression& CalcIORComplex(double IORn, double IORk, IDatasmithMaterialExpression& ToBeConnected90, IDatasmithMaterialExpression& ToBeConnected0);
-
+	IDatasmithMaterialExpression& CalcIORComplex(IDatasmithMaterialExpression& IORn, IDatasmithMaterialExpression& IORk, IDatasmithMaterialExpression& ToBeConnected90, IDatasmithMaterialExpression& ToBeConnected0);
 
 	void Connect(IDatasmithExpressionInput& Input, IDatasmithMaterialExpression& ValueExpression);
 	bool Connect(IDatasmithExpressionInput& Input, IDatasmithMaterialExpression* ValueExpression); // Connect if not null
diff --git a/Engine/Source/Programs/Enterprise/Datasmith/DatasmithMaxExporter/Private/MaxMaterialsToUEPbr/DatasmithMaxPhysicalMaterialToUEPbr.cpp b/Engine/Source/Programs/Enterprise/Datasmith/DatasmithMaxExporter/Private/MaxMaterialsToUEPbr/DatasmithMaxPhysicalMaterialToUEPbr.cpp
--- a/Engine/Source/Programs/Enterprise/Datasmith/DatasmithMaxExporter/Private/MaxMaterialsToUEPbr/DatasmithMaxPhysicalMaterialToUEPbr.cpp
+++ b/Engine/Source/Programs/Enterprise/Datasmith/DatasmithMaxExporter/Private/MaxMaterialsToUEPbr/DatasmithMaxPhysicalMaterialToUEPbr.cpp
@@ -201,28 +201,18 @@ IDatasmithMaterialExpression* FDatasmithMaxMaterialsToUEPbrExpressions::ApplyWei
 	return COMPOSE_OR_DEFAULT2(ValueExpression, Multiply, ValueExpression, WeightExpression);
 }
 
-IDatasmithMaterialExpression& FDatasmithMaxMaterialsToUEPbrExpressions::CalcIORComplex(double IORn, double IORk,
-	IDatasmithMaterialExpression& ToBeConnected90, IDatasmithMaterialExpression& ToBeConnected0)
+IDatasmithMaterialExpression& FDatasmithMaxMaterialsToUEPbrExpressions::CalcIORComplex(IDatasmithMaterialExpression& IORn, IDatasmithMaterialExpression& IORk, IDatasmithMaterialExpression& ToBeConnected90, IDatasmithMaterialExpression& ToBeConnected0)
 {
-	if (IORn >= 16)
-	{
-		return ToBeConnected90;
-	}
-
-	IORn = FMath::Max(IORn, 1.0f);
-
 	IDatasmithMaterialExpression& FresnelOrig = Fresnel(&Scalar(1.0), &Scalar(0.0));
 	IDatasmithMaterialExpression& Fresnel = OneMinus(FresnelOrig);
 	IDatasmithMaterialExpression& Fresnel2 = Multiply(Fresnel, Fresnel);
 
-	IDatasmithMaterialExpression& ConstantIORn = Scalar(IORn);
-	// ConstantIORn->Desc = TEXT("IOR Value");
-	IDatasmithMaterialExpression& ConstantIORn2 = Multiply(ConstantIORn, ConstantIORn);
+	// IORn->Desc = TEXT("IOR Value");
+	IDatasmithMaterialExpression& ConstantIORn2 = Multiply(IORn, IORn);
 
-	IDatasmithMaterialExpression& ConstantIORk = Scalar(IORk);
-	IDatasmithMaterialExpression& AddN2K2 = (IORk == 0) ? ConstantIORn2 : Add(ConstantIORn2, Multiply(ConstantIORk, ConstantIORk));
+	IDatasmithMaterialExpression& AddN2K2 = Add(ConstantIORn2, Multiply(IORk, IORk));
 
-	IDatasmithMaterialExpression& Mul2NCPre = Multiply(ConstantIORn, Scalar(2.0));
+	IDatasmithMaterialExpression& Mul2NCPre = Multiply(IORn, Scalar(2.0));
 	IDatasmithMaterialExpression& Mul2NC = Multiply(Mul2NCPre, Fresnel);
 
 	IDatasmithMaterialExpression& RsNumPre = Subtract(AddN2K2, Mul2NC);
@@ -244,6 +234,21 @@ IDatasmithMaterialExpression& FDatasmithMaxMaterialsToUEPbrExpressions::CalcIORC
 	return Lerp(ToBeConnected0, ToBeConnected90, Power(Res, Scalar(0.5)));
 }
 
+IDatasmithMaterialExpression& FDatasmithMaxMaterialsToUEPbrExpressions::CalcIORComplex(double IORn, double IORk,
+                                                                                       IDatasmithMaterialExpression& ToBeConnected90, IDatasmithMaterialExpression& ToBeConnected0)
+{
+	if (IORn >= 16)
+	{
+		return ToBeConnected90;
+	}
+
+	IORn = FMath::Max(IORn, 1.0f);
+	IDatasmithMaterialExpression& ConstantIORn = Scalar(IORn);
+	IDatasmithMaterialExpression& ConstantIORk = Scalar(IORk);
+
+	return CalcIORComplex(ConstantIORn, ConstantIORk, ToBeConnected90, ToBeConnected0);
+}
+
 void FDatasmithMaxMaterialsToUEPbrExpressions::Connect(IDatasmithExpressionInput& Input, IDatasmithMaterialExpression& ValueExpression)
 {
 	ValueExpression.ConnectExpression(Input);
