Datasmith Auto sync on solidworks keeps sending update to Unreal Editor

diff --git a/Engine/Source/Programs/Enterprise/Datasmith/DatasmithSolidworksExporter/Private/AssemblyDocument.cs b/Engine/Source/Programs/Enterprise/Datasmith/DatasmithSolidworksExporter/Private/AssemblyDocument.cs
--- a/Engine/Source/Programs/Enterprise/Datasmith/DatasmithSolidworksExporter/Private/AssemblyDocument.cs
+++ b/Engine/Source/Programs/Enterprise/Datasmith/DatasmithSolidworksExporter/Private/AssemblyDocument.cs
@@ -164,6 +164,8 @@ namespace DatasmithSolidworks
 			}
 			else
 			{
+				bool bHasDirtyMaterials = false;
+
 				if (CurrentDocMaterialsMap.Count != SyncState.ComponentsMaterialsMap.Count)
 				{
 					IEnumerable<string> Diff1 = CurrentDocMaterialsMap.Keys.Except(SyncState.ComponentsMaterialsMap.Keys);
@@ -176,10 +178,26 @@ namespace DatasmithSolidworks
 					// Components in the DiffSet have their materials changed
 					foreach (string CompName in DiffSet)
 					{
-						SetComponentDirty(CompName, EComponentDirtyState.Material);
+						bool bShouldSyncComponentMaterial = false;
+
+						if (SyncState.ExportedComponentsMap.ContainsKey(CompName))
+						{
+							try
+							{
+								Component2 Comp = SyncState.ExportedComponentsMap[CompName];
+								bShouldSyncComponentMaterial = !Comp.IsSuppressed() && (Comp.Visible == (int)swComponentVisibilityState_e.swComponentVisible);
+							}
+							catch { }
+						}
+
+						if (bShouldSyncComponentMaterial)
+						{
+							bHasDirtyMaterials = true;
+							SetComponentDirty(CompName, EComponentDirtyState.Material);
+						}
 					}
 
-					return true;
+					return bHasDirtyMaterials;
 				}
 
 				bool bHasDirtyComponents = false;
