Ability to specify a custom origin point

diff --git a/Engine/Source/Programs/Enterprise/Datasmith/DatasmithRevitExporter/Private/DatasmithRevitDocumentData.cs b/Engine/Source/Programs/Enterprise/Datasmith/DatasmithRevitExporter/Private/DatasmithRevitDocumentData.cs
--- a/Engine/Source/Programs/Enterprise/Datasmith/DatasmithRevitExporter/Private/DatasmithRevitDocumentData.cs
+++ b/Engine/Source/Programs/Enterprise/Datasmith/DatasmithRevitExporter/Private/DatasmithRevitDocumentData.cs
@@ -581,7 +581,7 @@ namespace DatasmithRevitExporter
 				LightActor.SetLabel(GetActorLabel());
 
 				// Set the world transform of the Datasmith light actor.
-				FDocumentData.SetActorTransform(InWorldTransform, LightActor);
+				DocumentData.SetActorTransform(InWorldTransform, LightActor);
 
 				// Set the base properties of the Datasmith light actor.
 				string LayerName = Category.GetCategory(CurrentElement.Document, BuiltInCategory.OST_LightingFixtureSource)?.Name ?? "Light Sources";
@@ -730,7 +730,7 @@ namespace DatasmithRevitExporter
 				FacadeActor.SetLabel(GetActorLabel());
 
 				// Set the world transform of the Datasmith RPC mesh actor.
-				FDocumentData.SetActorTransform(InWorldTransform, FacadeActor);
+				DocumentData.SetActorTransform(InWorldTransform, FacadeActor);
 
 				// Set the base properties of the Datasmith RPC mesh actor.
 				string LayerName = GetCategoryName();
@@ -812,7 +812,7 @@ namespace DatasmithRevitExporter
 				}
 
 				// Set the world transform of the Datasmith mesh actor.
-				FDocumentData.SetActorTransform(InWorldTransform, InElement.ElementActor);
+				DocumentData.SetActorTransform(InWorldTransform, InElement.ElementActor);
 
 				// Set the base properties of the Datasmith mesh actor.
 				string LayerName = GetCategoryName();
@@ -1062,7 +1062,13 @@ namespace DatasmithRevitExporter
 		private string									CurrentMaterialName = null;
 		private List<string>							MessageList = null;
 
-		public  string									DocumentId { get; private set; } = "";
+		// Apply world offset to elements
+		public	FSettings.EInsertionPoint				InsertionPoint { get; set; } = FSettings.EInsertionPoint.Default;
+
+		private XYZ										ProjectSurveyPoint = null;
+		private XYZ										ProjectBasePoint = null;
+
+		public string									DocumentId { get; private set; } = "";
 
 		public bool										bSkipMetadataExport { get; private set; } = false;
 		public Document									CurrentDocument { get; private set; } = null;
@@ -1077,6 +1083,8 @@ namespace DatasmithRevitExporter
 			string InLinkedDocumentId
 		)
 		{
+			InsertionPoint = FSettingsManager.CurrentSettings?.InsertionPoint ?? FSettings.EInsertionPoint.Default;
+
 			DirectLink = InDirectLink;
 			CurrentDocument = InDocument;
 			MessageList = InMessageList;
@@ -1850,8 +1858,17 @@ namespace DatasmithRevitExporter
 					// Set the world transform of the Datasmith placeholder actor.
 					XYZ BasePointPosition = BasePointBoundingBox.Min;
 
+					if (BasePointLocation.IsShared)
+					{
+						ProjectSurveyPoint = BasePointPosition;
+					}
+					else
+					{
+						ProjectBasePoint = BasePointPosition;
+					}
+
 					Transform TranslationMatrix = Transform.CreateTranslation(BasePointPosition);
-					FDocumentData.SetActorTransform(TranslationMatrix.Multiply(InWorldTransform), BasePointActor);
+					SetActorTransform(TranslationMatrix.Multiply(InWorldTransform), BasePointActor);
 
 					// Set the Datasmith placeholder actor layer to the base point category name.
 					BasePointActor.SetLayer(BasePointLocation.Category.Name);
@@ -2173,7 +2190,7 @@ namespace DatasmithRevitExporter
 			}
 		}
 
-		private static void SetActorTransform(
+		private void SetActorTransform(
 			Transform InWorldTransform,
 			FDatasmithFacadeActor IOActor
 		)
@@ -2183,6 +2200,16 @@ namespace DatasmithRevitExporter
 			XYZ transformBasisZ = InWorldTransform.BasisZ;
 			XYZ transformOrigin = InWorldTransform.Origin;
 
+			// Check if need to apply world offset to element transform
+			if (InsertionPoint != FSettings.EInsertionPoint.Default)
+			{
+				switch (InsertionPoint)
+				{
+					case FSettings.EInsertionPoint.BasePoint: transformOrigin -= ProjectBasePoint; break;
+					case FSettings.EInsertionPoint.SurveyPoint: transformOrigin -= ProjectSurveyPoint; break;
+				}
+			}
+
 			float[] worldMatrix = new float[16];
 
 			worldMatrix[0] = (float)transformBasisX.X;
diff --git a/Engine/Source/Programs/Enterprise/Datasmith/DatasmithRevitExporter/Private/DatasmithRevitSettingsDialog.cs b/Engine/Source/Programs/Enterprise/Datasmith/DatasmithRevitExporter/Private/DatasmithRevitSettingsDialog.cs
--- a/Engine/Source/Programs/Enterprise/Datasmith/DatasmithRevitExporter/Private/DatasmithRevitSettingsDialog.cs
+++ b/Engine/Source/Programs/Enterprise/Datasmith/DatasmithRevitExporter/Private/DatasmithRevitSettingsDialog.cs
@@ -86,6 +86,7 @@ namespace DatasmithRevitExporter
 
 	public class DatasmithRevitSettingsDialog : Form
 	{
+		private ComboBox InsertionPointCombo;
 		private ListBox GroupsList;
 		private NumericUpDown LevelOfTessellation;
 		private FSettings Settings;
@@ -109,11 +110,50 @@ namespace DatasmithRevitExporter
 			ToolTip OptionToolTip = new ToolTip();
 			OptionToolTip.AutoPopDelay = 10000; // milliseconds
 
-			FlowLayoutPanel LevelOfTesselationPanel = new FlowLayoutPanel();
-			LevelOfTesselationPanel.FlowDirection = FlowDirection.LeftToRight;
-			LevelOfTesselationPanel.WrapContents = false;
-			LevelOfTesselationPanel.AutoSize = true;
-			LevelOfTesselationPanel.Margin = new Padding(0, 0, 0, 30);
+			// Table layout for top level controls
+			TableLayoutPanel ControlsLayout = new TableLayoutPanel();
+			ControlsLayout.Name = "ControlsLayout";
+			ControlsLayout.ColumnCount = 2;
+			ControlsLayout.ColumnStyles.Add(new ColumnStyle());
+			ControlsLayout.ColumnStyles.Add(new ColumnStyle());
+			ControlsLayout.ColumnStyles.Add(new ColumnStyle());
+			ControlsLayout.RowCount = 4;
+			ControlsLayout.RowStyles.Add(new RowStyle(SizeType.AutoSize));
+			ControlsLayout.RowStyles.Add(new RowStyle(SizeType.AutoSize));
+			ControlsLayout.AutoSize = true;
+			ControlsLayout.AutoSizeMode = AutoSizeMode.GrowAndShrink;
+			ControlsLayout.Dock = DockStyle.Fill;
+			ControlsLayout.TabIndex = 0;
+			ControlsLayout.Margin = new Padding(0, 0, 0, 30);
+
+			// Insertion point option
+			{
+				Label InsertionPointLabel = new Label();
+				InsertionPointLabel.Name = "LevelOfTessellationLabel";
+				InsertionPointLabel.Text = DatasmithRevitResources.Strings.SettingsDialog_InsertionPoint;
+				InsertionPointLabel.Anchor = AnchorStyles.Left;
+				InsertionPointLabel.AutoSize = true;
+				InsertionPointLabel.Margin = new Padding(0, 0, 3, 0);
+				InsertionPointLabel.TabIndex = 2;
+				InsertionPointLabel.TextAlign = ContentAlignment.MiddleLeft;
+
+				InsertionPointCombo = new ComboBox();
+				InsertionPointCombo.Name = "InsertionPoint";
+				InsertionPointCombo.DropDownStyle = ComboBoxStyle.DropDownList;
+				InsertionPointCombo.AutoSize = true;
+				InsertionPointCombo.TabIndex = 3;
+
+				InsertionPointCombo.Items.Add("Default");
+				InsertionPointCombo.Items.Add("Base Point");
+				InsertionPointCombo.Items.Add("Survey Point");
+
+				OptionToolTip.SetToolTip(InsertionPointLabel, FormatTooltip(DatasmithRevitResources.Strings.SettingsDialog_InsertionPointTooltip));
+
+				ControlsLayout.Controls.Add(InsertionPointLabel, 0, 0);
+				ControlsLayout.Controls.Add(InsertionPointCombo, 1, 0);
+			}
+
+			// Tesselation level option
 			{
 				Label LevelOfTessellationLabel = new Label();
 				LevelOfTessellationLabel.Name = "LevelOfTessellationLabel";
@@ -136,8 +176,8 @@ namespace DatasmithRevitExporter
 
 				OptionToolTip.SetToolTip(LevelOfTessellationLabel, FormatTooltip(DatasmithRevitResources.Strings.SettingsDialog_LevelOfTesselationTooltip));
 
-				LevelOfTesselationPanel.Controls.Add(LevelOfTessellationLabel);
-				LevelOfTesselationPanel.Controls.Add(LevelOfTessellation);
+				ControlsLayout.Controls.Add(LevelOfTessellationLabel, 0, 1);
+				ControlsLayout.Controls.Add(LevelOfTessellation, 1, 1);
 			}
 
 			Label GroupsLabel = new Label();
@@ -247,12 +287,12 @@ namespace DatasmithRevitExporter
 			DialogLayout.Dock = DockStyle.Fill;
 			DialogLayout.TabIndex = 0;
 		
-			DialogLayout.Controls.Add(LevelOfTesselationPanel, 0, 0);
+			DialogLayout.Controls.Add(ControlsLayout, 0, 0);
 			DialogLayout.Controls.Add(GroupsLabel, 0, 1);
 			DialogLayout.Controls.Add(GroupsList, 0, 2);
 			DialogLayout.Controls.Add(ButtonsPanel, 1, 2);
 
-			DialogLayout.SetColumnSpan(LevelOfTesselationPanel, 2);
+			DialogLayout.SetColumnSpan(ControlsLayout, 2);
 			DialogLayout.SetColumnSpan(GroupsLabel, 2);
 
 			Name = "MetadataExportFilter";
@@ -291,6 +331,7 @@ namespace DatasmithRevitExporter
 			}
 			else
 			{
+				InsertionPointCombo.SelectedIndex = (int)Settings.InsertionPoint;
 				LevelOfTessellation.Value = Settings.LevelOfTesselation;
 
 				foreach (int Group in Settings.MetadataParamGroupsFilter)
@@ -316,6 +357,7 @@ namespace DatasmithRevitExporter
 		private void OnClosing(object InSender, FormClosingEventArgs InArgs)
 		{
 			Settings.LevelOfTesselation = Decimal.ToInt32(LevelOfTessellation.Value);
+			Settings.InsertionPoint = (FSettings.EInsertionPoint)InsertionPointCombo.SelectedIndex;
 			Settings.MetadataParamGroupsFilter = AddedBuiltinParamGroups.Values;
 			FSettingsManager.WriteSettings(Document, Settings);
 		}
diff --git a/Engine/Source/Programs/Enterprise/Datasmith/DatasmithRevitExporter/Private/DatasmithRevitSettingsManager.cs b/Engine/Source/Programs/Enterprise/Datasmith/DatasmithRevitExporter/Private/DatasmithRevitSettingsManager.cs
--- a/Engine/Source/Programs/Enterprise/Datasmith/DatasmithRevitExporter/Private/DatasmithRevitSettingsManager.cs
+++ b/Engine/Source/Programs/Enterprise/Datasmith/DatasmithRevitExporter/Private/DatasmithRevitSettingsManager.cs
@@ -11,6 +11,13 @@ namespace DatasmithRevitExporter
 {
 	public class FSettings
 	{
+		public enum EInsertionPoint
+		{
+			Default,
+			BasePoint,
+			SurveyPoint
+		};
+
 		private IList<string> _MetadataParamNamesFilter = new List<string>();
 		private IList<int> _MetadataParamGroupsFilter = new List<int>();
 
@@ -19,6 +26,7 @@ namespace DatasmithRevitExporter
 		private HashSet<int> MetadataParamGroupsSet;
 
 		public int LevelOfTesselation { get; set; } = 8;
+		public EInsertionPoint InsertionPoint { get; set; } = EInsertionPoint.Default;
 
 		public IList<string> MetadataParamNamesFilter 
 		{
@@ -77,11 +85,12 @@ namespace DatasmithRevitExporter
 	{
 		static class FSettingsSchema
 		{
-			readonly static Guid SchemaGuid = new Guid("{EAFD9DF3-1E44-4B16-B9DF-E7DD404729A3}");
+			readonly static Guid SchemaGuid = new Guid("{F910B32B-32C6-4BD1-BFD1-7DBF81FFC41B}");
 			
 			public readonly static string MetadataParamNamesField = "MetadataParamNames";
 			public readonly static string MetadataParamGroupsField = "MetadataParamGroups";
 			public readonly static string LevelOfTesselationField = "LevelOfTesselation";
+			public readonly static string InsertionPointField = "InsertionPoint";
 
 			public static Schema GetSchema()
 			{
@@ -97,6 +106,7 @@ namespace DatasmithRevitExporter
 				SchemaBuilder.AddArrayField(MetadataParamNamesField, typeof(string));
 				SchemaBuilder.AddArrayField(MetadataParamGroupsField, typeof(int));
 				SchemaBuilder.AddSimpleField(LevelOfTesselationField, typeof(int));
+				SchemaBuilder.AddSimpleField(InsertionPointField, typeof(int));
 
 				return SchemaBuilder.Finish();
 			}
@@ -148,12 +158,11 @@ namespace DatasmithRevitExporter
 		{
 			FSettings Settings = new FSettings();
 
-			Entity SettingsEntity = GetSettingsEntity(Doc);
-
-			if (SettingsEntity == null || !SettingsEntity.IsValid())
+			void InitDefaultSettings()
 			{
 				// No settings in document, create defaults
 				Settings.LevelOfTesselation = 8;
+				Settings.InsertionPoint = 0;
 				Settings.MetadataParamGroupsFilter.Add((int)Autodesk.Revit.DB.BuiltInParameterGroup.PG_GEOMETRY);
 				Settings.MetadataParamGroupsFilter.Add((int)Autodesk.Revit.DB.BuiltInParameterGroup.PG_IDENTITY_DATA);
 				Settings.MetadataParamGroupsFilter.Add((int)Autodesk.Revit.DB.BuiltInParameterGroup.PG_MATERIALS);
@@ -162,11 +171,26 @@ namespace DatasmithRevitExporter
 
 				WriteSettings(Doc, Settings);
 			}
+
+			Entity SettingsEntity = GetSettingsEntity(Doc);
+
+			if (SettingsEntity == null || !SettingsEntity.IsValid())
+			{
+				InitDefaultSettings();
+			}
 			else
 			{
-				Settings.MetadataParamNamesFilter = SettingsEntity.Get<IList<string>>(FSettingsSchema.MetadataParamNamesField);
-				Settings.MetadataParamGroupsFilter = SettingsEntity.Get<IList<int>>(FSettingsSchema.MetadataParamGroupsField);
-				Settings.LevelOfTesselation = SettingsEntity.Get<int>(FSettingsSchema.LevelOfTesselationField);
+				try
+				{
+					Settings.MetadataParamNamesFilter = SettingsEntity.Get<IList<string>>(FSettingsSchema.MetadataParamNamesField);
+					Settings.MetadataParamGroupsFilter = SettingsEntity.Get<IList<int>>(FSettingsSchema.MetadataParamGroupsField);
+					Settings.LevelOfTesselation = SettingsEntity.Get<int>(FSettingsSchema.LevelOfTesselationField);
+					Settings.InsertionPoint = (FSettings.EInsertionPoint)SettingsEntity.Get<int>(FSettingsSchema.InsertionPointField);
+				}
+				catch
+				{
+					InitDefaultSettings();
+				}
 			}
 
 			return Settings;
@@ -190,6 +214,7 @@ namespace DatasmithRevitExporter
 					SettingsEntity.Set(FSettingsSchema.MetadataParamNamesField, Settings.MetadataParamNamesFilter);
 					SettingsEntity.Set(FSettingsSchema.MetadataParamGroupsField, Settings.MetadataParamGroupsFilter);
 					SettingsEntity.Set(FSettingsSchema.LevelOfTesselationField, Settings.LevelOfTesselation);
+					SettingsEntity.Set(FSettingsSchema.InsertionPointField, (int)Settings.InsertionPoint);
 
 					// Identify settings data storage
 					Entity IdEntity = new Entity(DataStorageUniqueIdSchema.GetSchema());
