DatasmithRhinoExporter - Minor refactor: splitting DatasmithRhinoExportContext.cs into multiple files as the classes derived from "DatsmithInfoBase" had become quite large.

diff --git a/Engine/Source/Programs/Enterprise/Datasmith/DatasmithRhinoExporter/DatasmithRhino6/DatasmithRhino6.csproj b/Engine/Source/Programs/Enterprise/Datasmith/DatasmithRhinoExporter/DatasmithRhino6/DatasmithRhino6.csproj
--- a/Engine/Source/Programs/Enterprise/Datasmith/DatasmithRhinoExporter/DatasmithRhino6/DatasmithRhino6.csproj
+++ b/Engine/Source/Programs/Enterprise/Datasmith/DatasmithRhinoExporter/DatasmithRhino6/DatasmithRhino6.csproj
@@ -111,6 +111,9 @@
     </Compile>
     <Compile Include="..\Private\ElementExporters\*.cs">
       <Link>DatasmithRhinoExporter\ElementExporters\%(Filename)%(Extension)</Link>
+    </Compile>
+	<Compile Include="..\Private\ExportContext\*.cs">
+      <Link>DatasmithRhinoExporter\ExportContext\%(Filename)%(Extension)</Link>
     </Compile>
     <Compile Include="..\Private\Utils\*.cs">
       <Link>DatasmithRhinoExporter\Utils\%(Filename)%(Extension)</Link>
diff --git a/Engine/Source/Programs/Enterprise/Datasmith/DatasmithRhinoExporter/DatasmithRhino7/DatasmithRhino7.csproj b/Engine/Source/Programs/Enterprise/Datasmith/DatasmithRhinoExporter/DatasmithRhino7/DatasmithRhino7.csproj
--- a/Engine/Source/Programs/Enterprise/Datasmith/DatasmithRhinoExporter/DatasmithRhino7/DatasmithRhino7.csproj
+++ b/Engine/Source/Programs/Enterprise/Datasmith/DatasmithRhinoExporter/DatasmithRhino7/DatasmithRhino7.csproj
@@ -112,6 +112,9 @@
     <Compile Include="..\Private\ElementExporters\*.cs">
       <Link>DatasmithRhinoExporter\ElementExporters\%(Filename)%(Extension)</Link>
     </Compile>
+    <Compile Include="..\Private\ExportContext\*.cs">
+      <Link>DatasmithRhinoExporter\ExportContext\%(Filename)%(Extension)</Link>
+    </Compile>
     <Compile Include="..\Private\Utils\*.cs">
       <Link>DatasmithRhinoExporter\Utils\%(Filename)%(Extension)</Link>
     </Compile>
diff --git a/Engine/Source/Programs/Enterprise/Datasmith/DatasmithRhinoExporter/Private/DatasmithRhinoExportContext.cs b/Engine/Source/Programs/Enterprise/Datasmith/DatasmithRhinoExporter/Private/DatasmithRhinoExportContext.cs
deleted file mode 100644
diff --git a/Engine/Source/Programs/Enterprise/Datasmith/DatasmithRhinoExporter/Private/DatasmithRhinoSceneExporter.cs b/Engine/Source/Programs/Enterprise/Datasmith/DatasmithRhinoExporter/Private/DatasmithRhinoSceneExporter.cs
--- a/Engine/Source/Programs/Enterprise/Datasmith/DatasmithRhinoExporter/Private/DatasmithRhinoSceneExporter.cs
+++ b/Engine/Source/Programs/Enterprise/Datasmith/DatasmithRhinoExporter/Private/DatasmithRhinoSceneExporter.cs
@@ -1,6 +1,7 @@
 // Copyright Epic Games, Inc. All Rights Reserved.
 
 using DatasmithRhino.ElementExporters;
+using DatasmithRhino.ExportContext;
 using DatasmithRhino.Properties.Localization;
 using DatasmithRhino.Utils;
 
diff --git a/Engine/Source/Programs/Enterprise/Datasmith/DatasmithRhinoExporter/Private/DirectLink/DatasmithRhinoChangeListener.cs b/Engine/Source/Programs/Enterprise/Datasmith/DatasmithRhinoExporter/Private/DirectLink/DatasmithRhinoChangeListener.cs
--- a/Engine/Source/Programs/Enterprise/Datasmith/DatasmithRhinoExporter/Private/DirectLink/DatasmithRhinoChangeListener.cs
+++ b/Engine/Source/Programs/Enterprise/Datasmith/DatasmithRhinoExporter/Private/DirectLink/DatasmithRhinoChangeListener.cs
@@ -1,5 +1,6 @@
 // Copyright Epic Games, Inc. All Rights Reserved.
 
+using DatasmithRhino.ExportContext;
 using DatasmithRhino.Properties.Localization;
 
 using Rhino;
diff --git a/Engine/Source/Programs/Enterprise/Datasmith/DatasmithRhinoExporter/Private/DirectLink/DatasmithRhinoDirectLinkManager.cs b/Engine/Source/Programs/Enterprise/Datasmith/DatasmithRhinoExporter/Private/DirectLink/DatasmithRhinoDirectLinkManager.cs
--- a/Engine/Source/Programs/Enterprise/Datasmith/DatasmithRhinoExporter/Private/DirectLink/DatasmithRhinoDirectLinkManager.cs
+++ b/Engine/Source/Programs/Enterprise/Datasmith/DatasmithRhinoExporter/Private/DirectLink/DatasmithRhinoDirectLinkManager.cs
@@ -1,5 +1,6 @@
 // Copyright Epic Games, Inc. All Rights Reserved.
 
+using DatasmithRhino.ExportContext;
 using Rhino;
 using System;
 using System.IO;
diff --git a/Engine/Source/Programs/Enterprise/Datasmith/DatasmithRhinoExporter/Private/ElementExporters/DatasmithRhinoActorExporter.cs b/Engine/Source/Programs/Enterprise/Datasmith/DatasmithRhinoExporter/Private/ElementExporters/DatasmithRhinoActorExporter.cs
--- a/Engine/Source/Programs/Enterprise/Datasmith/DatasmithRhinoExporter/Private/ElementExporters/DatasmithRhinoActorExporter.cs
+++ b/Engine/Source/Programs/Enterprise/Datasmith/DatasmithRhinoExporter/Private/ElementExporters/DatasmithRhinoActorExporter.cs
@@ -1,5 +1,6 @@
 // Copyright Epic Games, Inc. All Rights Reserved.
 
+using DatasmithRhino.ExportContext;
 using DatasmithRhino.Properties.Localization;
 using DatasmithRhino.Utils;
 
@@ -115,7 +116,7 @@ namespace DatasmithRhino.ElementExporters
 					}
 				}
 			}
-			else if(Node is DatasmithActorCameraInfo ActorCameraInfo)
+			else if(Node is DatasmithCameraActorInfo ActorCameraInfo)
 			{
 				ExportedActor = CreateCameraActor(ActorCameraInfo);
 			}
@@ -239,7 +240,7 @@ namespace DatasmithRhino.ElementExporters
 			}
 		}
 
-		private static FDatasmithFacadeActorCamera CreateCameraActor(DatasmithActorCameraInfo InCameraInfo)
+		private static FDatasmithFacadeActorCamera CreateCameraActor(DatasmithCameraActorInfo InCameraInfo)
 		{
 			string HashedName = FDatasmithFacadeElement.GetStringHash(InCameraInfo.Name);
 			return new FDatasmithFacadeActorCamera(HashedName);
diff --git a/Engine/Source/Programs/Enterprise/Datasmith/DatasmithRhinoExporter/Private/ElementExporters/DatasmithRhinoMaterialExporter.cs b/Engine/Source/Programs/Enterprise/Datasmith/DatasmithRhinoExporter/Private/ElementExporters/DatasmithRhinoMaterialExporter.cs
--- a/Engine/Source/Programs/Enterprise/Datasmith/DatasmithRhinoExporter/Private/ElementExporters/DatasmithRhinoMaterialExporter.cs
+++ b/Engine/Source/Programs/Enterprise/Datasmith/DatasmithRhinoExporter/Private/ElementExporters/DatasmithRhinoMaterialExporter.cs
@@ -1,5 +1,6 @@
 // Copyright Epic Games, Inc. All Rights Reserved.
 
+using DatasmithRhino.ExportContext;
 using DatasmithRhino.Utils;
 using Rhino.DocObjects;
 using Rhino.Geometry;
diff --git a/Engine/Source/Programs/Enterprise/Datasmith/DatasmithRhinoExporter/Private/ElementExporters/DatasmithRhinoMeshExporter.cs b/Engine/Source/Programs/Enterprise/Datasmith/DatasmithRhinoExporter/Private/ElementExporters/DatasmithRhinoMeshExporter.cs
--- a/Engine/Source/Programs/Enterprise/Datasmith/DatasmithRhinoExporter/Private/ElementExporters/DatasmithRhinoMeshExporter.cs
+++ b/Engine/Source/Programs/Enterprise/Datasmith/DatasmithRhinoExporter/Private/ElementExporters/DatasmithRhinoMeshExporter.cs
@@ -1,5 +1,6 @@
 // Copyright Epic Games, Inc. All Rights Reserved.
 
+using DatasmithRhino.ExportContext;
 using Rhino.Geometry;
 using System.Collections.Generic;
 
diff --git a/Engine/Source/Programs/Enterprise/Datasmith/DatasmithRhinoExporter/Private/ElementExporters/DatasmithRhinoTextureExporter.cs b/Engine/Source/Programs/Enterprise/Datasmith/DatasmithRhinoExporter/Private/ElementExporters/DatasmithRhinoTextureExporter.cs
--- a/Engine/Source/Programs/Enterprise/Datasmith/DatasmithRhinoExporter/Private/ElementExporters/DatasmithRhinoTextureExporter.cs
+++ b/Engine/Source/Programs/Enterprise/Datasmith/DatasmithRhinoExporter/Private/ElementExporters/DatasmithRhinoTextureExporter.cs
@@ -1,6 +1,7 @@
 // Copyright Epic Games, Inc. All Rights Reserved.
 
 
+using DatasmithRhino.ExportContext;
 using Rhino.DocObjects;
 using System.Collections.Generic;
 
diff --git a/Engine/Source/Programs/Enterprise/Datasmith/DatasmithRhinoExporter/Private/ElementExporters/IDatasmithRhinoElementExporter.cs b/Engine/Source/Programs/Enterprise/Datasmith/DatasmithRhinoExporter/Private/ElementExporters/IDatasmithRhinoElementExporter.cs
--- a/Engine/Source/Programs/Enterprise/Datasmith/DatasmithRhinoExporter/Private/ElementExporters/IDatasmithRhinoElementExporter.cs
+++ b/Engine/Source/Programs/Enterprise/Datasmith/DatasmithRhinoExporter/Private/ElementExporters/IDatasmithRhinoElementExporter.cs
@@ -1,5 +1,6 @@
 // Copyright Epic Games, Inc. All Rights Reserved.
 
+using DatasmithRhino.ExportContext;
 using DatasmithRhino.Utils;
 using System;
 using System.Collections.Generic;
diff --git a/Engine/Source/Programs/Enterprise/Datasmith/DatasmithRhinoExporter/Private/ExportContext/DatasmithActorInfo.cs b/Engine/Source/Programs/Enterprise/Datasmith/DatasmithRhinoExporter/Private/ExportContext/DatasmithActorInfo.cs
new file mode 100644
--- /dev/null
+++ b/Engine/Source/Programs/Enterprise/Datasmith/DatasmithRhinoExporter/Private/ExportContext/DatasmithActorInfo.cs
@@ -0,0 +1,346 @@
+// Copyright Epic Games, Inc. All Rights Reserved.
+
+using DatasmithRhino.Utils;
+using Rhino.DocObjects;
+using Rhino.Geometry;
+using System;
+using System.Collections.Generic;
+using System.Linq;
+
+namespace DatasmithRhino.ExportContext
+{
+	public class DatasmithActorInfo : DatasmithInfoBase
+	{
+		public override DirectLinkSynchronizationStatus DirectLinkStatus
+		{
+			get
+			{
+				DirectLinkSynchronizationStatus Status = base.DirectLinkStatus;
+
+				if (DefinitionNode != null)
+				{
+					DirectLinkSynchronizationStatus DefinitionStatus = DefinitionNode.DirectLinkStatus;
+
+					// Use the Definition status when the instanced actor is synced.
+					// Deleting a definition also takes priority on the instance status.
+					if (Status == DirectLinkSynchronizationStatus.Synced || DefinitionStatus == DirectLinkSynchronizationStatus.PendingDeletion)
+					{
+						return DefinitionStatus;
+					}
+				}
+				return Status;
+			}
+		}
+
+		public FDatasmithFacadeActor DatasmithActor { get { return ExportedElement as FDatasmithFacadeActor; } }
+
+		public virtual Guid RhinoObjectId
+		{
+			get
+			{
+				if (RhinoCommonObject is ModelComponent RhinoModelComponent)
+				{
+					return RhinoModelComponent.Id;
+				}
+				else if (RhinoCommonObject is ViewportInfo RhinoViewportInfo)
+				{
+					return RhinoViewportInfo.Id;
+				}
+				return Guid.Empty;
+			}
+		}
+
+		public Transform WorldTransform { get; private set; } = Transform.Identity;
+		public int MaterialIndex { get; set; } = -1;
+		public bool bOverrideMaterial { get; private set; } = false;
+
+		public bool bIsRoot { get; private set; } = true;
+		public DatasmithActorInfo Parent { get; private set; } = null;
+		private LinkedList<DatasmithActorInfo> ChildrenInternal = new LinkedList<DatasmithActorInfo>();
+		public ICollection<DatasmithActorInfo> Children { get => ChildrenInternal; }
+
+		public bool bIsInstanceDefinition { get; private set; } = false;
+		public DatasmithActorInfo DefinitionNode { get; private set; } = null;
+		private LinkedList<DatasmithActorInfo> InstanceNodesInternal = new LinkedList<DatasmithActorInfo>();
+		public ICollection<DatasmithActorInfo> InstanceNodes { get => InstanceNodesInternal; }
+
+		public HashSet<int> LayerIndices { get; private set; } = new HashSet<int>();
+		/// <summary>
+		/// Used to determine variation occurring in the Layer hierarchy for that node.
+		/// </summary>
+		private List<int> RelativeLayerIndices = new List<int>();
+		public Layer VisibilityLayer { get; private set; } = null;
+
+		public bool bIsVisible
+		{
+			get
+			{
+				bool bVisibleLocally = true;
+				if (RhinoCommonObject is Layer RhinoLayer)
+				{
+					// This recursion ensure that only layer actors with an actual exported object under them are considered visible.
+					// ie. Layers containing no mesh are not visible.
+					return RhinoLayer.IsVisible
+						&& ChildrenInternal.Any(Child => Child.bIsVisible);
+				}
+				else if (RhinoCommonObject is RhinoObject CurrentRhinoObject)
+				{
+					bVisibleLocally = CurrentRhinoObject.Visible;
+				}
+				else if (bIsRoot)
+				{
+					return true;
+				}
+
+				return bVisibleLocally
+					&& (VisibilityLayer == null || VisibilityLayer.IsVisible)
+					&& (DefinitionNode == null || DefinitionNode.bIsVisible);
+			}
+		}
+
+		public DatasmithActorInfo(Transform NodeTransform, string InName, string InUniqueLabel, string InBaseLabel)
+			: base(null, InName, InUniqueLabel, InBaseLabel)
+		{
+			WorldTransform = NodeTransform;
+		}
+
+		public DatasmithActorInfo(ModelComponent InModelComponent, string InName, string InUniqueLabel, string InBaseLabel, int InMaterialIndex, bool bInOverrideMaterial, Layer InVisibilityLayer, IEnumerable<int> InLayerIndices = null)
+			: base(InModelComponent, InName, InUniqueLabel, InBaseLabel)
+		{
+			bIsInstanceDefinition = RhinoCommonObject is InstanceDefinition;
+			MaterialIndex = InMaterialIndex;
+			bOverrideMaterial = bInOverrideMaterial;
+			VisibilityLayer = InVisibilityLayer;
+
+			if (InLayerIndices != null)
+			{
+				LayerIndices.UnionWith(InLayerIndices);
+				RelativeLayerIndices.AddRange(InLayerIndices);
+			}
+
+			WorldTransform = DatasmithRhinoUtilities.GetCommonObjectTransform(InModelComponent);
+		}
+
+		protected DatasmithActorInfo(Rhino.Runtime.CommonObject InModelComponent, string InName, string InUniqueLabel, string InBaseLabel)
+			: base(InModelComponent, InName, InUniqueLabel, InBaseLabel)
+		{ }
+
+		public override void ApplyDiffs(DatasmithInfoBase OtherInfo)
+		{
+			base.ApplyDiffs(OtherInfo);
+
+			if (OtherInfo is DatasmithActorInfo OtherActorInfo)
+			{
+				//Rhino "replaces" object instead of modifying them, we must update the our reference to the object.
+				RhinoCommonObject = OtherActorInfo.RhinoCommonObject;
+
+				MaterialIndex = OtherActorInfo.MaterialIndex;
+				bOverrideMaterial = OtherActorInfo.bOverrideMaterial;
+				RelativeLayerIndices = OtherActorInfo.RelativeLayerIndices;
+				LayerIndices = OtherActorInfo.LayerIndices;
+				VisibilityLayer = OtherActorInfo.VisibilityLayer;
+
+				MaterialIndex = OtherActorInfo.MaterialIndex;
+				bOverrideMaterial = OtherActorInfo.bOverrideMaterial;
+
+				//Update world transform, if it has changed we need to modify transform of children as well.
+				Transform NewWorldTransform = Parent != null
+					? Transform.Multiply(Parent.WorldTransform, OtherActorInfo.WorldTransform)
+					: OtherActorInfo.WorldTransform;
+				if (NewWorldTransform != WorldTransform)
+				{
+					Transform InverseTransform;
+					if (WorldTransform.TryGetInverse(out InverseTransform) || InverseTransform.IsValid)
+					{
+						Transform DiffTransform = NewWorldTransform * InverseTransform;
+						ApplyTransform(DiffTransform);
+					}
+					else
+					{
+						WorldTransform = NewWorldTransform;
+						foreach (DatasmithActorInfo Descendant in GetDescendantEnumerator())
+						{
+							Descendant.ApplyModifiedStatus();
+							Descendant.WorldTransform = Transform.Multiply(Descendant.Parent.WorldTransform, DatasmithRhinoUtilities.GetCommonObjectTransform(Descendant.RhinoCommonObject));
+						}
+					}
+				}
+			}
+			else
+			{
+				System.Diagnostics.Debug.Fail("OtherInfo does not derive from DatasmithActorInfo");
+			}
+		}
+
+		public void UpdateRhinoObject(Rhino.Runtime.CommonObject InCommonObject)
+		{
+			// We are not flagging the info as modified, this will be done by other operations.
+			// This is to allow the case to update a reference to an immuable object.
+			RhinoCommonObject = InCommonObject;
+		}
+
+		public List<string> GetTags(DatasmithRhinoExportContext ExportContext)
+		{
+			List<string> Tags = new List<string>();
+
+			if (RhinoCommonObject != null)
+			{
+				Tags.Add(string.Format("Rhino.ID: {0}", RhinoObjectId));
+				string ComponentTypeString = DatasmithRhinoUniqueNameGenerator.GetDefaultTypeName(RhinoCommonObject);
+				Tags.Add(string.Format("Rhino.Entity.Type: {0}", ComponentTypeString));
+
+				//Add the groups this object belongs to.
+				RhinoObject InRhinoObject = RhinoCommonObject as RhinoObject;
+				if (InRhinoObject != null && InRhinoObject.GroupCount > 0)
+				{
+					int[] GroupIndices = InRhinoObject.GetGroupList();
+					for (int GroupArrayIndex = 0; GroupArrayIndex < GroupIndices.Length; ++GroupArrayIndex)
+					{
+						string GroupName = ExportContext.GroupIndexToName[GroupIndices[GroupArrayIndex]];
+						if (GroupName != null)
+						{
+							Tags.Add(GroupName);
+						}
+					}
+				}
+			}
+
+			return Tags;
+		}
+
+		private void SetParent(DatasmithActorInfo InParent)
+		{
+			if (Parent != null)
+			{
+				// Reset the absolute transform and layer to their relative values.
+				WorldTransform = DatasmithRhinoUtilities.GetCommonObjectTransform(RhinoCommonObject);
+				LayerIndices = new HashSet<int>(RelativeLayerIndices);
+			}
+
+			bIsRoot = false;
+			Parent = InParent;
+
+			if (InParent != null)
+			{
+				bIsInstanceDefinition = InParent.bIsInstanceDefinition;
+				WorldTransform = Transform.Multiply(InParent.WorldTransform, WorldTransform);
+				LayerIndices.UnionWith(InParent.LayerIndices);
+			}
+		}
+
+		public void ApplyTransform(Transform InTransform)
+		{
+			foreach (DatasmithActorInfo ActorInfo in GetEnumerator())
+			{
+				ActorInfo.ApplyModifiedStatus();
+				ActorInfo.WorldTransform = InTransform * ActorInfo.WorldTransform;
+			}
+		}
+
+		public void SetDefinitionNode(DatasmithActorInfo InDefinitionNode)
+		{
+			System.Diagnostics.Debug.Assert(InDefinitionNode.bIsInstanceDefinition, "Trying to create an instance from a node belonging in the root tree");
+			DefinitionNode = InDefinitionNode;
+			InDefinitionNode.InstanceNodesInternal.AddLast(this);
+			LayerIndices.UnionWith(InDefinitionNode.LayerIndices);
+		}
+
+		public void AddChild(DatasmithActorInfo ChildHierarchyNodeInfo)
+		{
+			ChildHierarchyNodeInfo.SetParent(this);
+			ChildrenInternal.AddLast(ChildHierarchyNodeInfo);
+		}
+
+		/// <summary>
+		/// Returns the number of hierarchical descendants this nodes has.
+		/// </summary>
+		/// <returns></returns>
+		public int GetDescendantsCount()
+		{
+			int DescendantCount = ChildrenInternal.Count;
+
+			foreach (DatasmithActorInfo CurrentChild in ChildrenInternal)
+			{
+				DescendantCount += CurrentChild.GetDescendantsCount();
+			}
+
+			return DescendantCount;
+		}
+
+		public void Reparent(DatasmithActorInfo NewParent)
+		{
+			Parent.ChildrenInternal.Remove(this);
+			NewParent.AddChild(this);
+		}
+
+		/// <summary>
+		/// Remove all actor infos with the Deleted DirectLinkStatus from the hierarchy.
+		/// </summary>
+		public void PurgeDeleted()
+		{
+			LinkedListNode<DatasmithActorInfo> CurrentNode = ChildrenInternal.First;
+			while (CurrentNode != null)
+			{
+				if (CurrentNode.Value.DirectLinkStatus != DirectLinkSynchronizationStatus.Deleted)
+				{
+					CurrentNode.Value.PurgeDeleted();
+					CurrentNode = CurrentNode.Next;
+				}
+				else
+				{
+					LinkedListNode<DatasmithActorInfo> NodeToDelete = CurrentNode;
+					CurrentNode = CurrentNode.Next;
+					ChildrenInternal.Remove(NodeToDelete);
+				}
+			}
+
+			CurrentNode = InstanceNodesInternal.First;
+			while (CurrentNode != null)
+			{
+				if (CurrentNode.Value.DirectLinkStatus == DirectLinkSynchronizationStatus.Deleted)
+				{
+					LinkedListNode<DatasmithActorInfo> NodeToDelete = CurrentNode;
+					CurrentNode = CurrentNode.Next;
+					InstanceNodesInternal.Remove(NodeToDelete);
+				}
+				else
+				{
+					CurrentNode = CurrentNode.Next;
+				}
+			}
+		}
+
+		/// <summary>
+		/// Custom enumerator implementation returning this Actor and all its descendant.
+		/// </summary>
+		/// <returns></returns>
+		public IEnumerable<DatasmithActorInfo> GetEnumerator()
+		{
+			yield return this;
+
+			foreach (var Child in ChildrenInternal)
+			{
+				foreach (var ChildEnumValue in Child.GetEnumerator())
+				{
+					yield return ChildEnumValue;
+				}
+			}
+		}
+
+		/// <summary>
+		/// Custom enumerator implementation for returning all descendants of this Actor.
+		/// </summary>
+		/// <returns></returns>
+		public IEnumerable<DatasmithActorInfo> GetDescendantEnumerator()
+		{
+			foreach (var Child in ChildrenInternal)
+			{
+				foreach (var ChildEnumValue in Child.GetEnumerator())
+				{
+					yield return ChildEnumValue;
+				}
+			}
+		}
+	}
+
+}
\ No newline at end of file
diff --git a/Engine/Source/Programs/Enterprise/Datasmith/DatasmithRhinoExporter/Private/ExportContext/DatasmithCameraActorInfo.cs b/Engine/Source/Programs/Enterprise/Datasmith/DatasmithRhinoExporter/Private/ExportContext/DatasmithCameraActorInfo.cs
new file mode 100644
--- /dev/null
+++ b/Engine/Source/Programs/Enterprise/Datasmith/DatasmithRhinoExporter/Private/ExportContext/DatasmithCameraActorInfo.cs
@@ -0,0 +1,52 @@
+// Copyright Epic Games, Inc. All Rights Reserved.
+
+using Rhino.DocObjects;
+using System;
+
+namespace DatasmithRhino.ExportContext
+{
+	public class DatasmithCameraActorInfo : DatasmithActorInfo
+	{
+		public override Guid RhinoObjectId
+		{
+			get
+			{
+				if (RhinoCommonObject is ViewportInfo RhinoViewportInfo)
+				{
+					return RhinoViewportInfo.Id;
+				}
+
+				return Guid.Empty;
+			}
+		}
+
+		public FDatasmithFacadeActorCamera ExportedCameraActor { get { return ExportedElement as FDatasmithFacadeActorCamera; } }
+
+		/// <summary>
+		/// Returns a hash of the properties of the ViewportInfo relevant for Datasmith export.
+		/// We use this because Rhino does not provide an event for when the named view table is modified.
+		/// That means every time the document is modified we must check for named view changes, the hash allows us to avoid false-positives modifications.
+		/// </summary>
+		public String NamedViewHash { get; private set; }
+
+		public DatasmithCameraActorInfo(ViewportInfo InViewportInfo, string InName, string InUniqueLabel, string InBaseLabel, string InHash)
+			: base(InViewportInfo, InName, InUniqueLabel, InBaseLabel)
+		{
+			NamedViewHash = InHash;
+		}
+
+		public override void ApplyDiffs(DatasmithInfoBase OtherInfo)
+		{
+			base.ApplyDiffs(OtherInfo);
+
+			if (OtherInfo is DatasmithCameraActorInfo OtherActorCameraInfo)
+			{
+				NamedViewHash = OtherActorCameraInfo.NamedViewHash;
+			}
+			else
+			{
+				System.Diagnostics.Debug.Fail("OtherInfo does not derive from DatasmithActorCameraInfo");
+			}
+		}
+	}
+}
\ No newline at end of file
diff --git a/Engine/Source/Programs/Enterprise/Datasmith/DatasmithRhinoExporter/Private/ExportContext/DatasmithInfoBase.cs b/Engine/Source/Programs/Enterprise/Datasmith/DatasmithRhinoExporter/Private/ExportContext/DatasmithInfoBase.cs
new file mode 100644
--- /dev/null
+++ b/Engine/Source/Programs/Enterprise/Datasmith/DatasmithRhinoExporter/Private/ExportContext/DatasmithInfoBase.cs
@@ -0,0 +1,147 @@
+// Copyright Epic Games, Inc. All Rights Reserved.
+
+using System.Runtime.CompilerServices;
+
+namespace DatasmithRhino.ExportContext
+{
+	public enum DirectLinkSynchronizationStatus
+	{
+		None = 0,
+		Created = 1 << 0,
+		Modified = 1 << 1,
+		Deleted = 1 << 2,
+		Synced = 1 << 3,
+		PendingDeletion = 1 << 4,
+		PendingHidding = 1 << 5,
+		//Same as Deleted, except we do not clean up the object afterward.
+		Hidden = 1 << 6,
+	}
+
+	public abstract class DatasmithInfoBase
+	{
+		public Rhino.Runtime.CommonObject RhinoCommonObject { get; protected set; }
+
+		/// <summary>
+		/// Used as a unique ID corresponding to the IDatasmithElement::Name field.
+		/// </summary>
+		public string Name { get; private set; }
+
+		/// <summary>
+		/// Label corresponding to the IDatasmithElement::Label field. It is generated from the BaseLabel to ensure its unicity.
+		/// </summary>
+		public string UniqueLabel { get; private set; }
+
+		/// <summary>
+		/// BaseLabel is used to generate the UniqueLabel. Changes to the BaseLabel are reflected on the UniqueLabel.
+		/// </summary>
+		public string BaseLabel { get; private set; }
+
+		private DirectLinkSynchronizationStatus InternalDirectLinkStatus = DirectLinkSynchronizationStatus.Created;
+		private DirectLinkSynchronizationStatus PreviousDirectLinkStatus { get; set; } = DirectLinkSynchronizationStatus.None;
+		public virtual DirectLinkSynchronizationStatus DirectLinkStatus
+		{
+			get => InternalDirectLinkStatus;
+		}
+		public FDatasmithFacadeElement ExportedElement { get; private set; } = null;
+
+		public DatasmithInfoBase(Rhino.Runtime.CommonObject InRhinoObject, string InName, string InUniqueLabel, string InBaseLabel)
+		{
+			RhinoCommonObject = InRhinoObject;
+			Name = InName;
+			UniqueLabel = InUniqueLabel;
+			BaseLabel = InBaseLabel;
+		}
+
+		[MethodImpl(MethodImplOptions.AggressiveInlining)]
+		public void ApplyModifiedStatus()
+		{
+			SetDirectLinkStatus(DirectLinkSynchronizationStatus.Modified);
+		}
+
+		[MethodImpl(MethodImplOptions.AggressiveInlining)]
+		public void ApplyHiddenStatus()
+		{
+			SetDirectLinkStatus(DirectLinkSynchronizationStatus.PendingHidding);
+		}
+
+		[MethodImpl(MethodImplOptions.AggressiveInlining)]
+		public void ApplyDeletedStatus()
+		{
+			SetDirectLinkStatus(DirectLinkSynchronizationStatus.PendingDeletion);
+		}
+
+		private void SetDirectLinkStatus(DirectLinkSynchronizationStatus Status)
+		{
+			// We ignore the "modified" status if the current status is "created", "deleted" or "hidden".
+			// To "undelete" or unhide an element use RestorePreviousDirectLinkStatus().
+			const DirectLinkSynchronizationStatus UnmodifiableStates = DirectLinkSynchronizationStatus.Created | DirectLinkSynchronizationStatus.Deleted
+				| DirectLinkSynchronizationStatus.PendingDeletion | DirectLinkSynchronizationStatus.PendingHidding | DirectLinkSynchronizationStatus.Hidden;
+
+			if (InternalDirectLinkStatus != Status
+				&& (Status != DirectLinkSynchronizationStatus.Modified
+				|| (InternalDirectLinkStatus & UnmodifiableStates) == DirectLinkSynchronizationStatus.None))
+			{
+				PreviousDirectLinkStatus = InternalDirectLinkStatus;
+				InternalDirectLinkStatus = Status;
+			}
+		}
+
+		public void ApplySyncedStatus()
+		{
+			// Since the Delete and Hide status have a pending phase, we should not override the PreviousDirectLinkStatus in those cases.
+			if (InternalDirectLinkStatus == DirectLinkSynchronizationStatus.PendingDeletion)
+			{
+				InternalDirectLinkStatus = DirectLinkSynchronizationStatus.Deleted;
+			}
+			else if (InternalDirectLinkStatus == DirectLinkSynchronizationStatus.PendingHidding)
+			{
+				InternalDirectLinkStatus = DirectLinkSynchronizationStatus.Hidden;
+			}
+			//If we are not hidden or deleted
+			else if ((InternalDirectLinkStatus & ~(DirectLinkSynchronizationStatus.Deleted | DirectLinkSynchronizationStatus.Hidden)) != DirectLinkSynchronizationStatus.None)
+			{
+				PreviousDirectLinkStatus = InternalDirectLinkStatus;
+				InternalDirectLinkStatus = DirectLinkSynchronizationStatus.Synced;
+			}
+		}
+
+		/// <summary>
+		/// Used to undo the last change to the DirectLinkStatus property. It is intended to be used mainly for "undeleted" objects for which the deletion was not synced.
+		/// </summary>
+		/// <returns></returns>
+		public bool RestorePreviousDirectLinkStatus()
+		{
+			if (InternalDirectLinkStatus == DirectLinkSynchronizationStatus.Hidden)
+			{
+				// We are in hidden state, that means if we want to restore the DatasmithElement it must be flagged as "created".
+				InternalDirectLinkStatus = DirectLinkSynchronizationStatus.Created;
+				PreviousDirectLinkStatus = DirectLinkSynchronizationStatus.None;
+				return true;
+			}
+			else if (PreviousDirectLinkStatus != DirectLinkSynchronizationStatus.None)
+			{
+				InternalDirectLinkStatus = PreviousDirectLinkStatus;
+				PreviousDirectLinkStatus = DirectLinkSynchronizationStatus.None;
+				return true;
+			}
+
+			return false;
+		}
+
+		public void SetExportedElement(FDatasmithFacadeElement InExportedElement)
+		{
+			ExportedElement = InExportedElement;
+		}
+
+		public virtual void ApplyDiffs(DatasmithInfoBase OtherInfo)
+		{
+			SetDirectLinkStatus(DirectLinkSynchronizationStatus.Modified);
+
+			if (BaseLabel != OtherInfo.BaseLabel)
+			{
+				BaseLabel = OtherInfo.BaseLabel;
+				UniqueLabel = OtherInfo.UniqueLabel;
+			}
+		}
+	}
+}
\ No newline at end of file
diff --git a/Engine/Source/Programs/Enterprise/Datasmith/DatasmithRhinoExporter/Private/ExportContext/DatasmithMaterialInfo.cs b/Engine/Source/Programs/Enterprise/Datasmith/DatasmithRhinoExporter/Private/ExportContext/DatasmithMaterialInfo.cs
new file mode 100644
--- /dev/null
+++ b/Engine/Source/Programs/Enterprise/Datasmith/DatasmithRhinoExporter/Private/ExportContext/DatasmithMaterialInfo.cs
@@ -0,0 +1,48 @@
+// Copyright Epic Games, Inc. All Rights Reserved.
+
+using Rhino.DocObjects;
+using System.Collections.Generic;
+
+namespace DatasmithRhino.ExportContext
+{
+	public class DatasmithMaterialInfo : DatasmithInfoBase
+	{
+		public Material RhinoMaterial { get { return RhinoCommonObject as Material; } }
+		public FDatasmithFacadeUEPbrMaterial ExportedMaterial { get { return ExportedElement as FDatasmithFacadeUEPbrMaterial; } }
+
+		/// <summary>
+		/// Set of rhino material indexes that are represented by this DatasmithMaterialInfo
+		/// </summary>
+		public HashSet<int> MaterialIndexes { get; } = new HashSet<int>();
+
+		private List<string> InternalTextureHashes;
+		public IReadOnlyList<string> TextureHashes { get => InternalTextureHashes; }
+
+		public DatasmithMaterialInfo(Material InRhinoMaterial, string InName, string InUniqueLabel, string InBaseLabel, List<string> InTextureHashes)
+			: base(InRhinoMaterial, InName, InUniqueLabel, InBaseLabel)
+		{
+			InternalTextureHashes = new List<string>(InTextureHashes);
+		}
+
+		public override void ApplyDiffs(DatasmithInfoBase OtherInfo)
+		{
+			base.ApplyDiffs(OtherInfo);
+
+			if (OtherInfo is DatasmithMaterialInfo OtherMaterialInfo)
+			{
+				//Rhino "replaces" object instead of modifying them, we must update the our reference to the object.
+				RhinoCommonObject = OtherMaterialInfo.RhinoCommonObject;
+
+				MaterialIndexes.Clear();
+				MaterialIndexes.UnionWith(OtherMaterialInfo.MaterialIndexes);
+
+				InternalTextureHashes.Clear();
+				InternalTextureHashes.AddRange(OtherMaterialInfo.TextureHashes);
+			}
+			else
+			{
+				System.Diagnostics.Debug.Fail("OtherInfo does not derive from DatasmithMaterialInfo");
+			}
+		}
+	}
+}
\ No newline at end of file
diff --git a/Engine/Source/Programs/Enterprise/Datasmith/DatasmithRhinoExporter/Private/ExportContext/DatasmithMeshInfo.cs b/Engine/Source/Programs/Enterprise/Datasmith/DatasmithRhinoExporter/Private/ExportContext/DatasmithMeshInfo.cs
new file mode 100644
--- /dev/null
+++ b/Engine/Source/Programs/Enterprise/Datasmith/DatasmithRhinoExporter/Private/ExportContext/DatasmithMeshInfo.cs
@@ -0,0 +1,57 @@
+// Copyright Epic Games, Inc. All Rights Reserved.
+
+using Rhino.Geometry;
+using System.Collections.Generic;
+
+namespace DatasmithRhino.ExportContext
+{
+	public class DatasmithMeshInfo : DatasmithInfoBase
+	{
+		public FDatasmithFacadeMeshElement ExportedMesh { get { return ExportedElement as FDatasmithFacadeMeshElement; } }
+
+		public List<Mesh> RhinoMeshes { get; private set; }
+		public Transform OffsetTransform { get; set; }
+		public List<int> MaterialIndices { get; set; }
+		public List<DatasmithTextureMappingData> TextureMappings { get; set; }
+
+		public DatasmithMeshInfo(IEnumerable<Mesh> InRhinoMeshes, Transform InOffset, List<int> InMaterialIndexes, List<DatasmithTextureMappingData> InTextureMappings, string InName, string InUniqueLabel, string InBaseLabel)
+			: base(null, InName, InUniqueLabel, InBaseLabel)
+		{
+			RhinoMeshes = new List<Mesh>(InRhinoMeshes);
+			OffsetTransform = InOffset;
+			MaterialIndices = InMaterialIndexes;
+			TextureMappings = InTextureMappings;
+		}
+
+		public override void ApplyDiffs(DatasmithInfoBase OtherInfo)
+		{
+			base.ApplyDiffs(OtherInfo);
+
+			if (OtherInfo is DatasmithMeshInfo OtherMeshInfo)
+			{
+				RhinoMeshes = OtherMeshInfo.RhinoMeshes;
+				OffsetTransform = OtherMeshInfo.OffsetTransform;
+				MaterialIndices = OtherMeshInfo.MaterialIndices;
+				TextureMappings = OtherMeshInfo.TextureMappings;
+			}
+			else
+			{
+				System.Diagnostics.Debug.Fail("OtherInfo does not derive from DatasmithMeshInfo");
+			}
+		}
+	}
+
+	public class DatasmithTextureMappingData
+	{
+		public int ChannelId;
+		public Rhino.Render.TextureMapping RhinoTextureMapping;
+		public Transform ObjectTransform;
+
+		public DatasmithTextureMappingData(int InChannelId, Rhino.Render.TextureMapping InTextureMapping, Transform InTransform)
+		{
+			ChannelId = InChannelId;
+			RhinoTextureMapping = InTextureMapping;
+			ObjectTransform = InTransform;
+		}
+	}
+}
\ No newline at end of file
diff --git a/Engine/Source/Programs/Enterprise/Datasmith/DatasmithRhinoExporter/Private/ExportContext/DatasmithRhinoExportContext.cs b/Engine/Source/Programs/Enterprise/Datasmith/DatasmithRhinoExporter/Private/ExportContext/DatasmithRhinoExportContext.cs
new file mode 100644
--- /dev/null
+++ b/Engine/Source/Programs/Enterprise/Datasmith/DatasmithRhinoExporter/Private/ExportContext/DatasmithRhinoExportContext.cs
@@ -0,0 +1,1714 @@
+// Copyright Epic Games, Inc. All Rights Reserved.
+
+using DatasmithRhino.Utils;
+using Rhino;
+using Rhino.Display;
+using Rhino.DocObjects;
+using Rhino.DocObjects.Tables;
+using Rhino.Geometry;
+using System;
+using System.Collections.Generic;
+using System.Linq;
+using System.Text;
+
+namespace DatasmithRhino.ExportContext
+{
+	public class DatasmithRhinoExportContext
+	{
+		public RhinoDoc RhinoDocument { get => ExportOptions.RhinoDocument; }
+		public DatasmithRhinoExportOptions ExportOptions { get; private set; }
+		public bool bIsParsed { get; private set; } = false;
+		public bool bIsDirty { get; private set; } = true;
+		public bool bExportedOnce { get; private set; } = false;
+		public bool bIsInWorksession {
+			get {
+				//Only check for worksession on Windows, the feature is not implemented on Mac and calling the API throws exception.
+				return Environment.OSVersion.Platform == PlatformID.Win32NT 
+					&& RhinoDocument.Worksession != null
+					&& RhinoDocument.Worksession.ModelCount > 1;
+			}
+		}
+		private string[] WorksessionDocumentPaths = null;
+
+		/// <summary>
+		/// If true the scene elements driven by the DocuementPropertiesChanged event will be parsed before export.
+		/// </summary>
+		public bool bDocumentPropertiesChanged = false;
+
+		public DatasmithActorInfo SceneRoot = null;
+		public Dictionary<InstanceDefinition, DatasmithActorInfo> InstanceDefinitionHierarchyNodeDictionary = new Dictionary<InstanceDefinition, DatasmithActorInfo>();
+		public Dictionary<Guid, DatasmithActorInfo> ObjectIdToHierarchyActorNodeDictionary = new Dictionary<Guid, DatasmithActorInfo>();
+		public Dictionary<Guid, DatasmithMeshInfo> ObjectIdToMeshInfoDictionary = new Dictionary<Guid, DatasmithMeshInfo>();
+		public Dictionary<string, DatasmithCameraActorInfo> NamedViewNameToCameraInfo = new Dictionary<string, DatasmithCameraActorInfo>();
+		public Dictionary<string, DatasmithMaterialInfo> MaterialHashToMaterialInfo = new Dictionary<string, DatasmithMaterialInfo>();
+		public Dictionary<string, DatasmithTextureInfo> TextureHashToTextureInfo = new Dictionary<string, DatasmithTextureInfo>();
+		public Dictionary<int, string> GroupIndexToName = new Dictionary<int, string>();
+		public Dictionary<Guid, RhinoObject> TextureMappindIdToRhinoObject = new Dictionary<Guid, RhinoObject>();
+
+		private Dictionary<int, string> MaterialIndexToMaterialHashDictionary = new Dictionary<int, string>();
+		private Dictionary<Guid, string> TextureIdToTextureHash = new Dictionary<Guid, string>();
+		private Dictionary<int, string> LayerIndexToLayerString = new Dictionary<int, string>();
+		private Dictionary<int, HashSet<int>> LayerIndexToLayerIndexHierarchy = new Dictionary<int, HashSet<int>>();
+		private Dictionary<int, HashSet<DatasmithInfoBase>> MaterialIndexToElementInfosMap = new Dictionary<int, HashSet<DatasmithInfoBase>>();
+		private DatasmithRhinoUniqueNameGenerator ActorLabelGenerator = new DatasmithRhinoUniqueNameGenerator();
+		private DatasmithRhinoUniqueNameGenerator MaterialLabelGenerator = new DatasmithRhinoUniqueNameGenerator();
+		private DatasmithRhinoUniqueNameGenerator TextureLabelGenerator = new DatasmithRhinoUniqueNameGenerator();
+		private ViewportInfo ActiveViewportInfo;
+		private int DummyLayerIndex = -1;
+
+		public DatasmithRhinoExportContext(DatasmithRhinoExportOptions InOptions)
+		{
+			ExportOptions = InOptions;
+			
+			const string RootName = "SceneRoot";
+			SceneRoot = new DatasmithActorInfo(ExportOptions.Xform, RootName, RootName, RootName);
+			SceneRoot.ApplySyncedStatus();
+		}
+
+		public void ParseDocument(bool bForceParse = false)
+		{
+			if (!bIsParsed || bForceParse)
+			{
+				try
+				{
+					RhinoViewport ActiveViewport = RhinoDocument.Views.ActiveView?.ActiveViewport;
+
+					if (bIsInWorksession)
+					{
+						WorksessionDocumentPaths = RhinoDocument.Worksession.ModelPaths;
+					}
+
+					//Update current active viewport.
+					ActiveViewportInfo = ActiveViewport == null ? null : new ViewportInfo(ActiveViewport);
+
+					DatasmithRhinoProgressManager.Instance.UpdateCurrentTaskProgress(0.1f);
+					ParseGroupNames();
+					DatasmithRhinoProgressManager.Instance.UpdateCurrentTaskProgress(0.66f);
+					ParseRhinoHierarchy();
+					ParseRhinoCameras();
+					DatasmithRhinoProgressManager.Instance.UpdateCurrentTaskProgress(1f);
+					ParseAllRhinoMeshes();
+
+					bIsParsed = true;
+					bIsDirty = true;
+				}
+				catch (Exception UnexpectedException)
+				{
+					// We can't resume halfway through the scene parsing, just reset the context and start over next time.
+					ResetContext();
+					// throw forward the exception.
+					
+					throw new Exception(UnexpectedException.Message, UnexpectedException);
+				}
+			}
+			else
+			{
+				// The scene is already parsed. We just need to update the visibility flags.
+				HashSet<RhinoObject> OutNewObjects = new HashSet<RhinoObject>();
+				UpdateHiddenFlagsRecursively(SceneRoot, OutNewObjects);
+				
+				// Parse the meshes of unhidden actors being exported for the first time.
+				AddOrModifyMeshesFromRhinoObjects(OutNewObjects);
+
+				if (bDocumentPropertiesChanged)
+				{
+					ParseRhinoCameras();
+					bDocumentPropertiesChanged = false;
+				}
+			}
+		}
+
+		private void ResetContext()
+		{
+			const string RootName = "SceneRoot";
+			SceneRoot = new DatasmithActorInfo(ExportOptions.Xform, RootName, RootName, RootName);
+			SceneRoot.ApplySyncedStatus();
+			
+			bIsParsed = false;
+			InstanceDefinitionHierarchyNodeDictionary = new Dictionary<InstanceDefinition, DatasmithActorInfo>();
+			ObjectIdToHierarchyActorNodeDictionary = new Dictionary<Guid, DatasmithActorInfo>();
+			ObjectIdToMeshInfoDictionary = new Dictionary<Guid, DatasmithMeshInfo>();
+			NamedViewNameToCameraInfo = new Dictionary<string, DatasmithCameraActorInfo>();
+			MaterialHashToMaterialInfo = new Dictionary<string, DatasmithMaterialInfo>();
+			TextureHashToTextureInfo = new Dictionary<string, DatasmithTextureInfo>();
+			GroupIndexToName = new Dictionary<int, string>();
+			TextureMappindIdToRhinoObject = new Dictionary<Guid, RhinoObject>();
+
+			MaterialIndexToMaterialHashDictionary = new Dictionary<int, string>();
+			TextureIdToTextureHash = new Dictionary<Guid, string>();
+			LayerIndexToLayerString = new Dictionary<int, string>();
+			LayerIndexToLayerIndexHierarchy = new Dictionary<int, HashSet<int>>();
+			MaterialIndexToElementInfosMap = new Dictionary<int, HashSet<DatasmithInfoBase>>();
+			ActorLabelGenerator = new DatasmithRhinoUniqueNameGenerator();
+			MaterialLabelGenerator = new DatasmithRhinoUniqueNameGenerator();
+			TextureLabelGenerator = new DatasmithRhinoUniqueNameGenerator();
+			DummyLayerIndex = -1;
+		}
+
+		public DatasmithMaterialInfo GetMaterialInfoFromMaterialIndex(int MaterialIndex)
+		{
+			if(MaterialIndexToMaterialHashDictionary.TryGetValue(MaterialIndex, out string MaterialHash))
+			{
+				if(MaterialHashToMaterialInfo.TryGetValue(MaterialHash, out DatasmithMaterialInfo MaterialInfo))
+				{
+					return MaterialInfo;
+				}
+			}
+
+			return null;
+		}
+
+		public DatasmithTextureInfo GetTextureInfoFromRhinoTexture(Guid TextureId)
+		{
+			if (TextureIdToTextureHash.TryGetValue(TextureId, out string TextureHash))
+			{
+				if (TextureHashToTextureInfo.TryGetValue(TextureHash, out DatasmithTextureInfo TextureInfo))
+				{
+					return TextureInfo;
+				}
+			}
+
+			return null;
+		}
+
+		public string GetNodeLayerString(DatasmithActorInfo Node)
+		{
+			StringBuilder Buider = new StringBuilder();
+			foreach (int LayerIndex in Node.LayerIndices)
+			{
+				if (LayerIndexToLayerString.TryGetValue(LayerIndex, out string LayerString))
+				{
+					if (Buider.Length == 0)
+					{
+						Buider.Append(LayerString);
+					}
+					else
+					{
+						Buider.Append("," + LayerString);
+					}
+				}
+			}
+
+			return Buider.ToString();
+		}
+
+		/// <summary>
+		/// This function should be called after each export to update the context tracking data.
+		/// </summary>
+		public void OnPostExport()
+		{
+			// Block definitions must be set to synced state after exporting, since their status is not updated during export.
+			foreach (DatasmithActorInfo InstanceDefinitionInfo in InstanceDefinitionHierarchyNodeDictionary.Values)
+			{
+				foreach (DatasmithActorInfo DefinitionNode in InstanceDefinitionInfo.GetEnumerator())
+				{
+					DefinitionNode.ApplySyncedStatus();
+				}
+			}
+
+			if (bExportedOnce)
+			{
+				Cleanup();
+			}
+			
+			bExportedOnce = true;
+			bIsDirty = false;
+		}
+
+		/// <summary>
+		/// Remove cached data for deleted object no longer needed.
+		/// </summary>
+		private void Cleanup()
+		{
+			CleanCacheDictionary(InstanceDefinitionHierarchyNodeDictionary);
+			CleanCacheDictionary(ObjectIdToHierarchyActorNodeDictionary);
+			CleanCacheDictionary(ObjectIdToMeshInfoDictionary);
+			CleanCacheDictionary(MaterialHashToMaterialInfo);
+			CleanCacheDictionary(TextureHashToTextureInfo);
+
+			SceneRoot.PurgeDeleted();
+			foreach (DatasmithActorInfo InstanceDefinitionInfo in InstanceDefinitionHierarchyNodeDictionary.Values)
+			{
+				InstanceDefinitionInfo.PurgeDeleted();
+			}
+
+			CleanMaterialElementCache();
+		}
+
+		private void CleanCacheDictionary<KeyType, InfoType>(Dictionary<KeyType, InfoType> InDictionary) where InfoType : DatasmithInfoBase
+		{
+			List<KeyType> KeysOfElementsToDelete = new List<KeyType>();
+
+			foreach (KeyValuePair<KeyType, InfoType> CurrentKeyValuePair in InDictionary)
+			{
+				if (CurrentKeyValuePair.Value.DirectLinkStatus == DirectLinkSynchronizationStatus.Deleted)
+				{
+					KeysOfElementsToDelete.Add(CurrentKeyValuePair.Key);
+				}
+			}
+
+			foreach (KeyType CurrentKey in KeysOfElementsToDelete)
+			{
+				InDictionary.Remove(CurrentKey);
+			}
+		}
+
+		private void CleanMaterialElementCache()
+		{
+			foreach (HashSet<DatasmithInfoBase> ElementSet in MaterialIndexToElementInfosMap.Values)
+			{
+				List<DatasmithInfoBase> ElementsToRemove = new List<DatasmithInfoBase>();
+				foreach (DatasmithInfoBase ElementInSet in ElementSet)
+				{
+					if (ElementInSet.DirectLinkStatus == DirectLinkSynchronizationStatus.Deleted)
+					{
+						ElementsToRemove.Add(ElementInSet);
+					}
+				}
+
+				for (int ElementIndex = 0; ElementIndex < ElementsToRemove.Count; ++ElementIndex)
+				{
+					ElementSet.Remove(ElementsToRemove[ElementIndex]);
+				}
+			}
+		}
+
+		public void OnModelUnitChange()
+		{
+			FDatasmithFacadeElement.SetWorldUnitScale((float)RhinoMath.UnitScale(RhinoDocument.ModelUnitSystem, UnitSystem.Centimeters));
+			ExportOptions.ModelUnitSystem = RhinoDocument.ModelUnitSystem;
+
+			foreach (InstanceDefinition BlockInstanceDefinition in InstanceDefinitionHierarchyNodeDictionary.Keys)
+			{
+				UpdateDefinitionNode(BlockInstanceDefinition);
+			}
+
+			foreach (DatasmithActorInfo DescendantInfo in SceneRoot.GetDescendantEnumerator())
+			{
+				RhinoObject DescendantRhinoObject = DescendantInfo.RhinoCommonObject as RhinoObject;
+				bool bIsNotInstancedObject = DescendantRhinoObject == null || DescendantInfo.DefinitionNode == null || (DescendantRhinoObject.ObjectType == ObjectType.InstanceReference && !DescendantInfo.bIsInstanceDefinition);
+				
+				if (bIsNotInstancedObject)
+				{
+					const bool bReparent = false;
+					ModifyActor(DescendantRhinoObject, bReparent);
+				}
+			}
+
+			bIsDirty = true;
+		}
+
+		/// <summary>
+		/// Actor creation event used when a RhinoObject is created.
+		/// </summary>
+		/// <param name="InModelComponent"></param>
+		public void AddActor(ModelComponent InModelComponent)
+		{
+			if (ObjectIdToHierarchyActorNodeDictionary.TryGetValue(InModelComponent.Id, out DatasmithActorInfo CachedActorInfo))
+			{
+				// The actor already exists in the cache, it means that no direct link synchronization was done since its deletion.
+				// In that case, we can simply undelete it.
+				UndeleteActor(CachedActorInfo);
+				return;
+			}
+
+			DatasmithActorInfo ActorParentInfo = GetActorParentInfo(InModelComponent);
+			if (InModelComponent is Layer InLayer)
+			{
+				RecursivelyParseLayerHierarchy(InLayer, ActorParentInfo);
+			}
+			else if (InModelComponent is RhinoObject InRhinoObject)
+			{
+				RecursivelyParseObjectInstance(new[] { InRhinoObject }, ActorParentInfo);
+				RegisterActorsToContext(ActorParentInfo.GetDescendantEnumerator());
+
+				HashSet<RhinoObject> CollectedMeshObjects = new HashSet<RhinoObject>();
+				if (InRhinoObject is InstanceObject InRhinoInstance 
+					&& InstanceDefinitionHierarchyNodeDictionary.TryGetValue(InRhinoInstance.InstanceDefinition, out DatasmithActorInfo InstanceDefinitionInfo))
+				{
+					CollectExportedRhinoObjectsFromInstanceDefinition(CollectedMeshObjects, InstanceDefinitionInfo);
+				}
+				else
+				{
+					CollectedMeshObjects.Add(InRhinoObject);
+				}
+				AddOrModifyMeshesFromRhinoObjects(CollectedMeshObjects);
+			}
+
+			bIsDirty = true;
+		}
+
+		private void UndeleteActor(DatasmithActorInfo ActorInfo)
+		{
+			foreach (DatasmithActorInfo ActorInfoValue in ActorInfo.GetEnumerator())
+			{
+				System.Diagnostics.Debug.Assert(ActorInfoValue.DirectLinkStatus == DirectLinkSynchronizationStatus.PendingDeletion);
+				if (ActorInfoValue.RestorePreviousDirectLinkStatus())
+				{
+					if (ObjectIdToMeshInfoDictionary.TryGetValue(ActorInfoValue.RhinoObjectId, out DatasmithMeshInfo MeshInfo))
+					{
+						MeshInfo.RestorePreviousDirectLinkStatus();
+					}
+				}
+				else
+				{
+					System.Diagnostics.Debug.Fail("The previous direct link state was lost on an undeleted actor");
+					break;
+				}
+			}
+
+			if (ActorInfo.RhinoCommonObject is ModelComponent RhinoModelComponent)
+			{
+				//Undelete is always used during an undo, and undo often use "undelete" to restore an actor state (instead of calling modify).
+				ModifyActor(RhinoModelComponent, /*bReparent=*/false);
+			}
+
+			bIsDirty = true;
+		}
+
+		/// <summary>
+		/// Specific Actor modification event used when the only change is the RhinoObject transform. We avoid reprocessing the Mesh and simply update the actor position.
+		/// </summary>
+		/// <param name="InRhinoObject"></param>
+		/// <param name="InTransform"></param>
+		public void MoveActor(RhinoObject InRhinoObject, Transform InTransform)
+		{
+			bool bHasMesh = false;
+			if (ObjectIdToMeshInfoDictionary.TryGetValue(InRhinoObject.Id, out DatasmithMeshInfo MeshInfo))
+			{
+				// #ueent-todo	The MeshInfo is the one holding the OffsetTransform but changing it actually affects the actor and not the mesh. 
+				//				It is weird that the actor is the one modified here, we should move the OffsetTransform into it's associated DatasmithActorInfo.
+				MeshInfo.OffsetTransform = InTransform * MeshInfo.OffsetTransform;
+				bHasMesh = true;
+			}
+
+			bool bIsInstance = InRhinoObject is InstanceObject;
+			if ((bHasMesh || bIsInstance)
+				&& ObjectIdToHierarchyActorNodeDictionary.TryGetValue(InRhinoObject.Id, out DatasmithActorInfo ActorInfo))
+			{
+				if (bIsInstance)
+				{
+					ActorInfo.ApplyTransform(InTransform);
+				}
+				else
+				{
+					//Mesh offset has changed, this will affect related actors.
+					ActorInfo.ApplyModifiedStatus();
+				}
+			}
+
+			bIsDirty = true;
+		}
+
+		/// <summary>
+		/// Generic Actor modification event used when the RhinoObject may have multiple property modifications, both the associated DatasmithActor and DatasmithMesh will be synced.
+		/// </summary>
+		/// <param name="InModelComponent"></param>
+		/// <param name="bReparent"></param>
+		public void ModifyActor(ModelComponent InModelComponent, bool bReparent)
+		{
+			RhinoObject InRhinoObject = InModelComponent as RhinoObject;
+			if (ObjectIdToHierarchyActorNodeDictionary.TryGetValue(InModelComponent.Id, out DatasmithActorInfo ActorInfo))
+			{
+				Layer InLayer = InModelComponent as Layer;
+				DatasmithActorInfo ActorParentInfo = bReparent 
+					? GetActorParentInfo(InModelComponent)
+					: ActorInfo.Parent;
+
+				DatasmithActorInfo DiffActorInfo = null;
+				if (InRhinoObject != null)
+				{
+					DiffActorInfo = TryParseObjectInstance(InRhinoObject, ActorParentInfo, out DatasmithActorInfo DefinitionRootNode);
+				}
+				else if (InLayer != null)
+				{
+					DiffActorInfo = TryParseLayer(InLayer, ActorParentInfo);
+				}
+
+				if (DiffActorInfo != null)
+				{
+					bool bMaterialChanged = ActorInfo.MaterialIndex != DiffActorInfo.MaterialIndex;
+					if (bMaterialChanged)
+					{
+						if (InRhinoObject != null)
+						{
+							AddObjectMaterialReference(InRhinoObject, DiffActorInfo.MaterialIndex);
+						}
+						else
+						{
+							AddMaterialIndexMapping(DiffActorInfo.MaterialIndex);
+						}
+						UpdateMaterialElementMapping(ActorInfo, DiffActorInfo.MaterialIndex, ActorInfo.MaterialIndex);
+					}
+
+					ActorInfo.ApplyDiffs(DiffActorInfo);
+					if (bReparent)
+					{
+						ActorInfo.Reparent(ActorParentInfo);
+					}
+
+					if (bReparent || bMaterialChanged)
+					{
+						UpdateChildActorsMaterialIndex(ActorInfo);
+					}
+				}
+				else
+				{
+					DeleteActor(InModelComponent);
+					return;
+				}
+			}
+			else
+			{
+				AddActor(InModelComponent);
+				return;
+			}
+
+			if (InRhinoObject != null)
+			{
+				RhinoObject[] RhinoObjects = { InRhinoObject };
+				AddOrModifyMeshesFromRhinoObjects(RhinoObjects);
+			}
+
+			bIsDirty = true;
+		}
+
+		public void UpdateActorObject(ModelComponent InModelComponent)
+		{
+			if (ObjectIdToHierarchyActorNodeDictionary.TryGetValue(InModelComponent.Id, out DatasmithActorInfo ActorInfo))
+			{
+				ActorInfo.UpdateRhinoObject(InModelComponent);
+			}
+		}
+
+		private void UpdateHiddenFlagsRecursively(DatasmithActorInfo ActorInfo, HashSet<RhinoObject> OutNewObjects)
+		{
+			// if the visibility changed, update the status of the descendants.
+			bool bIsVisible = ActorInfo.bIsVisible;
+			bool bHasHiddenFlag = (ActorInfo.DirectLinkStatus & (DirectLinkSynchronizationStatus.PendingHidding | DirectLinkSynchronizationStatus.Hidden)) != DirectLinkSynchronizationStatus.None;
+			if (bIsVisible == bHasHiddenFlag)
+			{
+				if (!bIsVisible)
+				{
+					// If we are hiding the actor, then we know that all its children are hidden as well.
+					foreach (DatasmithActorInfo DescendantActor in ActorInfo.GetEnumerator())
+					{
+						DescendantActor.ApplyHiddenStatus();
+					}
+
+					return;
+				}
+				else
+				{
+					// If we are unhiding the actor, we must recompute the visibility of each children as it may be affected by other layers as well.
+					ActorInfo.RestorePreviousDirectLinkStatus();
+
+					// If we are unhidding an actor which previous status was "Created" that means we never synced it and must generate its associated mesh.
+					if (ActorInfo.DirectLinkStatus == DirectLinkSynchronizationStatus.Created 
+						&& ActorInfo.RhinoCommonObject is RhinoObject NewRhinoObject
+						&& NewRhinoObject.ObjectType != ObjectType.InstanceReference)
+					{
+						OutNewObjects.Add(NewRhinoObject);
+					}
+				}
+			}
+
+			foreach (DatasmithActorInfo ChildActorInfo in ActorInfo.Children)
+			{
+				UpdateHiddenFlagsRecursively(ChildActorInfo, OutNewObjects);
+			}
+		}
+
+		private DatasmithActorInfo GetActorParentInfo(ModelComponent InModelComponent)
+		{
+			Layer ParentLayer = null;
+			if (InModelComponent is RhinoObject InRhinoObject)
+			{
+				ParentLayer = RhinoDocument.Layers.FindIndex(InRhinoObject.Attributes.LayerIndex);
+			}
+			else if (InModelComponent is Layer InLayer)
+			{
+				ParentLayer = RhinoDocument.Layers.FindId(InLayer.ParentLayerId);
+			}
+
+			DatasmithActorInfo ActorParentInfo;
+			if (ParentLayer != null)
+			{
+				ObjectIdToHierarchyActorNodeDictionary.TryGetValue(ParentLayer.Id, out ActorParentInfo);
+			}
+			else
+			{
+				ActorParentInfo = SceneRoot;
+			}
+
+			return ActorParentInfo;
+		}
+
+		/// <summary>
+		/// Actor deletion event used when a RhinoObject is deleted, calling this will remove the associated DatasmithActorElement on the next sync.
+		/// </summary>
+		/// <param name="InModelComponent"></param>
+		public void DeleteActor(ModelComponent InModelComponent)
+		{
+			if (ObjectIdToHierarchyActorNodeDictionary.TryGetValue(InModelComponent.Id, out DatasmithActorInfo ActorInfo))
+			{
+				foreach (DatasmithActorInfo ActorInfoValue in ActorInfo.GetEnumerator())
+				{
+					ActorInfoValue.ApplyDeletedStatus();
+
+					// If this actor is not a block instance, we can delete its associated mesh.
+					if (ActorInfoValue.DefinitionNode == null && ObjectIdToMeshInfoDictionary.TryGetValue(ActorInfoValue.RhinoObjectId, out DatasmithMeshInfo MeshInfo))
+					{
+						MeshInfo.ApplyDeletedStatus();
+					}
+				}
+			}
+
+			bIsDirty = true;
+		}
+
+		public void UpdateChildActorsMaterialIndex(DatasmithActorInfo ActorInfo)
+		{
+			foreach (DatasmithActorInfo ChildActorInfo in ActorInfo.Children)
+			{
+				ChildActorInfo.ApplyModifiedStatus();
+				bool bMaterialAffectedByParent = ChildActorInfo.RhinoCommonObject is RhinoObject ChildRhinoObject && ChildRhinoObject.Attributes.MaterialSource != ObjectMaterialSource.MaterialFromObject;
+
+				if (bMaterialAffectedByParent)
+				{
+					DatasmithActorInfo DefinitionInfo = ChildActorInfo.DefinitionNode != null ? ChildActorInfo.DefinitionNode : ChildActorInfo;
+					RhinoObject TargetObject = DefinitionInfo.RhinoCommonObject as RhinoObject;
+
+					if (TargetObject != null)
+					{
+						ChildActorInfo.MaterialIndex = GetObjectMaterialIndex(TargetObject, ChildActorInfo.Parent);
+						if (!ChildActorInfo.bOverrideMaterial)
+						{
+							if (ObjectIdToMeshInfoDictionary.TryGetValue(TargetObject.Id, out DatasmithMeshInfo MeshInfo))
+							{
+								Dictionary<Mesh,ObjectAttributes> MeshAttributePairs = GetMeshAttributePairs(TargetObject, MeshInfo.RhinoMeshes);
+
+								MeshInfo.ApplyModifiedStatus();
+								MeshInfo.MaterialIndices = GetMeshesMaterialIndices(MeshAttributePairs, DefinitionInfo);
+							}
+						}
+					}
+				}
+			}
+
+			bIsDirty = true;
+		}
+		
+		private void ParseRhinoHierarchy()
+		{
+			DatasmithActorInfo DummyDocumentNode = SceneRoot;
+
+			if (bIsInWorksession)
+			{
+				// Rhino worksession adds dummy layers (with non consistent IDs) to all of the linked documents except the active one.
+				// This can cause reimport inconsistencies when the active document changes, as dummies are shuffled and some may be created while others destroyed.
+				// To address this issue we add our own "Current Document Layer" dummy layer, and use the file path as ActorElement name, 
+				// that way there are no actors deleted and the Datasmith IDs stay consistent.
+				string DummyLayerName = RhinoDocument.Path;
+				string DummyLayerLabel = RhinoDocument.Name;
+				const int DefaultMaterialIndex = -1;
+				//Decrement the dummy layer index for each dummy layer added to ensure unique id for each document.
+				int[] DummyLayerIndices = { DummyLayerIndex-- };
+
+				DummyDocumentNode = GenerateDummyNodeInfo(DummyLayerName, DummyLayerLabel, DefaultMaterialIndex, DummyLayerIndices);
+				SceneRoot.AddChild(DummyDocumentNode);
+				LayerIndexToLayerString.Add(DummyLayerIndices[0], BuildLayerString(DummyDocumentNode.UniqueLabel, SceneRoot));
+			}
+
+			foreach (var CurrentLayer in RhinoDocument.Layers)
+			{
+				//Only add Layers directly under root, the recursion will do the rest.
+				if (CurrentLayer.ParentLayerId == Guid.Empty)
+				{
+					DatasmithActorInfo ParentNode = bIsInWorksession && !CurrentLayer.IsReference
+						? DummyDocumentNode
+						: SceneRoot;
+
+					RecursivelyParseLayerHierarchy(CurrentLayer, ParentNode);
+				}
+			}
+
+			// Update the hidden flags on all the actors
+			HashSet<RhinoObject> OutNewObjects = new HashSet<RhinoObject>();
+			UpdateHiddenFlagsRecursively(SceneRoot, OutNewObjects);
+		}
+
+		private void RecursivelyParseLayerHierarchy(Layer CurrentLayer, DatasmithActorInfo ParentNode)
+		{
+			if (TryParseLayer(CurrentLayer, ParentNode) is DatasmithActorInfo ActorNodeInfo)
+			{
+				ParentNode.AddChild(ActorNodeInfo);
+				ObjectIdToHierarchyActorNodeDictionary.Add(CurrentLayer.Id, ActorNodeInfo);
+				LayerIndexToLayerString[CurrentLayer.Index] = BuildLayerString(CurrentLayer, ParentNode);
+				AddMaterialIndexMapping(CurrentLayer.RenderMaterialIndex);
+				UpdateMaterialElementMapping(ActorNodeInfo, ActorNodeInfo.MaterialIndex);
+
+				RhinoObject[] ObjectsInLayer = RhinoDocument.Objects.FindByLayer(CurrentLayer);
+				RecursivelyParseObjectInstance(ObjectsInLayer, ActorNodeInfo);
+				RegisterActorsToContext(ActorNodeInfo.GetDescendantEnumerator());
+
+				Layer[] ChildrenLayer = CurrentLayer.GetChildren();
+				if (ChildrenLayer != null)
+				{
+					foreach (var ChildLayer in ChildrenLayer)
+					{
+						RecursivelyParseLayerHierarchy(ChildLayer, ActorNodeInfo);
+					}
+				}
+			}
+		}
+
+		private DatasmithActorInfo TryParseLayer(Layer CurrentLayer, DatasmithActorInfo ParentNode)
+		{
+			if ((ExportOptions.bSkipHidden && !CurrentLayer.IsVisible) || CurrentLayer.IsDeleted)
+			{
+				return null;
+			}
+
+			int MaterialIndex = CurrentLayer.RenderMaterialIndex;
+			return GenerateNodeInfo(CurrentLayer, ParentNode.bIsInstanceDefinition, MaterialIndex, CurrentLayer, new[] { CurrentLayer.Index });
+		}
+
+		private string BuildLayerString(Layer CurrentLayer, DatasmithActorInfo ParentNode)
+		{
+			return BuildLayerString(DatasmithRhinoUniqueNameGenerator.GetTargetName(CurrentLayer), ParentNode);
+		}
+
+		private string BuildLayerString(string LayerName, DatasmithActorInfo ParentNode)
+		{
+			string CurrentLayerName = LayerName.Replace(',', '_');
+
+			if (!ParentNode.bIsRoot)
+			{
+				Layer ParentLayer = ParentNode.RhinoCommonObject as Layer;
+				// If ParentLayer is null, then the parent node is a dummy document layer and we are in a worksession.
+				int ParentLayerIndex = ParentLayer == null ? -1 : ParentLayer.Index;
+
+				if (LayerIndexToLayerString.TryGetValue(ParentLayerIndex, out string ParentLayerString))
+				{
+					return string.Format("{0}_{1}", ParentLayerString, CurrentLayerName);
+				}
+			}
+
+			return CurrentLayerName;
+		}
+
+		private void RecursivelyParseObjectInstance(RhinoObject[] InObjects, DatasmithActorInfo ParentNode)
+		{
+			foreach (RhinoObject CurrentObject in InObjects)
+			{
+				if (TryParseObjectInstance(CurrentObject, ParentNode, out DatasmithActorInfo DefinitionRootNode) is DatasmithActorInfo ObjectNodeInfo)
+				{
+					ParentNode.AddChild(ObjectNodeInfo);
+
+					if (DefinitionRootNode != null)
+					{
+						InstanciateDefinition(ObjectNodeInfo, DefinitionRootNode);
+					}
+				}
+			}
+		}
+
+		private void RegisterActorsToContext(IEnumerable<DatasmithActorInfo> ActorInfos)
+		{
+			foreach (DatasmithActorInfo CurrentInfo in ActorInfos)
+			{
+				if (CurrentInfo.DirectLinkStatus != DirectLinkSynchronizationStatus.PendingDeletion)
+				{
+					// Make sure to update the material usage mapping, this is required to quickly find affected actors on material changes.
+					UpdateMaterialElementMapping(CurrentInfo, CurrentInfo.MaterialIndex);
+					
+					RhinoObject CurrentRhinoObject = CurrentInfo.RhinoCommonObject as RhinoObject;
+					bool bIsNotInstancedObject = CurrentRhinoObject == null || CurrentInfo.DefinitionNode == null || (CurrentRhinoObject.ObjectType == ObjectType.InstanceReference && !CurrentInfo.bIsInstanceDefinition);
+					
+					if (bIsNotInstancedObject)
+					{
+						ObjectIdToHierarchyActorNodeDictionary[CurrentInfo.RhinoObjectId] = CurrentInfo;
+
+						if (CurrentRhinoObject != null)
+						{
+							AddObjectMaterialReference(CurrentRhinoObject, CurrentInfo.MaterialIndex);
+						}
+					}
+				}
+			}
+		}
+
+		private DatasmithActorInfo TryParseObjectInstance(RhinoObject InObject, DatasmithActorInfo ParentNode, out DatasmithActorInfo DefinitionRootNode)
+		{
+			DefinitionRootNode = null;
+			if (InObject == null
+				|| IsObjectIgnoredBySelection(InObject, ParentNode)
+				|| IsUnsupportedObject(InObject))
+			{
+				// Skip the object.
+				return null;
+			}
+
+			if (InObject.ObjectType == ObjectType.InstanceReference)
+			{
+				InstanceObject CurrentInstance = InObject as InstanceObject;
+				DefinitionRootNode = GetOrCreateDefinitionRootNode(CurrentInstance.InstanceDefinition);
+
+				if (DefinitionRootNode.Children.Count == 0)
+				{
+					// Don't instantiate empty definitions.
+					return null;
+				}
+			}
+
+			int MaterialIndex = GetObjectMaterialIndex(InObject, ParentNode);
+			Layer VisibilityLayer;
+			IEnumerable<int> LayerIndices;
+			if (ParentNode.bIsRoot && ParentNode.bIsInstanceDefinition)
+			{
+				VisibilityLayer = RhinoDocument.Layers.FindIndex(InObject.Attributes.LayerIndex);
+				// The objects inside a Block definitions may be defined in a different layer than the one we are currently in.
+				LayerIndices = GetOrCreateLayerIndexHierarchy(InObject.Attributes.LayerIndex);
+			}
+			else
+			{
+				VisibilityLayer = ParentNode.VisibilityLayer;
+				LayerIndices = null;
+			}
+
+			return GenerateNodeInfo(InObject, ParentNode.bIsInstanceDefinition, MaterialIndex, VisibilityLayer, LayerIndices);
+		}
+
+		private bool IsObjectIgnoredBySelection(RhinoObject InObject, DatasmithActorInfo ParentNode)
+		{
+			return !ParentNode.bIsInstanceDefinition
+					&& ExportOptions.bWriteSelectedObjectsOnly
+					&& InObject.IsSelected(/*checkSubObjects=*/true) == 0;
+		}
+
+		private bool IsUnsupportedObject(RhinoObject InObject)
+		{
+//Disabling obsolete warning as GetRenderPrimitiveList() is deprecated since Rhino5 but as of Rhino7 no alternative exists.
+#pragma warning disable CS0612
+			// Geometry objects without meshes are currently not supported, unless they are points.
+			return InObject.ComponentType == ModelComponentType.ModelGeometry 
+				&& !InObject.IsMeshable(MeshType.Render)
+				&& InObject.ObjectType != ObjectType.Point
+				&& (InObject.ObjectType != ObjectType.Curve || null == InObject.GetRenderPrimitiveList(ActiveViewportInfo, true));
+#pragma warning restore CS0612
+		}
+
+		private void InstanciateDefinition(DatasmithActorInfo ParentNode, DatasmithActorInfo DefinitionNode)
+		{
+			ParentNode.SetDefinitionNode(DefinitionNode);
+
+			foreach (DatasmithActorInfo DefinitionChildNode in DefinitionNode.Children)
+			{
+				AddChildInstance(ParentNode, DefinitionChildNode);
+			}
+		}
+
+		private DatasmithActorInfo AddChildInstance(DatasmithActorInfo ParentNode, DatasmithActorInfo ChildDefinitionNode)
+		{
+			int MaterialIndex = GetObjectMaterialIndex(ChildDefinitionNode.RhinoCommonObject as RhinoObject, ParentNode);
+			DatasmithActorInfo InstanceChildNode = GenerateInstanceNodeInfo(ParentNode, ChildDefinitionNode, MaterialIndex, ParentNode.VisibilityLayer);
+			ParentNode.AddChild(InstanceChildNode);
+
+			InstanciateDefinition(InstanceChildNode, ChildDefinitionNode);
+
+			return InstanceChildNode;
+		}
+
+		private DatasmithActorInfo GetOrCreateDefinitionRootNode(InstanceDefinition InInstanceDefinition)
+		{
+			DatasmithActorInfo InstanceRootNode;
+
+			//If a hierarchy node does not exist for this instance definition, create one.
+			if (!InstanceDefinitionHierarchyNodeDictionary.TryGetValue(InInstanceDefinition, out InstanceRootNode))
+			{
+				InstanceRootNode = ParseInstanceDefinition(InInstanceDefinition);
+
+				RegisterActorsToContext(InstanceRootNode.GetEnumerator());
+				InstanceDefinitionHierarchyNodeDictionary.Add(InInstanceDefinition, InstanceRootNode);
+			}
+
+			return InstanceRootNode;
+		}
+
+		private DatasmithActorInfo ParseInstanceDefinition(InstanceDefinition InInstanceDefinition)
+		{
+			const bool bIsInstanceDefinition = true;
+			const int MaterialIndex = -1;
+			const Layer NullVisibilityLayer = null;
+
+			DatasmithActorInfo InstanceRootNode = GenerateNodeInfo(InInstanceDefinition, bIsInstanceDefinition, MaterialIndex, NullVisibilityLayer);
+
+			RhinoObject[] InstanceObjects = InInstanceDefinition.GetObjects();
+			RecursivelyParseObjectInstance(InstanceObjects, InstanceRootNode);
+
+			return InstanceRootNode;
+		}
+
+		public void UpdateDefinitionNode(InstanceDefinition InInstanceDefinition)
+		{
+			if (InstanceDefinitionHierarchyNodeDictionary.TryGetValue(InInstanceDefinition, out DatasmithActorInfo DefinitionRootNode))
+			{
+				if (ParseInstanceDefinition(InInstanceDefinition) is DatasmithActorInfo DiffActorInfo)
+				{
+					// Apply changes on InstanceDefinition Actors.
+					{
+						DefinitionRootNode.ApplyDiffs(DiffActorInfo);
+
+						// To avoid costly O(n^2) lookups, create HashSet and Dictionary to accelerate the search.
+						HashSet<DatasmithActorInfo> DeletedChildren = new HashSet<DatasmithActorInfo>(DefinitionRootNode.Children);
+						Dictionary<string, DatasmithActorInfo> ChildNameToIndex = new Dictionary<string, DatasmithActorInfo>(DefinitionRootNode.Children.Count);
+						foreach (DatasmithActorInfo CurrentChild in DefinitionRootNode.Children)
+						{
+							ChildNameToIndex.Add(CurrentChild.Name, CurrentChild);
+						}
+
+						foreach (DatasmithActorInfo OtherChildInfo in DiffActorInfo.Children)
+						{
+							if (ChildNameToIndex.TryGetValue(OtherChildInfo.Name, out DatasmithActorInfo CurrentChild))
+							{
+								DeletedChildren.Remove(CurrentChild);
+								CurrentChild.ApplyDiffs(OtherChildInfo);
+							}
+							else
+							{
+								//A new child is present, we must add it.
+								//We're not calling Reparent() here because that would change the size of the array we are iterating upon.
+								DefinitionRootNode.AddChild(OtherChildInfo);
+
+								RecursivelyAddChildInstance(DefinitionRootNode, OtherChildInfo);
+							}
+						}
+
+						foreach (DatasmithActorInfo DeletedChild in DeletedChildren)
+						{
+							DeletedChild.ApplyDeletedStatus();
+						}
+
+						RegisterActorsToContext(DefinitionRootNode.GetEnumerator());
+					}
+
+					// Apply changes on InstanceDefinition Meshes
+					{
+						List<RhinoObject> ObjectsNeedingMeshParsing = new List<RhinoObject>();
+
+						foreach (DatasmithActorInfo ActorInfo in DefinitionRootNode.GetEnumerator())
+						{
+							if (ActorInfo.DirectLinkStatus == DirectLinkSynchronizationStatus.Created
+								|| ActorInfo.DirectLinkStatus == DirectLinkSynchronizationStatus.Modified)
+							{
+								if (ActorInfo.RhinoCommonObject is RhinoObject ModifiedObject)
+								{
+									ObjectsNeedingMeshParsing.Add(ModifiedObject);
+								}
+							}
+							else if (ActorInfo.DirectLinkStatus == DirectLinkSynchronizationStatus.PendingDeletion)
+							{
+								if (ObjectIdToMeshInfoDictionary.TryGetValue(ActorInfo.RhinoObjectId, out DatasmithMeshInfo MeshInfo))
+								{
+									MeshInfo.ApplyDeletedStatus();
+								}
+							}
+						}
+
+						// Add newly created meshes.
+						AddOrModifyMeshesFromRhinoObjects(ObjectsNeedingMeshParsing);
+					}
+
+					bIsDirty = true;
+				}
+			}
+		}
+
+
+		private void RecursivelyAddChildInstance(DatasmithActorInfo DefinitionParent, DatasmithActorInfo ChildDefinition)
+		{
+			foreach (DatasmithActorInfo ActorInstance in DefinitionParent.InstanceNodes)
+			{
+				if (ActorInstance.DirectLinkStatus != DirectLinkSynchronizationStatus.PendingDeletion)
+				{
+					DatasmithActorInfo ChildInstance = AddChildInstance(ActorInstance, ChildDefinition);
+					RecursivelyAddChildInstance(ActorInstance, ChildInstance);
+				}
+			}
+		}
+
+		/// <summary>
+		/// Creates a new hierarchy node info by using the DefinitionNodeInfo as a base, and using InstanceParentNodeInfo for creating a unique name. Used when instancing block definitions.
+		/// </summary>
+		/// <param name="InstanceParentNodeInfo"></param>
+		/// <param name="DefinitionNodeInfo"></param>
+		/// <param name="MaterialIndex"></param>
+		/// <returns></returns>
+		private DatasmithActorInfo GenerateInstanceNodeInfo(DatasmithActorInfo InstanceParentNodeInfo, DatasmithActorInfo DefinitionNodeInfo, int MaterialIndex, Layer VisibilityLayer)
+		{
+			string Name = string.Format("{0}_{1}", InstanceParentNodeInfo.Name, DefinitionNodeInfo.Name);
+			string BaseLabel = DefinitionNodeInfo.BaseLabel;
+			string UniqueLabel = ActorLabelGenerator.GenerateUniqueNameFromBaseName(DefinitionNodeInfo.UniqueLabel);
+			bool bOverrideMaterial = DefinitionNodeInfo.MaterialIndex != MaterialIndex || DefinitionNodeInfo.bOverrideMaterial;
+
+			return new DatasmithActorInfo(DefinitionNodeInfo.RhinoCommonObject as ModelComponent, Name, UniqueLabel, BaseLabel, MaterialIndex, bOverrideMaterial, VisibilityLayer);
+		}
+
+		/// <summary>
+		/// Creates a new hierarchy node info that is not represented by any ModelComponent, used to add an empty actor to the scene. 
+		/// </summary>
+		/// <param name="UniqueID"></param>
+		/// <param name="TargetLabel"></param>
+		/// <param name="MaterialIndex"></param>
+		/// <returns></returns>
+		private DatasmithActorInfo GenerateDummyNodeInfo(string UniqueID, string TargetLabel, int MaterialIndex, IEnumerable<int> LayerIndices)
+		{
+			string UniqueLabel = ActorLabelGenerator.GenerateUniqueNameFromBaseName(TargetLabel);
+			const ModelComponent NullModelComponent = null;
+			const bool bOverrideMaterial = false;
+			const Layer NullVisiblityLayer = null;
+
+			return new DatasmithActorInfo(NullModelComponent, UniqueID, UniqueLabel, TargetLabel, MaterialIndex, bOverrideMaterial, NullVisiblityLayer, LayerIndices);
+		}
+
+		/// <summary>
+		/// Creates a new hierarchy node info for a given Rhino Model Component, used to determine names and labels as well as linking.
+		/// </summary>
+		/// <param name="InModelComponent"></param>
+		/// <param name="ParentNode"></param>
+		/// <param name="MaterialIndex"></param>
+		/// <returns></returns>
+		private DatasmithActorInfo GenerateNodeInfo(ModelComponent InModelComponent, bool bIsInstanceDefinition, int MaterialIndex, Layer VisibilityLayer, IEnumerable<int> LayerIndices = null)
+		{
+			string Name = GetModelComponentName(InModelComponent);
+			string BaseLabel = DatasmithRhinoUniqueNameGenerator.GetTargetName(InModelComponent);
+			string UniqueLabel = bIsInstanceDefinition
+				? BaseLabel
+				: ActorLabelGenerator.GenerateUniqueName(InModelComponent);
+			const bool bOverrideMaterial = false;
+
+			return new DatasmithActorInfo(InModelComponent, Name, UniqueLabel, BaseLabel, MaterialIndex, bOverrideMaterial, VisibilityLayer, LayerIndices);
+		}
+
+		/// <summary>
+		/// Gets the unique datasmith element name for the given ModelComponent object.
+		/// Some ModelComponent may have volatile IDs, it's important to use this function to ensure the name is based on some other attributes in those cases.
+		/// </summary>
+		/// <param name="InModelComponent"></param>
+		/// <returns></returns>
+		private string GetModelComponentName(ModelComponent InModelComponent)
+		{
+			Layer RhinoLayer = InModelComponent as Layer;
+
+			if (bIsInWorksession && RhinoLayer != null && RhinoLayer.IsReference && RhinoLayer.ParentLayerId == Guid.Empty)
+			{
+				// This is a dummy layer added by Rhino representing the linked document.
+				// Use the document absolute filepath as unique name to ensure consistency across worksessions.
+				string DocumentFilePath = GetRhinoDocumentPathFromLayerName(RhinoLayer.Name);
+				if (DocumentFilePath != null)
+				{
+					return DocumentFilePath;
+				}
+			}
+
+			return InModelComponent.Id.ToString();
+		}
+
+		private string GetRhinoDocumentPathFromLayerName(string LayerName)
+		{
+			if (WorksessionDocumentPaths != null)
+			{
+				foreach (string CurrentPath in WorksessionDocumentPaths)
+				{
+					if (CurrentPath.Contains(LayerName))
+					{
+						return CurrentPath;
+					}
+				}
+			}
+
+			return null;
+		}
+
+		private int GetObjectMaterialIndex(RhinoObject InRhinoObject, DatasmithActorInfo ParentNodeInfo)
+		{
+			ObjectMaterialSource MaterialSource = InRhinoObject.Attributes.MaterialSource;
+
+			if(MaterialSource == ObjectMaterialSource.MaterialFromParent)
+			{
+				// Special use-case for block instances with the source material set to "Parent".
+				// The material source in this case is actually the layer in which the definition objects are.
+				if(InRhinoObject is InstanceObject)
+				{
+					InstanceObject Instance = InRhinoObject as InstanceObject;
+					RhinoObject[] DefinitionObjects = Instance.InstanceDefinition.GetObjects();
+					if(DefinitionObjects != null && DefinitionObjects.Length > 0)
+					{
+						Layer ObjectLayer = RhinoDocument.Layers.FindIndex(InRhinoObject.Attributes.LayerIndex);
+						return ObjectLayer.RenderMaterialIndex;
+					}
+				}
+				else if(ParentNodeInfo == null)
+				{
+					MaterialSource = ObjectMaterialSource.MaterialFromLayer;
+				}
+			}
+
+			switch(MaterialSource)
+			{
+				case ObjectMaterialSource.MaterialFromLayer:
+					Layer ObjectLayer = RhinoDocument.Layers.FindIndex(InRhinoObject.Attributes.LayerIndex);
+					return ObjectLayer.RenderMaterialIndex;
+				case ObjectMaterialSource.MaterialFromParent:
+					return ParentNodeInfo.MaterialIndex;
+				case ObjectMaterialSource.MaterialFromObject:
+				default:
+					return InRhinoObject.Attributes.MaterialIndex;
+			}
+		}
+
+		public void ModifyMaterial(int MaterialIndex)
+		{
+			if (MaterialIndexToMaterialHashDictionary.TryGetValue(MaterialIndex, out string ExistingMaterialHash))
+			{
+				Material IndexedMaterial = MaterialIndex == -1
+					? Material.DefaultMaterial
+					: RhinoDocument.Materials.FindIndex(MaterialIndex);
+
+				string MaterialHash = DatasmithRhinoUtilities.GetMaterialHash(IndexedMaterial);
+				
+				if (ExistingMaterialHash != MaterialHash)
+				{
+					MaterialIndexToMaterialHashDictionary[MaterialIndex] = MaterialHash;
+					DatasmithMaterialInfo ModifiedMaterialInfo = GetOrCreateMaterialHashMapping(MaterialHash, IndexedMaterial);
+					ModifiedMaterialInfo.MaterialIndexes.Add(MaterialIndex);
+
+					SetModifiedFlagOnElementsUsingMaterial(MaterialIndex);
+
+					if (MaterialHashToMaterialInfo.TryGetValue(ExistingMaterialHash, out DatasmithMaterialInfo ExistingMaterialInfo))
+					{
+						ExistingMaterialInfo.MaterialIndexes.Remove(MaterialIndex);
+						
+						// If existing material is no longer used by any material index.
+						if (ExistingMaterialInfo.MaterialIndexes.Count == 0)
+						{
+							// Remove any texture that is no longer valid.
+							foreach (string TextureHash in ExistingMaterialInfo.TextureHashes)
+							{
+								if (TextureHashToTextureInfo.TryGetValue(TextureHash, out DatasmithTextureInfo TextureInfo)
+									&& TextureInfo.RhinoTexture.Disposed)
+								{
+									TextureInfo.ApplyDeletedStatus();
+								}
+							}
+
+							if (ModifiedMaterialInfo.DirectLinkStatus == DirectLinkSynchronizationStatus.Created)
+							{
+								// If the ModifiedMaterialInfo is new and the existing material is no longer referenced, 
+								// that means we are actually modifying the existing material in-place.
+								// Delete the ModifiedMateralInfo and update the ExistingMaterialInfo with its values.
+								MaterialHashToMaterialInfo[MaterialHash] = ExistingMaterialInfo;
+								ExistingMaterialInfo.ApplyDiffs(ModifiedMaterialInfo);
+								MaterialHashToMaterialInfo.Remove(ExistingMaterialHash);
+							}
+							else
+							{
+								// If the material index is no longer in used, we can delete the Material.
+								ExistingMaterialInfo.ApplyDeletedStatus();
+							}
+						}
+					}
+
+					bIsDirty = true;
+				}
+				else
+				{
+					//Material didn't change, nothing to do.
+					return;
+				}
+			}
+		}
+
+		private void SetModifiedFlagOnElementsUsingMaterial(int MaterialIndex)
+		{
+			if (MaterialIndexToElementInfosMap.TryGetValue(MaterialIndex, out HashSet<DatasmithInfoBase> ElementInfos))
+			{
+				foreach (DatasmithInfoBase ElementInfo in ElementInfos)
+				{
+					ElementInfo.ApplyModifiedStatus();
+				}
+			}
+		}
+
+		public void DeleteMaterial(int MaterialIndex)
+		{
+			if (MaterialIndexToMaterialHashDictionary.TryGetValue(MaterialIndex, out string MaterialHash))
+			{
+				MaterialIndexToMaterialHashDictionary.Remove(MaterialIndex);
+
+				if (MaterialHashToMaterialInfo.TryGetValue(MaterialHash, out DatasmithMaterialInfo MaterialInfo))
+				{
+					MaterialInfo.MaterialIndexes.Remove(MaterialIndex);
+
+					if (MaterialInfo.MaterialIndexes.Count == 0)
+					{
+						if (MaterialInfo.DirectLinkStatus == DirectLinkSynchronizationStatus.Created)
+						{
+							MaterialHashToMaterialInfo.Remove(MaterialHash);
+						}
+						else
+						{
+							MaterialInfo.ApplyDeletedStatus();
+						}
+					}
+				}
+
+				bIsDirty = true;
+			}
+		}
+
+		private void AddObjectMaterialReference(RhinoObject InObject, int MaterialIndex)
+		{
+			if (InObject.ObjectType == ObjectType.Brep)
+			{
+				BrepObject InBrepObject = InObject as BrepObject;
+				if (InBrepObject.HasSubobjectMaterials)
+				{
+					foreach (ComponentIndex CurrentIndex in InBrepObject.SubobjectMaterialComponents)
+					{
+						if (InBrepObject.GetMaterial(CurrentIndex) is Material ComponentMaterial)
+						{
+							AddMaterialIndexMapping(ComponentMaterial.Index);
+						}
+					}
+				}
+			}
+
+			AddMaterialIndexMapping(MaterialIndex);
+		}
+
+		private void AddMaterialIndexMapping(int MaterialIndex)
+		{
+			if(!MaterialIndexToMaterialHashDictionary.ContainsKey(MaterialIndex))
+			{
+				Material IndexedMaterial = MaterialIndex == -1
+					? Material.DefaultMaterial
+					: RhinoDocument.Materials.FindIndex(MaterialIndex);
+
+				string MaterialHash = DatasmithRhinoUtilities.GetMaterialHash(IndexedMaterial);
+				MaterialIndexToMaterialHashDictionary.Add(MaterialIndex, MaterialHash);
+
+				DatasmithMaterialInfo MaterialInfo = GetOrCreateMaterialHashMapping(MaterialHash, IndexedMaterial);
+				MaterialInfo.MaterialIndexes.Add(MaterialIndex);
+			}
+		}
+
+		private DatasmithMaterialInfo GetOrCreateMaterialHashMapping(string MaterialHash, Material RhinoMaterial)
+		{
+			DatasmithMaterialInfo Result;
+
+			if (!MaterialHashToMaterialInfo.TryGetValue(MaterialHash, out Result))
+			{
+				string BaseLabel = DatasmithRhinoUniqueNameGenerator.GetTargetName(RhinoMaterial);
+				string UniqueLabel = MaterialLabelGenerator.GenerateUniqueName(RhinoMaterial);
+				string Name = FDatasmithFacadeElement.GetStringHash(UniqueLabel);
+
+				Texture[] MaterialTextures = RhinoMaterial.GetTextures();
+				List<string> TextureHashes = new List<string>(MaterialTextures.Length);
+				for (int TextureIndex = 0; TextureIndex < MaterialTextures.Length; ++TextureIndex)
+				{
+					Texture RhinoTexture = MaterialTextures[TextureIndex];
+					if (RhinoTexture != null && DatasmithRhinoUtilities.IsTextureSupported(RhinoTexture))
+					{
+						string TextureHash = DatasmithRhinoUtilities.GetTextureHash(RhinoTexture);
+						TextureHashes.Add(TextureHash);
+						UpdateTextureHashMapping(TextureHash, RhinoTexture);
+					}
+				}
+
+				Result = new DatasmithMaterialInfo(RhinoMaterial, Name, UniqueLabel, BaseLabel, TextureHashes);
+				MaterialHashToMaterialInfo.Add(MaterialHash, Result);
+			}
+			else if (Result.DirectLinkStatus == DirectLinkSynchronizationStatus.PendingDeletion)
+			{
+				Result.RestorePreviousDirectLinkStatus();
+			}
+
+			return Result;
+		}
+
+		private void UpdateTextureHashMapping(string TextureHash, Texture RhinoTexture)
+		{
+			Guid TextureId = RhinoTexture.Id;
+
+			if (TextureHashToTextureInfo.TryGetValue(TextureHash, out DatasmithTextureInfo ExistingTextureInfo))
+			{
+				ExistingTextureInfo.TextureIds.Add(TextureId);
+				if (ExistingTextureInfo.DirectLinkStatus == DirectLinkSynchronizationStatus.PendingDeletion)
+				{
+					ExistingTextureInfo.RestorePreviousDirectLinkStatus();
+				}
+			}
+			else
+			{
+				if (DatasmithRhinoUtilities.GetRhinoTextureNameAndPath(RhinoTexture, out string TextureName, out string TexturePath))
+				{
+					TextureName = TextureLabelGenerator.GenerateUniqueNameFromBaseName(TextureName);
+					DatasmithTextureInfo TextureInfo = new DatasmithTextureInfo(RhinoTexture, TextureName, TexturePath);
+					TextureInfo.TextureIds.Add(TextureId);
+					TextureHashToTextureInfo.Add(TextureHash, TextureInfo);
+				}
+			}
+
+			// Rhino 'Texture' instances are not immutable and the Texture file they point to may be updated.
+			// We need to detect when our DatasmithTextureInfo are no longer used by any Texture in Rhino and flag them for deletion.
+			if (TextureIdToTextureHash.TryGetValue(TextureId, out string PreviousHash))
+			{
+				if (TextureHash != PreviousHash
+					&& TextureHashToTextureInfo.TryGetValue(PreviousHash, out DatasmithTextureInfo TextureInfo))
+				{
+					TextureInfo.TextureIds.Remove(TextureId);
+					if (TextureInfo.TextureIds.Count == 0)
+					{
+						TextureInfo.ApplyDeletedStatus();
+					}
+				}
+			}
+
+			TextureIdToTextureHash[TextureId] = TextureHash;
+		}
+
+		public void UpdateGroups(GroupTableEventType UpdateType, Group UpdatedGroup)
+		{
+			switch (UpdateType)
+			{
+				case GroupTableEventType.Added:
+				case GroupTableEventType.Undeleted:
+					GroupIndexToName.Add(UpdatedGroup.Index, GetGroupFormatedName(UpdatedGroup));
+					break;
+				case GroupTableEventType.Deleted:
+					GroupIndexToName.Remove(UpdatedGroup.Index);
+					break;
+				case GroupTableEventType.Modified:
+					GroupIndexToName[UpdatedGroup.Index] = GetGroupFormatedName(UpdatedGroup);
+					break;
+				default:
+					break;
+			}
+
+			bIsDirty = true;
+		}
+
+		private void ParseGroupNames()
+		{
+			foreach (Group CurrentGroup in RhinoDocument.Groups)
+			{
+				int GroupIndex = CurrentGroup.Index;
+				string GroupName = GetGroupFormatedName(CurrentGroup);
+				GroupIndexToName.Add(GroupIndex, GroupName);
+			}
+		}
+
+		/// <summary>
+		/// Returned the exported group name tag from Group passed in parameter.
+		/// </summary>
+		/// <param name="InGroup"></param>
+		/// <returns></returns>
+		private string GetGroupFormatedName(Group InGroup)
+		{
+			return string.Format("Rhino.GroupName: {0}", String.IsNullOrEmpty(InGroup.Name)
+					? string.Format("Group{0}", InGroup.Index)
+					: InGroup.Name);
+		}
+
+		private HashSet<int> GetOrCreateLayerIndexHierarchy(int ChildLayerIndex)
+		{
+			HashSet<int> LayerIndexHierarchy; 
+			if (!LayerIndexToLayerIndexHierarchy.TryGetValue(ChildLayerIndex, out LayerIndexHierarchy))
+			{
+				LayerIndexHierarchy = new HashSet<int>();
+				Layer CurrentLayer = RhinoDocument.Layers.FindIndex(ChildLayerIndex);
+
+				if (CurrentLayer != null)
+				{
+					LayerIndexHierarchy.Add(CurrentLayer.Index);
+
+					//Add all parent layers to the hierarchy set
+					while (CurrentLayer.ParentLayerId != Guid.Empty)
+					{
+						CurrentLayer = RhinoDocument.Layers.FindId(CurrentLayer.ParentLayerId);
+						LayerIndexHierarchy.Add(CurrentLayer.Index);
+					}
+				}
+
+				//The above search is costly, to avoid doing it for every exported object, cache the result.
+				LayerIndexToLayerIndexHierarchy.Add(ChildLayerIndex, LayerIndexHierarchy);
+			}
+
+			return LayerIndexHierarchy;
+		}
+
+		public void ParseRhinoCameras()
+		{
+			HashSet<string> ViewsToRemove = new HashSet<string>(NamedViewNameToCameraInfo.Keys);
+			bool bHasChanged = false;
+
+			foreach (ViewInfo NamedView in RhinoDocument.NamedViews)
+			{
+				string NamedViewHash = DatasmithRhinoUtilities.GetNamedViewHash(NamedView);
+				
+				if (NamedViewNameToCameraInfo.TryGetValue(NamedView.Name, out DatasmithCameraActorInfo ExistingCameraInfo))
+				{
+					// Named view exists in DatasmithScene, check if it has changed.
+					ViewsToRemove.Remove(NamedView.Name);
+					
+					if (NamedViewHash != ExistingCameraInfo.NamedViewHash)
+					{
+						DatasmithCameraActorInfo DiffInfo = GenerateActorCameraInfo(NamedView, NamedViewHash);
+						ExistingCameraInfo.ApplyDiffs(DiffInfo);
+						bHasChanged = true;
+					}
+				}
+				else
+				{
+					// Named view does not exists in DatasmithScene, create a new one and add it to scene root actor.
+					DatasmithCameraActorInfo CameraInfo = GenerateActorCameraInfo(NamedView, NamedViewHash);
+
+					NamedViewNameToCameraInfo.Add(NamedView.Name, CameraInfo);
+					SceneRoot.AddChild(CameraInfo);
+					bHasChanged = true;
+				}
+			}
+
+			// Mark for removal all the named views that were not
+			foreach (string ViewToRemoveName in ViewsToRemove)
+			{
+				if (NamedViewNameToCameraInfo.TryGetValue(ViewToRemoveName, out DatasmithCameraActorInfo ViewToRemove))
+				{
+					ViewToRemove.ApplyDeletedStatus();
+				}
+				bHasChanged = true;
+			}
+
+			bIsDirty |= bHasChanged;
+		}
+
+		private DatasmithCameraActorInfo GenerateActorCameraInfo(ViewInfo NamedView, string NamedViewHash)
+		{
+			string BaseLabel = DatasmithRhinoUniqueNameGenerator.GetTargetName(NamedView);
+			string UniqueLabel = ActorLabelGenerator.GenerateUniqueNameFromBaseName(BaseLabel);
+			//We can't use the NamedView.Viewport.Id, as it may be unset.
+			string Name = UniqueLabel;
+
+			return new DatasmithCameraActorInfo(NamedView.Viewport, Name, UniqueLabel, BaseLabel, NamedViewHash);
+		}
+
+		private void ParseAllRhinoMeshes()
+		{
+			HashSet<RhinoObject> DocObjects = CollectExportedRhinoObjects();
+			AddOrModifyMeshesFromRhinoObjects(DocObjects);
+		}
+
+		private HashSet<RhinoObject> CollectExportedRhinoObjects()
+		{
+			ObjectEnumeratorSettings Settings = new ObjectEnumeratorSettings();
+			Settings.HiddenObjects = false;
+			//First get all non-instance directly in the scene
+			Settings.ObjectTypeFilter = ObjectType.AnyObject ^ (ObjectType.InstanceDefinition | ObjectType.InstanceReference);
+			Settings.ReferenceObjects = true;
+
+			//Calling GetObjectList instead of directly iterating through RhinoDocument.Objects as it seems that the ObjectTable may sometimes contain uninitialized RhinoObjects.
+			HashSet<RhinoObject> ExportedObjects = new HashSet<RhinoObject>(RhinoDocument.Objects.GetObjectList(Settings));
+
+			// Collect objects that are referenced inside instance definitions(blocks).
+			// We do this because if we were to call RhinoObject.MeshObjects() on an Instance object it would create a single mesh merging all the instance's children.
+			foreach (var CurrentInstanceHierarchy in InstanceDefinitionHierarchyNodeDictionary.Values)
+			{
+				CollectExportedRhinoObjectsFromInstanceDefinition(ExportedObjects, CurrentInstanceHierarchy);
+			}
+
+			return ExportedObjects;
+		}
+
+		private void CollectExportedRhinoObjectsFromInstanceDefinition(HashSet<RhinoObject> CollectedObjects, DatasmithActorInfo InstanceDefinitionHierarchy)
+		{
+			foreach (DatasmithActorInfo HierarchyNode in InstanceDefinitionHierarchy.GetEnumerator())
+			{
+				if (!HierarchyNode.bIsRoot && HierarchyNode.RhinoCommonObject is RhinoObject NodeObject && !(HierarchyNode.RhinoCommonObject is InstanceObject))
+				{
+					CollectedObjects.Add(NodeObject);
+				}
+			}
+		}
+
+		private void AddOrModifyMeshesFromRhinoObjects(IEnumerable<RhinoObject> RhinoObjects)
+		{
+			// Make sure all render meshes are generated before attempting to export them.
+			GenerateMissingRenderMeshes(RhinoObjects);
+
+			foreach (RhinoObject CurrentObject in RhinoObjects)
+			{
+				if (TryGenerateMeshInfoFromRhinoObjects(CurrentObject) is DatasmithMeshInfo MeshInfo)
+				{
+					if (ObjectIdToMeshInfoDictionary.TryGetValue(CurrentObject.Id, out DatasmithMeshInfo ExistingMeshInfo))
+					{
+						if (ExistingMeshInfo.DirectLinkStatus == DirectLinkSynchronizationStatus.PendingDeletion)
+						{
+							ExistingMeshInfo.RestorePreviousDirectLinkStatus();
+						}
+
+						List<int> PreviousMaterialIndices = ExistingMeshInfo.MaterialIndices;
+						ExistingMeshInfo.ApplyDiffs(MeshInfo);
+						UpdateMaterialElementMapping(ExistingMeshInfo, MeshInfo.MaterialIndices, PreviousMaterialIndices);
+					}
+					else
+					{
+						ObjectIdToMeshInfoDictionary.Add(CurrentObject.Id, MeshInfo);
+						UpdateMaterialElementMapping(MeshInfo, MeshInfo.MaterialIndices);
+					}
+
+					UpdateTextureMappingCache(CurrentObject);
+				}
+			}
+		}
+
+		private void GenerateMissingRenderMeshes(IEnumerable<RhinoObject> RhinoObjects)
+		{
+			foreach (var CurrentObject in RhinoObjects)
+			{
+				Mesh[] RenderMeshes = CurrentObject.GetMeshes(MeshType.Render);
+				
+				if (RenderMeshes?.Length == 0 && CurrentObject.IsMeshable(MeshType.Render))
+				{
+					MeshingParameters MeshParams = CurrentObject.GetRenderMeshParameters();
+					CurrentObject.CreateMeshes(MeshType.Render, MeshParams, false);
+				}
+			}
+		}
+
+		private DatasmithMeshInfo TryGenerateMeshInfoFromRhinoObjects(RhinoObject InRhinoObject)
+		{
+			DatasmithMeshInfo MeshInfo = null;
+			Mesh[] RenderMeshes = null;
+
+			if (ActiveViewportInfo != null)
+			{
+				//Disabling obsolete warning as GetRenderPrimitiveList() is deprecated since Rhino5 but as of Rhino7 no alternative exists.
+#pragma warning disable CS0612
+				RenderMeshes = InRhinoObject.GetRenderPrimitiveList(ActiveViewportInfo, false)?.ToMeshArray();
+#pragma warning restore CS0612
+			}
+			if (RenderMeshes == null)
+			{
+				RenderMeshes = InRhinoObject.GetMeshes(MeshType.Render);
+			}
+
+			if (RenderMeshes != null && RenderMeshes.Length > 0)
+			{
+				Dictionary<Mesh, ObjectAttributes> MeshAttributePairs = GetMeshAttributePairs(InRhinoObject, RenderMeshes);
+				List<DatasmithTextureMappingData> InTextureMappings = GetObjectTextureMappings(InRhinoObject);
+
+				if (MeshAttributePairs.Count > 0)
+				{
+					MeshInfo = GenerateMeshInfo(InRhinoObject.Id, MeshAttributePairs, InTextureMappings);
+				}
+			}
+
+			return MeshInfo;
+		}
+
+		private Dictionary<Mesh, ObjectAttributes> GetMeshAttributePairs(RhinoObject InRhinoObject, IList<Mesh> RenderMeshes)
+		{
+			Dictionary<Mesh, ObjectAttributes> MeshAttributePairs = new Dictionary<Mesh, ObjectAttributes>(RenderMeshes.Count);
+
+			BrepObject CurrentBrep = (InRhinoObject as BrepObject);
+			if (CurrentBrep != null && CurrentBrep.HasSubobjectMaterials)
+			{
+				RhinoObject[] SubObjects = CurrentBrep.GetSubObjects();
+
+				for (int Index = 0, MaxLength = Math.Min(RenderMeshes.Count, SubObjects.Length); Index < MaxLength; ++Index)
+				{
+					if (RenderMeshes[Index] != null)
+					{
+						MeshAttributePairs[RenderMeshes[Index]] = SubObjects[Index].Attributes;
+					}
+				}
+			}
+			else
+			{
+				for (int RenderMeshIndex = 0; RenderMeshIndex < RenderMeshes.Count; ++RenderMeshIndex)
+				{
+					if (RenderMeshes[RenderMeshIndex] != null)
+					{
+						MeshAttributePairs[RenderMeshes[RenderMeshIndex]] = InRhinoObject.Attributes;
+					}
+				}
+			}
+
+			return MeshAttributePairs;
+		}
+
+		private List<DatasmithTextureMappingData> GetObjectTextureMappings(RhinoObject InRhinoObject)
+		{
+			if (InRhinoObject.GetTextureChannels() is int[] TextureChannels)
+			{
+				List<DatasmithTextureMappingData> TextureMappings = new List<DatasmithTextureMappingData>(TextureChannels.Length);
+				foreach (int TextureChannelID in TextureChannels)
+				{
+					if (InRhinoObject.GetTextureMapping(TextureChannelID, out Transform ObjectTransform) is Rhino.Render.TextureMapping ObjectTextureMapping)
+					{
+						TextureMappings.Add(new DatasmithTextureMappingData(TextureChannelID, ObjectTextureMapping, ObjectTransform));
+					}
+				}
+				
+				return TextureMappings;
+			}
+
+			return new List<DatasmithTextureMappingData>();
+		}
+
+		private DatasmithMeshInfo GenerateMeshInfo(Guid ObjectID, IReadOnlyDictionary<Mesh, ObjectAttributes> MeshAttributePairs, List<DatasmithTextureMappingData> TextureMappingList)
+		{
+			DatasmithMeshInfo MeshInfo = null;
+
+			if (ObjectIdToHierarchyActorNodeDictionary.TryGetValue(ObjectID, out DatasmithActorInfo HierarchyActorNode))
+			{
+				Transform OffsetTransform = Transform.Translation(DatasmithRhinoUtilities.CenterMeshesOnPivot(MeshAttributePairs.Keys));
+				List<int> MaterialIndices = GetMeshesMaterialIndices(MeshAttributePairs, HierarchyActorNode);
+
+				string Name = FDatasmithFacadeElement.GetStringHash("M:" + HierarchyActorNode.Name);
+				string UniqueLabel = HierarchyActorNode.UniqueLabel;
+				string BaseLabel = HierarchyActorNode.BaseLabel;
+
+				MeshInfo = new DatasmithMeshInfo(MeshAttributePairs.Keys, OffsetTransform, MaterialIndices, TextureMappingList, Name, UniqueLabel, BaseLabel);
+			}
+			else
+			{
+				RhinoApp.WriteLine("Could not find the corresponding hierarchy node for the object ID: {0}", ObjectID);
+			}
+
+			return MeshInfo;
+		}
+
+		private List<int> GetMeshesMaterialIndices(IReadOnlyDictionary<Mesh, ObjectAttributes> MeshAttributePairs, DatasmithActorInfo ActorInfo)
+		{
+			List<int> MaterialIndices = new List<int>(MeshAttributePairs.Count);
+			foreach (ObjectAttributes CurrentAttributes in MeshAttributePairs.Values)
+			{
+				int MaterialIndex = GetMaterialIndexFromAttributes(ActorInfo, CurrentAttributes);
+				MaterialIndices.Add(MaterialIndex);
+			}
+
+			return MaterialIndices;
+		}
+
+		private int GetMaterialIndexFromAttributes(DatasmithActorInfo HierarchyActorNode, ObjectAttributes Attributes)
+		{
+			if (Attributes.MaterialSource == ObjectMaterialSource.MaterialFromObject)
+			{
+				return Attributes.MaterialIndex;
+			}
+
+			//The parent material and layer material is baked into the HierarchyActorNode.
+			return HierarchyActorNode.MaterialIndex;
+		}
+
+
+		/// <summary>
+		/// Update the material index mapping to associated Element info. We need this to dirty the element when the material assignment changes.
+		/// </summary>
+		/// <param name="ElementInfo">The element referencing the material</param>
+		/// <param name="NewMaterialIndices">Material indices referenced by the element</param>
+		/// <param name="PreviousMaterialIndices">Previous material indices referenced by the element, can pass null if this is a new element</param>
+		private void UpdateMaterialElementMapping(DatasmithInfoBase ElementInfo, List<int> NewMaterialIndices, List<int> PreviousMaterialIndices = null)
+		{
+			if (PreviousMaterialIndices != null)
+			{
+				if (NewMaterialIndices.Count == PreviousMaterialIndices.Count)
+				{
+					bool bSameMaterialAssignment = true;
+					for (int MaterialIndex = 0; MaterialIndex < NewMaterialIndices.Count; ++MaterialIndex)
+					{
+						if (NewMaterialIndices[MaterialIndex] != PreviousMaterialIndices[MaterialIndex])
+						{
+							bSameMaterialAssignment = false;
+							break;
+						}
+					}
+
+					//If the material assignments are the same, there is nothing to do.
+					if (bSameMaterialAssignment)
+					{
+						return;
+					}
+				}
+
+				// Remove previous material mapping.
+				foreach (int MaterialIndex in PreviousMaterialIndices.Distinct())
+				{
+					if (MaterialIndexToElementInfosMap.TryGetValue(MaterialIndex, out HashSet<DatasmithInfoBase> MeshInfos))
+					{
+						MeshInfos.Remove(ElementInfo);
+					}
+				}
+			}
+
+			// Add new material mapping
+			foreach (int MaterialIndex in NewMaterialIndices.Distinct())
+			{
+				HashSet<DatasmithInfoBase> ElementInfos;
+				if (!MaterialIndexToElementInfosMap.TryGetValue(MaterialIndex, out ElementInfos))
+				{
+					ElementInfos = new HashSet<DatasmithInfoBase>();
+					MaterialIndexToElementInfosMap.Add(MaterialIndex, ElementInfos);
+				}
+
+				ElementInfos.Add(ElementInfo);
+			}
+		}
+
+		/// <summary>
+		/// Update the material index mapping to associated Element info. We need this to dirty the element when the material assignment changes.
+		/// </summary>
+		/// <param name="ElementInfo">The element referencing the material</param>
+		/// <param name="NewMaterialIndex">Material index referenced by the element</param>
+		/// <param name="PreviousMaterialIndex">Previous material index referenced by the element, can pass null if this is a new element</param>
+		private void UpdateMaterialElementMapping(DatasmithInfoBase ElementInfo, int NewMaterialIndex, int? PreviousMaterialIndex = null)
+		{
+			if (PreviousMaterialIndex != null)
+			{
+				if (NewMaterialIndex == PreviousMaterialIndex.Value)
+				{
+					//If the material assignments are the same, there is nothing to do.
+					return;
+				}
+
+				// Remove previous material mapping.
+				if (MaterialIndexToElementInfosMap.TryGetValue(PreviousMaterialIndex.Value, out HashSet<DatasmithInfoBase> MeshInfos))
+				{
+					MeshInfos.Remove(ElementInfo);
+				}
+			}
+
+			// Add new material mapping
+			HashSet<DatasmithInfoBase> ElementInfos;
+			if (!MaterialIndexToElementInfosMap.TryGetValue(NewMaterialIndex, out ElementInfos))
+			{
+				ElementInfos = new HashSet<DatasmithInfoBase>();
+				MaterialIndexToElementInfosMap.Add(NewMaterialIndex, ElementInfos);
+			}
+
+			ElementInfos.Add(ElementInfo);
+		}
+
+		private void UpdateTextureMappingCache(RhinoObject InRhinoObject)
+		{
+			foreach (int ChannelId in InRhinoObject.GetTextureChannels())
+			{
+				Rhino.Render.TextureMapping Mapping = InRhinoObject.GetTextureMapping(ChannelId);
+				TextureMappindIdToRhinoObject[Mapping.Id] = InRhinoObject;
+			}
+		}
+	}
+}
\ No newline at end of file
diff --git a/Engine/Source/Programs/Enterprise/Datasmith/DatasmithRhinoExporter/Private/ExportContext/DatasmithTextureInfo.cs b/Engine/Source/Programs/Enterprise/Datasmith/DatasmithRhinoExporter/Private/ExportContext/DatasmithTextureInfo.cs
new file mode 100644
--- /dev/null
+++ b/Engine/Source/Programs/Enterprise/Datasmith/DatasmithRhinoExporter/Private/ExportContext/DatasmithTextureInfo.cs
@@ -0,0 +1,34 @@
+// Copyright Epic Games, Inc. All Rights Reserved.
+
+using DatasmithRhino.Utils;
+using Rhino.DocObjects;
+using Rhino.Geometry;
+using System;
+using System.Collections.Generic;
+
+namespace DatasmithRhino.ExportContext
+{
+	public class DatasmithTextureInfo : DatasmithInfoBase
+	{
+		public Texture RhinoTexture { get { return RhinoCommonObject as Texture; } }
+		public FDatasmithFacadeTexture ExportedTexture { get { return ExportedElement as FDatasmithFacadeTexture; } }
+		public string FilePath { get; private set; }
+
+		/// <summary>
+		/// Holds the Ids of all the Rhino `Texture` objects using that DatasmithTexture.
+		/// </summary>
+		public HashSet<Guid> TextureIds { get; } = new HashSet<Guid>();
+
+		public DatasmithTextureInfo(Texture InRhinoTexture, string InName, string InFilePath)
+			: base(InRhinoTexture, FDatasmithFacadeElement.GetStringHash(InName), InName, InName)
+		{
+			FilePath = InFilePath;
+		}
+
+		public bool IsSupported()
+		{
+			return DatasmithRhinoUtilities.IsTextureSupported(RhinoTexture);
+		}
+	}
+
+}
\ No newline at end of file
diff --git a/Engine/Source/Programs/Enterprise/Datasmith/DatasmithRhinoExporter/Private/Utils/DatasmithRhinoUtilities.cs b/Engine/Source/Programs/Enterprise/Datasmith/DatasmithRhinoExporter/Private/Utils/DatasmithRhinoUtilities.cs
--- a/Engine/Source/Programs/Enterprise/Datasmith/DatasmithRhinoExporter/Private/Utils/DatasmithRhinoUtilities.cs
+++ b/Engine/Source/Programs/Enterprise/Datasmith/DatasmithRhinoExporter/Private/Utils/DatasmithRhinoUtilities.cs
@@ -1,5 +1,6 @@
 // Copyright Epic Games, Inc. All Rights Reserved.
 
+using DatasmithRhino.ExportContext;
 using Rhino;
 using Rhino.DocObjects;
 using Rhino.Geometry;
