[Revit Exporter] Implement export of decals

diff --git a/Engine/Source/Developer/Datasmith/DatasmithFacade/Private/DatasmithFacadeDecal.cpp b/Engine/Source/Developer/Datasmith/DatasmithFacade/Private/DatasmithFacadeDecal.cpp
new file mode 100644
--- /dev/null
+++ b/Engine/Source/Developer/Datasmith/DatasmithFacade/Private/DatasmithFacadeDecal.cpp
@@ -0,0 +1,92 @@
+// Copyright Epic Games, Inc. All Rights Reserved.
+
+#include "DatasmithFacadeDecal.h"
+
+FDatasmithFacadeActorDecal::FDatasmithFacadeActorDecal(
+	const TCHAR* InElementName
+)
+	: FDatasmithFacadeActor(FDatasmithSceneFactory::CreateDecalActor(InElementName))
+{}
+
+FDatasmithFacadeActorDecal::FDatasmithFacadeActorDecal(
+	const TSharedRef<IDatasmithDecalActorElement>& InInternalActor
+)
+	: FDatasmithFacadeActor(InInternalActor)
+{}
+
+TSharedRef<IDatasmithDecalActorElement> FDatasmithFacadeActorDecal::GetDatasmithDecalActorElement() const
+{
+	return StaticCastSharedRef<IDatasmithDecalActorElement>(InternalDatasmithElement);
+}
+
+void FDatasmithFacadeActorDecal::GetDimensions(double& OutX, double& OutY, double& OutZ) const
+{
+	FVector DecalDim = GetDatasmithDecalActorElement()->GetDimensions();
+	OutX = DecalDim.X;
+	OutY = DecalDim.Y;
+	OutZ = DecalDim.Z;
+}
+
+void FDatasmithFacadeActorDecal::SetDimensions(double InX, double InY, double InZ)
+{
+	GetDatasmithDecalActorElement()->SetDimensions(FVector(InX, InY, InZ));
+}
+
+const TCHAR* FDatasmithFacadeActorDecal::GetDecalMaterialPathName() const
+{
+	return GetDatasmithDecalActorElement()->GetDecalMaterialPathName();
+}
+
+void FDatasmithFacadeActorDecal::SetDecalMaterialPathName(const TCHAR* InName)
+{
+	GetDatasmithDecalActorElement()->SetDecalMaterialPathName(InName);
+}
+
+int32 FDatasmithFacadeActorDecal::GetSortOrder() const
+{
+	return GetDatasmithDecalActorElement()->GetSortOrder();
+}
+
+void FDatasmithFacadeActorDecal::SetSortOrder(int32 InSortOrder)
+{
+	GetDatasmithDecalActorElement()->SetSortOrder(InSortOrder);
+}
+
+// class FDatasmithFacadeDecalMaterial
+
+FDatasmithFacadeDecalMaterial::FDatasmithFacadeDecalMaterial(
+	const TSharedRef<IDatasmithDecalMaterialElement>& InMaterialRef // Datasmith master material element
+)	
+	: FDatasmithFacadeBaseMaterial(InMaterialRef)
+{}
+
+FDatasmithFacadeDecalMaterial::FDatasmithFacadeDecalMaterial(
+	const TCHAR* InElementName
+)
+	: FDatasmithFacadeBaseMaterial(FDatasmithSceneFactory::CreateDecalMaterial(InElementName))
+{}
+
+TSharedRef<IDatasmithDecalMaterialElement> FDatasmithFacadeDecalMaterial::GetDatasmithDecalMaterial() const
+{
+	return StaticCastSharedRef<IDatasmithDecalMaterialElement>(InternalDatasmithElement);
+}
+
+const TCHAR* FDatasmithFacadeDecalMaterial::GetDiffuseTexturePathName() const
+{
+	return GetDatasmithDecalMaterial()->GetDiffuseTexturePathName();
+}
+
+void FDatasmithFacadeDecalMaterial::SetDiffuseTexturePathName(const TCHAR* DiffuseTexturePathName)
+{
+	GetDatasmithDecalMaterial()->SetDiffuseTexturePathName(DiffuseTexturePathName);
+}
+
+const TCHAR* FDatasmithFacadeDecalMaterial::GetNormalTexturePathName() const
+{
+	return GetDatasmithDecalMaterial()->GetNormalTexturePathName();
+}
+
+void FDatasmithFacadeDecalMaterial::SetNormalTexturePathName(const TCHAR* NormalTexturePathName)
+{
+	GetDatasmithDecalMaterial()->SetNormalTexturePathName(NormalTexturePathName);
+}
diff --git a/Engine/Source/Developer/Datasmith/DatasmithFacade/Private/DatasmithFacadeMaterial.cpp b/Engine/Source/Developer/Datasmith/DatasmithFacade/Private/DatasmithFacadeMaterial.cpp
--- a/Engine/Source/Developer/Datasmith/DatasmithFacade/Private/DatasmithFacadeMaterial.cpp
+++ b/Engine/Source/Developer/Datasmith/DatasmithFacade/Private/DatasmithFacadeMaterial.cpp
@@ -5,6 +5,7 @@
 #include "DatasmithFacadeKeyValueProperty.h"
 #include "DatasmithFacadeScene.h"
 #include "DatasmithFacadeUEPbrMaterial.h"
+#include "DatasmithFacadeDecal.h"
 
 #include "DatasmithDefinitions.h"
 #include "DatasmithUtils.h"
@@ -34,6 +35,10 @@ FDatasmithFacadeBaseMaterial::EDatasmithMaterialType FDatasmithFacadeBaseMateria
 	{
 		return EDatasmithMaterialType::MasterMaterial;
 	}
+	else if (InMaterial->IsA(EDatasmithElementType::DecalMaterial))
+	{
+		return EDatasmithMaterialType::DecalMaterial;
+	}
 
 	return EDatasmithMaterialType::Unsupported;
 }
@@ -58,6 +63,8 @@ FDatasmithFacadeBaseMaterial* FDatasmithFacadeBaseMaterial::GetNewFacadeBaseMate
 			return new FDatasmithFacadeUEPbrMaterial(StaticCastSharedRef<IDatasmithUEPbrMaterialElement>(MaterialRef));
 		case EDatasmithMaterialType::MasterMaterial:
 			return new FDatasmithFacadeMasterMaterial(StaticCastSharedRef<IDatasmithMasterMaterialElement>(MaterialRef));
+		case EDatasmithMaterialType::DecalMaterial:
+			return new FDatasmithFacadeDecalMaterial(StaticCastSharedRef<IDatasmithDecalMaterialElement>(MaterialRef));
 		case EDatasmithMaterialType::Unsupported:
 		default:
 			return nullptr;
diff --git a/Engine/Source/Developer/Datasmith/DatasmithFacade/Public/DatasmithFacadeActor.h b/Engine/Source/Developer/Datasmith/DatasmithFacade/Public/DatasmithFacadeActor.h
--- a/Engine/Source/Developer/Datasmith/DatasmithFacade/Public/DatasmithFacadeActor.h
+++ b/Engine/Source/Developer/Datasmith/DatasmithFacade/Public/DatasmithFacadeActor.h
@@ -24,6 +24,7 @@ public:
 		StaticMeshActor,
 		Camera,
 		Actor,
+		Decal,
 		Unsupported,
 	};
 
diff --git a/Engine/Source/Developer/Datasmith/DatasmithFacade/Public/DatasmithFacadeDecal.h b/Engine/Source/Developer/Datasmith/DatasmithFacade/Public/DatasmithFacadeDecal.h
new file mode 100644
--- /dev/null
+++ b/Engine/Source/Developer/Datasmith/DatasmithFacade/Public/DatasmithFacadeDecal.h
@@ -0,0 +1,81 @@
+// Copyright Epic Games, Inc. All Rights Reserved.
+
+#pragma once
+
+// Datasmith facade.
+#include "DatasmithFacadeActor.h"
+#include "DatasmithFacadeMaterial.h"
+
+class DATASMITHFACADE_API FDatasmithFacadeActorDecal : public FDatasmithFacadeActor
+{
+	friend class FDatasmithFacadeScene;
+
+public:
+	FDatasmithFacadeActorDecal(
+		const TCHAR* InElementName // Datasmith element name
+	);
+
+	virtual ~FDatasmithFacadeActorDecal() {}
+
+	/** Get the Decal element size */
+	void GetDimensions(double& OutX, double& OutY, double& OutZ) const;
+
+	/** Set the Decal element size */
+	void SetDimensions(double InX, double InY, double InZ);
+
+	/** Get the path name of the Material associated with the actor */
+	const TCHAR* GetDecalMaterialPathName() const;
+
+	/** Set the path name of the Material that the Decal actor uses */
+	void SetDecalMaterialPathName(const TCHAR*);
+
+	/** Get the order in which Decal element is rendered */
+	int32 GetSortOrder() const;
+
+	/** Set the order in which decal elements are rendered.  Higher values draw later (on top) */
+	void SetSortOrder(int32);
+
+#ifdef SWIG_FACADE
+protected:
+#endif
+	FDatasmithFacadeActorDecal(
+		const TSharedRef<IDatasmithDecalActorElement>& InInternalActor
+	);
+
+	TSharedRef<IDatasmithDecalActorElement> GetDatasmithDecalActorElement() const;
+};
+
+class DATASMITHFACADE_API FDatasmithFacadeDecalMaterial : public FDatasmithFacadeBaseMaterial
+{
+	friend class FDatasmithFacadeScene;
+
+public:
+
+	FDatasmithFacadeDecalMaterial(
+		const TCHAR* InElementName // Datasmith element name
+	);
+
+	virtual ~FDatasmithFacadeDecalMaterial() {}
+
+	/** Get path name of the diffuse texture associated with the material */
+	const TCHAR* GetDiffuseTexturePathName() const;
+
+	/** Set path name of the diffuse texture associated with the material */
+	void SetDiffuseTexturePathName(const TCHAR* DiffuseTexturePathName);
+
+	/** Get path name of the normal texture associated with the material */
+	const TCHAR* GetNormalTexturePathName() const;
+
+	/** Set path name of the normal texture associated with the material */
+	void SetNormalTexturePathName(const TCHAR* NormalTexturePathName);
+
+#ifdef SWIG_FACADE
+protected:
+#endif
+
+	FDatasmithFacadeDecalMaterial(
+		const TSharedRef<IDatasmithDecalMaterialElement>& InMaterialRef // Datasmith master material element
+	);
+
+	TSharedRef<IDatasmithDecalMaterialElement> GetDatasmithDecalMaterial() const;
+};
diff --git a/Engine/Source/Developer/Datasmith/DatasmithFacade/Public/DatasmithFacadeMaterial.h b/Engine/Source/Developer/Datasmith/DatasmithFacade/Public/DatasmithFacadeMaterial.h
--- a/Engine/Source/Developer/Datasmith/DatasmithFacade/Public/DatasmithFacadeMaterial.h
+++ b/Engine/Source/Developer/Datasmith/DatasmithFacade/Public/DatasmithFacadeMaterial.h
@@ -16,6 +16,7 @@ public:
 	{
 		MasterMaterial,
 		UEPbrMaterial,
+		DecalMaterial,
 		Unsupported,
 	};
 
diff --git a/Engine/Source/Programs/Enterprise/Datasmith/DatasmithFacadeCSharp/Private/DatasmithFacadeCSharp.cpp b/Engine/Source/Programs/Enterprise/Datasmith/DatasmithFacadeCSharp/Private/DatasmithFacadeCSharp.cpp
--- a/Engine/Source/Programs/Enterprise/Datasmith/DatasmithFacadeCSharp/Private/DatasmithFacadeCSharp.cpp
+++ b/Engine/Source/Programs/Enterprise/Datasmith/DatasmithFacadeCSharp/Private/DatasmithFacadeCSharp.cpp
@@ -339,6 +339,7 @@ SWIGEXPORT void SWIGSTDCALL SWIGRegisterWStringCallback_DatasmithFacadeCSharp(SW
 #include "DatasmithFacadeVariant.h"
 #include "DatasmithFacadeLog.h"
 #include "DatasmithFacadeUtils.h"
+#include "DatasmithFacadeDecal.h"
 
 #include "IDatasmithExporterUIModule.h"
 #include "IDirectLinkUI.h"
@@ -7226,6 +7227,162 @@ SWIGEXPORT void SWIGSTDCALL CSharp_FDatasmithFacadeLevelVariantSets_RemoveVarian
 }
 
 
+SWIGEXPORT void * SWIGSTDCALL CSharp_new_FDatasmithFacadeActorDecal(wchar_t * jarg1) {
+  void * jresult ;
+  TCHAR *arg1 = (TCHAR *) 0 ;
+  FDatasmithFacadeActorDecal *result = 0 ;
+  
+  arg1 = (TCHAR *)jarg1; 
+  result = (FDatasmithFacadeActorDecal *)new FDatasmithFacadeActorDecal((TCHAR const *)arg1);
+  jresult = (void *)result; 
+  return jresult;
+}
+
+
+SWIGEXPORT void SWIGSTDCALL CSharp_delete_FDatasmithFacadeActorDecal(void * jarg1) {
+  FDatasmithFacadeActorDecal *arg1 = (FDatasmithFacadeActorDecal *) 0 ;
+  
+  arg1 = (FDatasmithFacadeActorDecal *)jarg1; 
+  delete arg1;
+}
+
+
+SWIGEXPORT void SWIGSTDCALL CSharp_FDatasmithFacadeActorDecal_GetDimensions(void * jarg1, double * jarg2, double * jarg3, double * jarg4) {
+  FDatasmithFacadeActorDecal *arg1 = (FDatasmithFacadeActorDecal *) 0 ;
+  double *arg2 = 0 ;
+  double *arg3 = 0 ;
+  double *arg4 = 0 ;
+  
+  arg1 = (FDatasmithFacadeActorDecal *)jarg1; 
+  arg2 = (double *)jarg2; 
+  arg3 = (double *)jarg3; 
+  arg4 = (double *)jarg4; 
+  ((FDatasmithFacadeActorDecal const *)arg1)->GetDimensions(*arg2,*arg3,*arg4);
+}
+
+
+SWIGEXPORT void SWIGSTDCALL CSharp_FDatasmithFacadeActorDecal_SetDimensions(void * jarg1, double jarg2, double jarg3, double jarg4) {
+  FDatasmithFacadeActorDecal *arg1 = (FDatasmithFacadeActorDecal *) 0 ;
+  double arg2 ;
+  double arg3 ;
+  double arg4 ;
+  
+  arg1 = (FDatasmithFacadeActorDecal *)jarg1; 
+  arg2 = (double)jarg2; 
+  arg3 = (double)jarg3; 
+  arg4 = (double)jarg4; 
+  (arg1)->SetDimensions(arg2,arg3,arg4);
+}
+
+
+SWIGEXPORT wchar_t * SWIGSTDCALL CSharp_FDatasmithFacadeActorDecal_GetDecalMaterialPathName(void * jarg1) {
+  wchar_t * jresult ;
+  FDatasmithFacadeActorDecal *arg1 = (FDatasmithFacadeActorDecal *) 0 ;
+  TCHAR *result = 0 ;
+  
+  arg1 = (FDatasmithFacadeActorDecal *)jarg1; 
+  result = (TCHAR *)((FDatasmithFacadeActorDecal const *)arg1)->GetDecalMaterialPathName();
+  jresult = (wchar_t *)result; 
+  return jresult;
+}
+
+
+SWIGEXPORT void SWIGSTDCALL CSharp_FDatasmithFacadeActorDecal_SetDecalMaterialPathName(void * jarg1, wchar_t * jarg2) {
+  FDatasmithFacadeActorDecal *arg1 = (FDatasmithFacadeActorDecal *) 0 ;
+  TCHAR *arg2 = (TCHAR *) 0 ;
+  
+  arg1 = (FDatasmithFacadeActorDecal *)jarg1; 
+  arg2 = (TCHAR *)jarg2; 
+  (arg1)->SetDecalMaterialPathName((TCHAR const *)arg2);
+}
+
+
+SWIGEXPORT int SWIGSTDCALL CSharp_FDatasmithFacadeActorDecal_GetSortOrder(void * jarg1) {
+  int jresult ;
+  FDatasmithFacadeActorDecal *arg1 = (FDatasmithFacadeActorDecal *) 0 ;
+  int32 result;
+  
+  arg1 = (FDatasmithFacadeActorDecal *)jarg1; 
+  result = (int32)((FDatasmithFacadeActorDecal const *)arg1)->GetSortOrder();
+  jresult = result; 
+  return jresult;
+}
+
+
+SWIGEXPORT void SWIGSTDCALL CSharp_FDatasmithFacadeActorDecal_SetSortOrder(void * jarg1, int jarg2) {
+  FDatasmithFacadeActorDecal *arg1 = (FDatasmithFacadeActorDecal *) 0 ;
+  int32 arg2 ;
+  
+  arg1 = (FDatasmithFacadeActorDecal *)jarg1; 
+  arg2 = (int32)jarg2; 
+  (arg1)->SetSortOrder(arg2);
+}
+
+
+SWIGEXPORT void * SWIGSTDCALL CSharp_new_FDatasmithFacadeDecalMaterial(wchar_t * jarg1) {
+  void * jresult ;
+  TCHAR *arg1 = (TCHAR *) 0 ;
+  FDatasmithFacadeDecalMaterial *result = 0 ;
+  
+  arg1 = (TCHAR *)jarg1; 
+  result = (FDatasmithFacadeDecalMaterial *)new FDatasmithFacadeDecalMaterial((TCHAR const *)arg1);
+  jresult = (void *)result; 
+  return jresult;
+}
+
+
+SWIGEXPORT void SWIGSTDCALL CSharp_delete_FDatasmithFacadeDecalMaterial(void * jarg1) {
+  FDatasmithFacadeDecalMaterial *arg1 = (FDatasmithFacadeDecalMaterial *) 0 ;
+  
+  arg1 = (FDatasmithFacadeDecalMaterial *)jarg1; 
+  delete arg1;
+}
+
+
+SWIGEXPORT wchar_t * SWIGSTDCALL CSharp_FDatasmithFacadeDecalMaterial_GetDiffuseTexturePathName(void * jarg1) {
+  wchar_t * jresult ;
+  FDatasmithFacadeDecalMaterial *arg1 = (FDatasmithFacadeDecalMaterial *) 0 ;
+  TCHAR *result = 0 ;
+  
+  arg1 = (FDatasmithFacadeDecalMaterial *)jarg1; 
+  result = (TCHAR *)((FDatasmithFacadeDecalMaterial const *)arg1)->GetDiffuseTexturePathName();
+  jresult = (wchar_t *)result; 
+  return jresult;
+}
+
+
+SWIGEXPORT void SWIGSTDCALL CSharp_FDatasmithFacadeDecalMaterial_SetDiffuseTexturePathName(void * jarg1, wchar_t * jarg2) {
+  FDatasmithFacadeDecalMaterial *arg1 = (FDatasmithFacadeDecalMaterial *) 0 ;
+  TCHAR *arg2 = (TCHAR *) 0 ;
+  
+  arg1 = (FDatasmithFacadeDecalMaterial *)jarg1; 
+  arg2 = (TCHAR *)jarg2; 
+  (arg1)->SetDiffuseTexturePathName((TCHAR const *)arg2);
+}
+
+
+SWIGEXPORT wchar_t * SWIGSTDCALL CSharp_FDatasmithFacadeDecalMaterial_GetNormalTexturePathName(void * jarg1) {
+  wchar_t * jresult ;
+  FDatasmithFacadeDecalMaterial *arg1 = (FDatasmithFacadeDecalMaterial *) 0 ;
+  TCHAR *result = 0 ;
+  
+  arg1 = (FDatasmithFacadeDecalMaterial *)jarg1; 
+  result = (TCHAR *)((FDatasmithFacadeDecalMaterial const *)arg1)->GetNormalTexturePathName();
+  jresult = (wchar_t *)result; 
+  return jresult;
+}
+
+
+SWIGEXPORT void SWIGSTDCALL CSharp_FDatasmithFacadeDecalMaterial_SetNormalTexturePathName(void * jarg1, wchar_t * jarg2) {
+  FDatasmithFacadeDecalMaterial *arg1 = (FDatasmithFacadeDecalMaterial *) 0 ;
+  TCHAR *arg2 = (TCHAR *) 0 ;
+  
+  arg1 = (FDatasmithFacadeDecalMaterial *)jarg1; 
+  arg2 = (TCHAR *)jarg2; 
+  (arg1)->SetNormalTexturePathName((TCHAR const *)arg2);
+}
+
+
 SWIGEXPORT void * SWIGSTDCALL CSharp_IDatasmithExporterUIModule_Get() {
   void * jresult ;
   IDatasmithExporterUIModule *result = 0 ;
@@ -7444,6 +7601,14 @@ SWIGEXPORT FDatasmithFacadeElement * SWIGSTDCALL CSharp_FDatasmithFacadeLevelVar
     return (FDatasmithFacadeElement *)jarg1;
 }
 
+SWIGEXPORT FDatasmithFacadeActor * SWIGSTDCALL CSharp_FDatasmithFacadeActorDecal_SWIGUpcast(FDatasmithFacadeActorDecal *jarg1) {
+    return (FDatasmithFacadeActor *)jarg1;
+}
+
+SWIGEXPORT FDatasmithFacadeBaseMaterial * SWIGSTDCALL CSharp_FDatasmithFacadeDecalMaterial_SWIGUpcast(FDatasmithFacadeDecalMaterial *jarg1) {
+    return (FDatasmithFacadeBaseMaterial *)jarg1;
+}
+
 #ifdef __cplusplus
 }
 #endif
diff --git a/Engine/Source/Programs/Enterprise/Datasmith/DatasmithFacadeCSharp/Public/DatasmithFacadeCSharpPINVOKE.cs b/Engine/Source/Programs/Enterprise/Datasmith/DatasmithFacadeCSharp/Public/DatasmithFacadeCSharpPINVOKE.cs
--- a/Engine/Source/Programs/Enterprise/Datasmith/DatasmithFacadeCSharp/Public/DatasmithFacadeCSharpPINVOKE.cs
+++ b/Engine/Source/Programs/Enterprise/Datasmith/DatasmithFacadeCSharp/Public/DatasmithFacadeCSharpPINVOKE.cs
@@ -1994,6 +1994,48 @@ class DatasmithFacadeCSharpPINVOKE {
   [global::System.Runtime.InteropServices.DllImport("DatasmithFacadeCSharp", EntryPoint="CSharp_FDatasmithFacadeLevelVariantSets_RemoveVariantSet")]
   public static extern void FDatasmithFacadeLevelVariantSets_RemoveVariantSet(global::System.Runtime.InteropServices.HandleRef jarg1, global::System.Runtime.InteropServices.HandleRef jarg2);
 
+  [global::System.Runtime.InteropServices.DllImport("DatasmithFacadeCSharp", EntryPoint="CSharp_new_FDatasmithFacadeActorDecal")]
+  public static extern global::System.IntPtr new_FDatasmithFacadeActorDecal([global::System.Runtime.InteropServices.MarshalAs(global::System.Runtime.InteropServices.UnmanagedType.LPWStr)]string jarg1);
+
+  [global::System.Runtime.InteropServices.DllImport("DatasmithFacadeCSharp", EntryPoint="CSharp_delete_FDatasmithFacadeActorDecal")]
+  public static extern void delete_FDatasmithFacadeActorDecal(global::System.Runtime.InteropServices.HandleRef jarg1);
+
+  [global::System.Runtime.InteropServices.DllImport("DatasmithFacadeCSharp", EntryPoint="CSharp_FDatasmithFacadeActorDecal_GetDimensions")]
+  public static extern void FDatasmithFacadeActorDecal_GetDimensions(global::System.Runtime.InteropServices.HandleRef jarg1, ref double jarg2, ref double jarg3, ref double jarg4);
+
+  [global::System.Runtime.InteropServices.DllImport("DatasmithFacadeCSharp", EntryPoint="CSharp_FDatasmithFacadeActorDecal_SetDimensions")]
+  public static extern void FDatasmithFacadeActorDecal_SetDimensions(global::System.Runtime.InteropServices.HandleRef jarg1, double jarg2, double jarg3, double jarg4);
+
+  [global::System.Runtime.InteropServices.DllImport("DatasmithFacadeCSharp", EntryPoint="CSharp_FDatasmithFacadeActorDecal_GetDecalMaterialPathName")]
+  public static extern global::System.IntPtr FDatasmithFacadeActorDecal_GetDecalMaterialPathName(global::System.Runtime.InteropServices.HandleRef jarg1);
+
+  [global::System.Runtime.InteropServices.DllImport("DatasmithFacadeCSharp", EntryPoint="CSharp_FDatasmithFacadeActorDecal_SetDecalMaterialPathName")]
+  public static extern void FDatasmithFacadeActorDecal_SetDecalMaterialPathName(global::System.Runtime.InteropServices.HandleRef jarg1, [global::System.Runtime.InteropServices.MarshalAs(global::System.Runtime.InteropServices.UnmanagedType.LPWStr)]string jarg2);
+
+  [global::System.Runtime.InteropServices.DllImport("DatasmithFacadeCSharp", EntryPoint="CSharp_FDatasmithFacadeActorDecal_GetSortOrder")]
+  public static extern int FDatasmithFacadeActorDecal_GetSortOrder(global::System.Runtime.InteropServices.HandleRef jarg1);
+
+  [global::System.Runtime.InteropServices.DllImport("DatasmithFacadeCSharp", EntryPoint="CSharp_FDatasmithFacadeActorDecal_SetSortOrder")]
+  public static extern void FDatasmithFacadeActorDecal_SetSortOrder(global::System.Runtime.InteropServices.HandleRef jarg1, int jarg2);
+
+  [global::System.Runtime.InteropServices.DllImport("DatasmithFacadeCSharp", EntryPoint="CSharp_new_FDatasmithFacadeDecalMaterial")]
+  public static extern global::System.IntPtr new_FDatasmithFacadeDecalMaterial([global::System.Runtime.InteropServices.MarshalAs(global::System.Runtime.InteropServices.UnmanagedType.LPWStr)]string jarg1);
+
+  [global::System.Runtime.InteropServices.DllImport("DatasmithFacadeCSharp", EntryPoint="CSharp_delete_FDatasmithFacadeDecalMaterial")]
+  public static extern void delete_FDatasmithFacadeDecalMaterial(global::System.Runtime.InteropServices.HandleRef jarg1);
+
+  [global::System.Runtime.InteropServices.DllImport("DatasmithFacadeCSharp", EntryPoint="CSharp_FDatasmithFacadeDecalMaterial_GetDiffuseTexturePathName")]
+  public static extern global::System.IntPtr FDatasmithFacadeDecalMaterial_GetDiffuseTexturePathName(global::System.Runtime.InteropServices.HandleRef jarg1);
+
+  [global::System.Runtime.InteropServices.DllImport("DatasmithFacadeCSharp", EntryPoint="CSharp_FDatasmithFacadeDecalMaterial_SetDiffuseTexturePathName")]
+  public static extern void FDatasmithFacadeDecalMaterial_SetDiffuseTexturePathName(global::System.Runtime.InteropServices.HandleRef jarg1, [global::System.Runtime.InteropServices.MarshalAs(global::System.Runtime.InteropServices.UnmanagedType.LPWStr)]string jarg2);
+
+  [global::System.Runtime.InteropServices.DllImport("DatasmithFacadeCSharp", EntryPoint="CSharp_FDatasmithFacadeDecalMaterial_GetNormalTexturePathName")]
+  public static extern global::System.IntPtr FDatasmithFacadeDecalMaterial_GetNormalTexturePathName(global::System.Runtime.InteropServices.HandleRef jarg1);
+
+  [global::System.Runtime.InteropServices.DllImport("DatasmithFacadeCSharp", EntryPoint="CSharp_FDatasmithFacadeDecalMaterial_SetNormalTexturePathName")]
+  public static extern void FDatasmithFacadeDecalMaterial_SetNormalTexturePathName(global::System.Runtime.InteropServices.HandleRef jarg1, [global::System.Runtime.InteropServices.MarshalAs(global::System.Runtime.InteropServices.UnmanagedType.LPWStr)]string jarg2);
+
   [global::System.Runtime.InteropServices.DllImport("DatasmithFacadeCSharp", EntryPoint="CSharp_IDatasmithExporterUIModule_Get")]
   public static extern global::System.IntPtr IDatasmithExporterUIModule_Get();
 
@@ -2131,4 +2173,10 @@ class DatasmithFacadeCSharpPINVOKE {
 
   [global::System.Runtime.InteropServices.DllImport("DatasmithFacadeCSharp", EntryPoint="CSharp_FDatasmithFacadeLevelVariantSets_SWIGUpcast")]
   public static extern global::System.IntPtr FDatasmithFacadeLevelVariantSets_SWIGUpcast(global::System.IntPtr jarg1);
+
+  [global::System.Runtime.InteropServices.DllImport("DatasmithFacadeCSharp", EntryPoint="CSharp_FDatasmithFacadeActorDecal_SWIGUpcast")]
+  public static extern global::System.IntPtr FDatasmithFacadeActorDecal_SWIGUpcast(global::System.IntPtr jarg1);
+
+  [global::System.Runtime.InteropServices.DllImport("DatasmithFacadeCSharp", EntryPoint="CSharp_FDatasmithFacadeDecalMaterial_SWIGUpcast")]
+  public static extern global::System.IntPtr FDatasmithFacadeDecalMaterial_SWIGUpcast(global::System.IntPtr jarg1);
 }
diff --git a/Engine/Source/Programs/Enterprise/Datasmith/DatasmithFacadeCSharp/Public/FDatasmithFacadeActor.cs b/Engine/Source/Programs/Enterprise/Datasmith/DatasmithFacadeCSharp/Public/FDatasmithFacadeActor.cs
--- a/Engine/Source/Programs/Enterprise/Datasmith/DatasmithFacadeCSharp/Public/FDatasmithFacadeActor.cs
+++ b/Engine/Source/Programs/Enterprise/Datasmith/DatasmithFacadeCSharp/Public/FDatasmithFacadeActor.cs
@@ -149,6 +149,8 @@ public class FDatasmithFacadeActor : FDatasmithFacadeElement {
 			return new FDatasmithFacadeActorMesh(objectPtr, true);
 		case FDatasmithFacadeActor.EActorType.Camera:
 			return new FDatasmithFacadeActorCamera(objectPtr, true);
+		case FDatasmithFacadeActor.EActorType.Decal:
+			return new FDatasmithFacadeActorDecal(objectPtr, true);
 		case FDatasmithFacadeActor.EActorType.Actor:
 			return new FDatasmithFacadeActor(objectPtr, true);
 		case FDatasmithFacadeActor.EActorType.Unsupported:
@@ -188,6 +190,8 @@ public class FDatasmithFacadeActor : FDatasmithFacadeElement {
 			return new FDatasmithFacadeActorMesh(objectPtr, true);
 		case FDatasmithFacadeActor.EActorType.Camera:
 			return new FDatasmithFacadeActorCamera(objectPtr, true);
+		case FDatasmithFacadeActor.EActorType.Decal:
+			return new FDatasmithFacadeActorDecal(objectPtr, true);
 		case FDatasmithFacadeActor.EActorType.Actor:
 			return new FDatasmithFacadeActor(objectPtr, true);
 		case FDatasmithFacadeActor.EActorType.Unsupported:
@@ -230,6 +234,7 @@ public class FDatasmithFacadeActor : FDatasmithFacadeElement {
     StaticMeshActor,
     Camera,
     Actor,
+    Decal,
     Unsupported
   }
 
diff --git a/Engine/Source/Programs/Enterprise/Datasmith/DatasmithFacadeCSharp/Public/FDatasmithFacadeBaseMaterial.cs b/Engine/Source/Programs/Enterprise/Datasmith/DatasmithFacadeCSharp/Public/FDatasmithFacadeBaseMaterial.cs
--- a/Engine/Source/Programs/Enterprise/Datasmith/DatasmithFacadeCSharp/Public/FDatasmithFacadeBaseMaterial.cs
+++ b/Engine/Source/Programs/Enterprise/Datasmith/DatasmithFacadeCSharp/Public/FDatasmithFacadeBaseMaterial.cs
@@ -43,6 +43,7 @@ public class FDatasmithFacadeBaseMaterial : FDatasmithFacadeElement {
   public enum EDatasmithMaterialType {
     MasterMaterial,
     UEPbrMaterial,
+    DecalMaterial,
     Unsupported
   }
 
diff --git a/Engine/Source/Programs/Enterprise/Datasmith/DatasmithFacadeCSharp/Public/FDatasmithFacadeScene.cs b/Engine/Source/Programs/Enterprise/Datasmith/DatasmithFacadeCSharp/Public/FDatasmithFacadeScene.cs
--- a/Engine/Source/Programs/Enterprise/Datasmith/DatasmithFacadeCSharp/Public/FDatasmithFacadeScene.cs
+++ b/Engine/Source/Programs/Enterprise/Datasmith/DatasmithFacadeCSharp/Public/FDatasmithFacadeScene.cs
@@ -83,6 +83,8 @@ public class FDatasmithFacadeScene : global::System.IDisposable {
 			return new FDatasmithFacadeActorMesh(objectPtr, true);
 		case FDatasmithFacadeActor.EActorType.Camera:
 			return new FDatasmithFacadeActorCamera(objectPtr, true);
+		case FDatasmithFacadeActor.EActorType.Decal:
+			return new FDatasmithFacadeActorDecal(objectPtr, true);
 		case FDatasmithFacadeActor.EActorType.Actor:
 			return new FDatasmithFacadeActor(objectPtr, true);
 		case FDatasmithFacadeActor.EActorType.Unsupported:
@@ -133,6 +135,8 @@ public class FDatasmithFacadeScene : global::System.IDisposable {
 			return new FDatasmithFacadeUEPbrMaterial(objectPtr, true);
 		case FDatasmithFacadeBaseMaterial.EDatasmithMaterialType.MasterMaterial:
 			return new FDatasmithFacadeMasterMaterial(objectPtr, true);
+		case FDatasmithFacadeBaseMaterial.EDatasmithMaterialType.DecalMaterial:
+			return new FDatasmithFacadeDecalMaterial(objectPtr, true);
 		case FDatasmithFacadeBaseMaterial.EDatasmithMaterialType.Unsupported:
 		default:
 			return null;
diff --git a/Engine/Source/Programs/Enterprise/Datasmith/DatasmithRevitExporter/DatasmithRevit2018/DatasmithRevit2018.csproj b/Engine/Source/Programs/Enterprise/Datasmith/DatasmithRevitExporter/DatasmithRevit2018/DatasmithRevit2018.csproj
--- a/Engine/Source/Programs/Enterprise/Datasmith/DatasmithRevitExporter/DatasmithRevit2018/DatasmithRevit2018.csproj
+++ b/Engine/Source/Programs/Enterprise/Datasmith/DatasmithRevitExporter/DatasmithRevit2018/DatasmithRevit2018.csproj
@@ -64,6 +64,9 @@
     <Compile Include="..\Private\DatasmithRevitDirectLink.cs">
       <Link>DatasmithRevitExporter\DatasmithRevitDirectLink.cs</Link>
     </Compile>
+    <Compile Include="..\Private\DatasmithRevitDecal.cs">
+      <Link>DatasmithRevitExporter\DatasmithRevitDecal.cs</Link>
+    </Compile>
     <Compile Include="..\Private\DatasmithRevitApplication.cs">
       <Link>DatasmithRevitExporter\DatasmithRevitApplication.cs</Link>
     </Compile>
diff --git a/Engine/Source/Programs/Enterprise/Datasmith/DatasmithRevitExporter/DatasmithRevit2019/DatasmithRevit2019.csproj b/Engine/Source/Programs/Enterprise/Datasmith/DatasmithRevitExporter/DatasmithRevit2019/DatasmithRevit2019.csproj
--- a/Engine/Source/Programs/Enterprise/Datasmith/DatasmithRevitExporter/DatasmithRevit2019/DatasmithRevit2019.csproj
+++ b/Engine/Source/Programs/Enterprise/Datasmith/DatasmithRevitExporter/DatasmithRevit2019/DatasmithRevit2019.csproj
@@ -64,6 +64,9 @@
     <Compile Include="..\Private\DatasmithRevitDirectLink.cs">
       <Link>DatasmithRevitExporter\DatasmithRevitDirectLink.cs</Link>
     </Compile>
+    <Compile Include="..\Private\DatasmithRevitDecal.cs">
+      <Link>DatasmithRevitExporter\DatasmithRevitDecal.cs</Link>
+    </Compile>
     <Compile Include="..\Private\DatasmithRevitApplication.cs">
       <Link>DatasmithRevitExporter\DatasmithRevitApplication.cs</Link>
     </Compile>
diff --git a/Engine/Source/Programs/Enterprise/Datasmith/DatasmithRevitExporter/DatasmithRevit2020/DatasmithRevit2020.csproj b/Engine/Source/Programs/Enterprise/Datasmith/DatasmithRevitExporter/DatasmithRevit2020/DatasmithRevit2020.csproj
--- a/Engine/Source/Programs/Enterprise/Datasmith/DatasmithRevitExporter/DatasmithRevit2020/DatasmithRevit2020.csproj
+++ b/Engine/Source/Programs/Enterprise/Datasmith/DatasmithRevitExporter/DatasmithRevit2020/DatasmithRevit2020.csproj
@@ -67,6 +67,9 @@
     <Compile Include="..\Private\DatasmithRevitDirectLink.cs">
       <Link>DatasmithRevitExporter\DatasmithRevitDirectLink.cs</Link>
     </Compile>
+    <Compile Include="..\Private\DatasmithRevitDecal.cs">
+      <Link>DatasmithRevitExporter\DatasmithRevitDecal.cs</Link>
+    </Compile>
     <Compile Include="..\Private\DatasmithRevitCommand.cs">
       <Link>DatasmithRevitExporter\DatasmithRevitCommand.cs</Link>
     </Compile>
diff --git a/Engine/Source/Programs/Enterprise/Datasmith/DatasmithRevitExporter/DatasmithRevit2021/DatasmithRevit2021.csproj b/Engine/Source/Programs/Enterprise/Datasmith/DatasmithRevitExporter/DatasmithRevit2021/DatasmithRevit2021.csproj
--- a/Engine/Source/Programs/Enterprise/Datasmith/DatasmithRevitExporter/DatasmithRevit2021/DatasmithRevit2021.csproj
+++ b/Engine/Source/Programs/Enterprise/Datasmith/DatasmithRevitExporter/DatasmithRevit2021/DatasmithRevit2021.csproj
@@ -64,6 +64,9 @@
     <Compile Include="..\Private\DatasmithRevitDirectLink.cs">
       <Link>DatasmithRevitExporter\DatasmithRevitDirectLink.cs</Link>
     </Compile>
+    <Compile Include="..\Private\DatasmithRevitDecal.cs">
+      <Link>DatasmithRevitExporter\DatasmithRevitDecal.cs</Link>
+    </Compile>
     <Compile Include="..\Private\DatasmithRevitApplication.cs">
       <Link>DatasmithRevitExporter\DatasmithRevitApplication.cs</Link>
     </Compile>
diff --git a/Engine/Source/Programs/Enterprise/Datasmith/DatasmithRevitExporter/DatasmithRevit2022/DatasmithRevit2022.csproj b/Engine/Source/Programs/Enterprise/Datasmith/DatasmithRevitExporter/DatasmithRevit2022/DatasmithRevit2022.csproj
--- a/Engine/Source/Programs/Enterprise/Datasmith/DatasmithRevitExporter/DatasmithRevit2022/DatasmithRevit2022.csproj
+++ b/Engine/Source/Programs/Enterprise/Datasmith/DatasmithRevitExporter/DatasmithRevit2022/DatasmithRevit2022.csproj
@@ -64,6 +64,9 @@
     <Compile Include="..\Private\DatasmithRevitDirectLink.cs">
       <Link>DatasmithRevitExporter\DatasmithRevitDirectLink.cs</Link>
     </Compile>
+    <Compile Include="..\Private\DatasmithRevitDecal.cs">
+      <Link>DatasmithRevitExporter\DatasmithRevitDecal.cs</Link>
+    </Compile>
     <Compile Include="..\Private\DatasmithRevitApplication.cs">
       <Link>DatasmithRevitExporter\DatasmithRevitApplication.cs</Link>
     </Compile>
diff --git a/Engine/Source/Programs/Enterprise/Datasmith/DatasmithRevitExporter/Private/DatasmithRevitDecal.cs b/Engine/Source/Programs/Enterprise/Datasmith/DatasmithRevitExporter/Private/DatasmithRevitDecal.cs
new file mode 100644
--- /dev/null
+++ b/Engine/Source/Programs/Enterprise/Datasmith/DatasmithRevitExporter/Private/DatasmithRevitDecal.cs
@@ -0,0 +1,169 @@
+// Copyright Epic Games, Inc. All Rights Reserved.
+
+using Autodesk.Revit.DB;
+using Autodesk.Revit.DB.Visual;
+using System.Collections.Generic;
+
+namespace DatasmithRevitExporter
+{
+	public class FDecalMaterial
+	{
+		public string MaterialName { get; private set; }
+		public string DiffuseTexturePath { get; private set; }
+		public string BumpTexturePath { get; private set; }
+
+		private void SetMaterialProperties(Asset InMaterialAsset)
+		{
+			const string DiffuseMapPropName = "decApp_diffuse";
+			const string BumpMapPropName = "decApp_bump_map";
+			const string TexturePathName = "unifiedbitmap_Bitmap";
+
+			string GetConnectedImagePath(AssetProperty InProp)
+			{
+				IList<AssetProperty> ConnectedProps = InProp.GetAllConnectedProperties();
+
+				foreach (AssetProperty InnerProp in ConnectedProps)
+				{
+					if (InnerProp.Type == AssetPropertyType.Asset)
+					{
+						Asset InnerPropAsset = InnerProp as Asset;
+
+						for (int Index = 0; Index < InnerPropAsset.Size; ++Index)
+						{
+#if REVIT_API_2018
+							AssetProperty AssetProp = InnerPropAsset[Index];
+#else
+							AssetProperty AssetProp = InnerPropAsset.Get(Index);
+#endif
+
+							if (AssetProp.Name == TexturePathName)
+							{
+								return (AssetProp as AssetPropertyString).Value;
+							}
+						}
+					}
+				}
+
+				return null;
+			}
+
+			bool bDiffusePropertyProcessed = false;
+			bool bBumpPropertyProcessed = false;
+
+			for (int PropIndex = 0; PropIndex < InMaterialAsset.Size; ++PropIndex)
+			{
+#if REVIT_API_2018
+				AssetProperty Prop = InMaterialAsset[PropIndex];
+#else
+				AssetProperty Prop = InMaterialAsset.Get(PropIndex);
+#endif
+				if (Prop.Name == DiffuseMapPropName)
+				{
+					DiffuseTexturePath = GetConnectedImagePath(Prop);
+					bDiffusePropertyProcessed = true;
+				}
+				else if (Prop.Name == BumpMapPropName)
+				{
+					BumpTexturePath = GetConnectedImagePath(Prop);
+					bBumpPropertyProcessed = true;
+				}
+
+				if (bDiffusePropertyProcessed && bBumpPropertyProcessed)
+				{
+					break;
+				}
+			}
+		}
+
+		public static ElementId GetDecalElementId(MaterialNode InMaterialNode)
+		{
+			const string PropertyName_DecalElemId = "decalelementId";
+
+			if (!InMaterialNode.HasOverriddenAppearance)
+			{
+				return ElementId.InvalidElementId;
+			}
+
+			Asset Decal = InMaterialNode.GetAppearanceOverride();
+
+			if (Decal == null || Decal.Name != "Decal")
+			{
+				return ElementId.InvalidElementId;
+			}
+
+			for (int PropIndex = 0; PropIndex < Decal.Size; ++PropIndex)
+			{
+#if REVIT_API_2018
+				AssetProperty Prop = Decal[PropIndex];
+#else
+				AssetProperty Prop = Decal.Get(PropIndex);
+#endif
+				if (Prop.Name == PropertyName_DecalElemId && Prop.Type == AssetPropertyType.Integer)
+				{
+					int Id = (Prop as AssetPropertyInteger).Value;
+					return new ElementId(Id);
+				}
+			}
+			return ElementId.InvalidElementId;
+		}
+
+		public static FDecalMaterial Create(MaterialNode InMaterialNode, Material InMaterial)
+		{
+			if (!InMaterialNode.HasOverriddenAppearance)
+			{
+				return null;
+			}
+
+			Asset Decal = InMaterialNode.GetAppearanceOverride();
+
+			if (Decal == null || Decal.Name != "Decal")
+			{
+				return null;
+			}
+
+			const string AssetName_Material = "DecalAppearance";
+
+			FDecalMaterial DecalMaterial = new FDecalMaterial();
+
+			DecalMaterial.MaterialName = FDatasmithFacadeElement.GetStringHash($"Decal_{FMaterialData.GetMaterialName(InMaterialNode, InMaterial)}");
+
+			// Look for DecalAppearance asset, which defines the material parameters
+
+			bool bMaterialInitialized = false;
+
+			for (int PropIndex = 0; PropIndex < Decal.Size; ++PropIndex)
+			{
+#if REVIT_API_2018
+				AssetProperty Prop = Decal[PropIndex];
+#else
+				AssetProperty Prop = Decal.Get(PropIndex);
+#endif
+				if (Prop.Type == AssetPropertyType.Reference)
+				{
+					AssetPropertyReference RefProp = (Prop as AssetPropertyReference);
+					foreach (AssetProperty InnerProp in RefProp.GetAllConnectedProperties())
+					{
+						if (InnerProp.Type == AssetPropertyType.Asset)
+						{
+							Asset InnerPropAsset = InnerProp as Asset;
+
+							if (InnerPropAsset.Name == AssetName_Material)
+							{
+								DecalMaterial.SetMaterialProperties(InnerPropAsset);
+								bMaterialInitialized = true;
+								break;
+							}
+						}
+					}
+				}
+
+				if (bMaterialInitialized)
+				{
+					break;
+				}
+			}
+
+			return DecalMaterial;
+		}
+	}
+}
diff --git a/Engine/Source/Programs/Enterprise/Datasmith/DatasmithRevitExporter/Private/DatasmithRevitDocumentData.cs b/Engine/Source/Programs/Enterprise/Datasmith/DatasmithRevitExporter/Private/DatasmithRevitDocumentData.cs
--- a/Engine/Source/Programs/Enterprise/Datasmith/DatasmithRevitExporter/Private/DatasmithRevitDocumentData.cs
+++ b/Engine/Source/Programs/Enterprise/Datasmith/DatasmithRevitExporter/Private/DatasmithRevitDocumentData.cs
@@ -114,6 +114,7 @@ namespace DatasmithRevitExporter
 			public bool							bOptimizeHierarchy = true;
 			public bool							bIsModified = true;
 			public bool							bAllowMeshInstancing = true;
+			public bool							bIsDecalElement = false;
 
 			public Dictionary<string, int>		MeshMaterialsMap = new Dictionary<string, int>();
 
@@ -807,8 +808,18 @@ namespace DatasmithRevitExporter
 					// Create a new Datasmith mesh actor.
 					// Hash the Datasmith mesh actor name to shorten it.
 					string HashedActorName = FDatasmithFacadeElement.GetStringHash("A:" + GetActorName(true));
-					InElement.ElementActor = new FDatasmithFacadeActorMesh(HashedActorName);
-					InElement.ElementActor.SetLabel(GetActorLabel());
+
+					if (BaseElementType.FamilyName == "Decal")
+					{
+						bOptimizeHierarchy = false;
+						bIsDecalElement = true;
+						InElement.ElementActor = new FDatasmithFacadeActorDecal(HashedActorName);
+					}
+					else
+					{
+						InElement.ElementActor = new FDatasmithFacadeActorMesh(HashedActorName);
+						InElement.ElementActor.SetLabel(GetActorLabel());
+					}
 				}
 
 				// Set the world transform of the Datasmith mesh actor.
@@ -1058,6 +1069,9 @@ namespace DatasmithRevitExporter
 		public Dictionary<string, FMaterialData>		MaterialDataMap = null;
 		public Dictionary<string, FMaterialData>		NewMaterialsMap = new Dictionary<string, FMaterialData>();
 
+		public Dictionary<ElementId, FElementData>		DecalElementsMap = new Dictionary<ElementId, FElementData>();
+		public Dictionary<ElementId, FDecalMaterial>	DecalMaterialsMap = new Dictionary<ElementId, FDecalMaterial>();
+
 		private Stack<FElementData>						ElementDataStack = new Stack<FElementData>();
 		private string									CurrentMaterialName = null;
 		private List<string>							MessageList = null;
@@ -1226,6 +1240,11 @@ namespace DatasmithRevitExporter
 				else
 				{
 					ElementData = new FElementData(InElement, InWorldTransform, this);
+
+					if (ElementData.bIsDecalElement)
+					{
+						DecalElementsMap.Add(InElement.Id, ElementData);
+					}
 				}
 			}
 
@@ -1442,6 +1461,21 @@ namespace DatasmithRevitExporter
 		{
 			Material CurrentMaterial = GetElement(InMaterialNode.MaterialId) as Material;
 
+			if (InMaterialNode.HasOverriddenAppearance)
+			{
+				ElementId ReferencingDecalId = FDecalMaterial.GetDecalElementId(InMaterialNode);
+
+				if (ReferencingDecalId != ElementId.InvalidElementId)
+				{
+					FDecalMaterial DecalMaterial = FDecalMaterial.Create(InMaterialNode, CurrentMaterial);
+
+					if (DecalMaterial != null)
+					{
+						DecalMaterialsMap.Add(ReferencingDecalId, DecalMaterial);
+					}
+				}
+			}
+
 			CurrentMaterialName = FMaterialData.GetMaterialName(InMaterialNode, CurrentMaterial);
 
 			if (!MaterialDataMap.ContainsKey(CurrentMaterialName) || (CurrentMaterial != null && DirectLink != null && DirectLink.IsMaterialDirty(CurrentMaterial)))
@@ -1655,6 +1689,8 @@ namespace DatasmithRevitExporter
 			HashSet<string> UniqueTextureNameSet
 		)
 		{
+			AddCollectedDecals(InDatasmithScene);
+
 			// Add the collected meshes from the Datasmith mesh dictionary to the Datasmith scene.
 			AddCollectedMeshes(InDatasmithScene);
 
@@ -1999,6 +2035,54 @@ namespace DatasmithRevitExporter
 			return DatasmithMesh;
 		}
 
+		private void AddCollectedDecals(FDatasmithFacadeScene InDatasmithScene)
+		{
+			foreach (KeyValuePair<ElementId, FDecalMaterial> DecalMaterialPair in DecalMaterialsMap)
+			{
+				FElementData DecalElement = null;
+				if (!DecalElementsMap.TryGetValue(DecalMaterialPair.Key, out DecalElement))
+				{
+					continue;
+				}
+
+				FDecalMaterial DecalMaterial = DecalMaterialPair.Value;
+
+				FDatasmithFacadeDecalMaterial DatasmithMaterial = new FDatasmithFacadeDecalMaterial(DecalMaterial.MaterialName);
+
+				if (!string.IsNullOrEmpty(DecalMaterial.DiffuseTexturePath))
+				{
+					FDatasmithFacadeTexture DiffuseTexture = FDatasmithFacadeMaterialsUtils.CreateSimpleTextureElement(DecalMaterial.DiffuseTexturePath);
+					DiffuseTexture.SetSRGB(FDatasmithFacadeTexture.EColorSpace.sRGB);
+					DiffuseTexture.SetTextureMode(FDatasmithFacadeTexture.ETextureMode.Diffuse);
+					DiffuseTexture.SetFile(DecalMaterial.DiffuseTexturePath);
+					DatasmithMaterial.SetDiffuseTexturePathName(DiffuseTexture.GetName());
+					InDatasmithScene.AddTexture(DiffuseTexture);
+				}
+
+				if (!string.IsNullOrEmpty(DecalMaterial.BumpTexturePath))
+				{
+					FDatasmithFacadeTexture BumpTexture = FDatasmithFacadeMaterialsUtils.CreateSimpleTextureElement(DecalMaterial.BumpTexturePath);
+					BumpTexture.SetSRGB(FDatasmithFacadeTexture.EColorSpace.sRGB);
+					BumpTexture.SetTextureMode(FDatasmithFacadeTexture.ETextureMode.Bump);
+					BumpTexture.SetFile(DecalMaterial.BumpTexturePath);
+					DatasmithMaterial.SetDiffuseTexturePathName(BumpTexture.GetName());
+					InDatasmithScene.AddTexture(BumpTexture);
+				}
+
+				InDatasmithScene.AddMaterial(DatasmithMaterial);
+
+				Transform DecalTransform = null;
+				XYZ DecalDimensions = null;
+				FUtils.GetDecalSpatialParams(DecalElement.CurrentElement, ref DecalTransform, ref DecalDimensions);
+
+				FDatasmithFacadeActorDecal DecalActor = DecalElement.ElementActor as FDatasmithFacadeActorDecal;
+				DecalActor.SetDimensions(DecalDimensions.Z, DecalDimensions.X, DecalDimensions.Y);
+				DecalActor.SetDecalMaterialPathName(DatasmithMaterial.GetName());
+
+				SetActorTransform(DecalTransform, DecalActor);
+			}
+		}
+
 		private void AddCollectedMeshes(
 			FDatasmithFacadeScene InDatasmithScene
 		)
diff --git a/Engine/Source/Programs/Enterprise/Datasmith/DatasmithRevitExporter/Private/DatasmithRevitUtils.cs b/Engine/Source/Programs/Enterprise/Datasmith/DatasmithRevitExporter/Private/DatasmithRevitUtils.cs
--- a/Engine/Source/Programs/Enterprise/Datasmith/DatasmithRevitExporter/Private/DatasmithRevitUtils.cs
+++ b/Engine/Source/Programs/Enterprise/Datasmith/DatasmithRevitExporter/Private/DatasmithRevitUtils.cs
@@ -107,5 +107,46 @@ namespace DatasmithRevitExporter
 				}
 			}
 		}
+
+		public static void GetDecalSpatialParams(Element InDecalElement, ref Transform OutDecalTransform, ref XYZ OutDecalDimensions)
+		{
+			List<Line> DecalQuad = new List<Line>();
+
+			GeometryElement GeomElement = InDecalElement.get_Geometry(new Options());
+
+			foreach (GeometryObject GeomObj in GeomElement)
+			{
+				if (GeomObj is Line QuadLine)
+				{
+					DecalQuad.Add(QuadLine);
+				}
+			}
+
+			if (DecalQuad.Count != 4)
+			{
+				return;
+			}
+
+			XYZ TopLeft = DecalQuad[0].Origin;
+			XYZ TopRight = DecalQuad[1].Origin;
+			XYZ BottomRight = DecalQuad[2].Origin;
+			XYZ BottomLeft = DecalQuad[3].Origin;
+
+			XYZ BasisY = (TopRight - TopLeft).Normalize();
+			XYZ BasisZ = (TopLeft - BottomLeft).Normalize();
+			XYZ BasisX = BasisZ.CrossProduct(BasisY).Normalize();
+
+			XYZ Origin = (TopLeft + BottomRight) * 0.5f;
+
+			OutDecalTransform = Transform.Identity;
+			OutDecalTransform.BasisX = BasisX;
+			OutDecalTransform.BasisY = BasisZ;
+			OutDecalTransform.BasisZ = BasisY;
+			OutDecalTransform.Origin = Origin;
+
+			const float CENTIMETERS_PER_FOOT = 30.48F;
+
+			OutDecalDimensions = new XYZ(DecalQuad[0].Length * CENTIMETERS_PER_FOOT * 0.5, DecalQuad[1].Length * CENTIMETERS_PER_FOOT * 0.5, 2.0);
+		}
 	}
 }
diff --git a/Engine/Source/Runtime/Datasmith/DatasmithCore/Private/DatasmithUtils.cpp b/Engine/Source/Runtime/Datasmith/DatasmithCore/Private/DatasmithUtils.cpp
--- a/Engine/Source/Runtime/Datasmith/DatasmithCore/Private/DatasmithUtils.cpp
+++ b/Engine/Source/Runtime/Datasmith/DatasmithCore/Private/DatasmithUtils.cpp
@@ -954,6 +954,18 @@ namespace DatasmithSceneUtilsImpl
 			ScanMaterialIDElement( LightActorElement->GetLightFunctionMaterial().Get() );
 		}
 
+		void ScanDecalActorElement(IDatasmithDecalActorElement* DecalActorElement)
+		{
+			if (const TCHAR* DecalMaterialName = DecalActorElement->GetDecalMaterialPathName())
+			{
+				if (TSharedPtr<IDatasmithElement>* MaterialElementPtr = AssetElementMapping.Find(MaterialPrefix + DecalMaterialName))
+				{
+					TSharedPtr<IDatasmithBaseMaterialElement> MaterialElement = StaticCastSharedPtr<IDatasmithBaseMaterialElement>(*MaterialElementPtr);
+					ReferencedMaterials.Add(MaterialElement);
+				}
+			}
+		}
+
 		void ParseSceneActor( const TSharedPtr<IDatasmithActorElement>& ActorElement )
 		{
 			if (!ActorElement.IsValid())
@@ -971,6 +983,10 @@ namespace DatasmithSceneUtilsImpl
 			{
 				ScanLightActorElement(static_cast<IDatasmithLightActorElement*>(ActorElement.Get()));
 			}
+			else if (ActorElement->IsA(EDatasmithElementType::Decal))
+			{
+				ScanDecalActorElement(static_cast<IDatasmithDecalActorElement*>(ActorElement.Get()));
+			}
 
 			for (int32 Index = 0; Index < ActorElement->GetChildrenCount(); ++Index)
 			{
@@ -1004,6 +1020,24 @@ namespace DatasmithSceneUtilsImpl
 			}
 		}
 
+		void ScanDecalMaterialElement(IDatasmithDecalMaterialElement* MaterialElement)
+		{
+			if (const TCHAR* DiffuseTextureName = MaterialElement->GetDiffuseTexturePathName())
+			{
+				if (DiffuseTextureName[0] != '/')
+				{
+					ReferencedTextures.Add(DiffuseTextureName);
+				}
+			}
+			if (const TCHAR* NormalTextureName = MaterialElement->GetNormalTexturePathName())
+			{
+				if (NormalTextureName[0] != '/')
+				{
+					ReferencedTextures.Add(NormalTextureName);
+				}
+			}
+		}
+
 		void ScanPbrMaterialElement(IDatasmithUEPbrMaterialElement* MaterialElement)
 		{
 			TFunction<void(IDatasmithMaterialExpression*)> ParseExpressionElement;
@@ -1236,6 +1270,10 @@ namespace DatasmithSceneUtilsImpl
 				{
 					ScanMasterMaterialElement(static_cast< IDatasmithMasterMaterialElement* >( MaterialElement.Get() ));
 				}
+				else if (MaterialElement->IsA(EDatasmithElementType::DecalMaterial))
+				{
+					ScanDecalMaterialElement(static_cast<IDatasmithDecalMaterialElement*>(MaterialElement.Get()));
+				}
 				else if ( MaterialElement->IsA( EDatasmithElementType::Material ) )
 				{
 					ScanLegacyMaterialElement(static_cast< IDatasmithMaterialElement* >( MaterialElement.Get() ));
