[3ds Max] DirectLink: Profile Sync/Export

- Export outputs into Messages dialog and log timings for each export stage. Always enabled

- Sync/AutoSync can enable this too, with "Datasmith_SetExportOption_StatSync true" mxs command

- removed info/warnings from outputting to Listener, only Messages and Unreal log

- Unreal log is located here - %LOCALAPPDATA%\UnrealDatasmithExporter\Saved\Logs\UnrealDatasmithExporter.log

diff --git a/Engine/Source/Programs/Enterprise/Datasmith/DatasmithMaxExporter/Private/DatasmithMaxDirectLinkLogging.cpp b/Engine/Source/Programs/Enterprise/Datasmith/DatasmithMaxExporter/Private/DatasmithMaxDirectLinkLogging.cpp
--- a/Engine/Source/Programs/Enterprise/Datasmith/DatasmithMaxExporter/Private/DatasmithMaxDirectLinkLogging.cpp
+++ b/Engine/Source/Programs/Enterprise/Datasmith/DatasmithMaxExporter/Private/DatasmithMaxDirectLinkLogging.cpp
@@ -40,12 +40,15 @@ void LogDebug(const TCHAR* Msg)
 
 void LogInfo(const TCHAR* Msg)
 {
-	mprintf(L"[%s]%s\n", *FDateTime::UtcNow().ToString(TEXT("%Y.%m.%d-%H.%M.%S:%s")), Msg);
-	UE_LOG(LogDatasmithMaxExporter, Error, TEXT("%s"), Msg);
+	UE_LOG(LogDatasmithMaxExporter, Display, TEXT("%s"), Msg);
+	LogInfoDialog(Msg);
 }
 
-
-
+void LogWarning(const TCHAR* Msg)
+{
+	UE_LOG(LogDatasmithMaxExporter, Warning, TEXT("%s"), Msg);
+	LogWarningDialog(Msg);
+}
 
 void LogFlush()
 {
@@ -68,6 +71,10 @@ void LogInfo(const FString& Msg)
 	LogInfo(*Msg);
 }
 
+void LogWarning(const FString& Msg)
+{
+	LogWarning(*Msg);
+}
 void LogDebugNode(const FString& Name, INode* Node)
 {
 #ifdef LOG_DEBUG_HEAVY_ENABLE
diff --git a/Engine/Source/Programs/Enterprise/Datasmith/DatasmithMaxExporter/Private/DatasmithMaxDirectLinkMaxScript.cpp b/Engine/Source/Programs/Enterprise/Datasmith/DatasmithMaxExporter/Private/DatasmithMaxDirectLinkMaxScript.cpp
--- a/Engine/Source/Programs/Enterprise/Datasmith/DatasmithMaxExporter/Private/DatasmithMaxDirectLinkMaxScript.cpp
+++ b/Engine/Source/Programs/Enterprise/Datasmith/DatasmithMaxExporter/Private/DatasmithMaxDirectLinkMaxScript.cpp
@@ -268,6 +268,8 @@ Primitive SetExportOption_##name##_pf(_M("Datasmith_SetExportOption_" #name), Se
 //////////////////////////////////////////
 DefinePersistentExportOption(SelectedOnly)
 DefinePersistentExportOption(AnimatedTransforms)
+
+DefinePersistentExportOption(StatSync)
 //////////////////////////////////////////
 
 #undef DefinePersistentExportOption
@@ -414,6 +416,11 @@ public:
 	}
 
 	void AddWarning(const TCHAR* Message)
+	{
+		Messages.Enqueue(FString(TEXT("WARNING:")) + Message);
+	}
+
+	void AddInfo(const TCHAR* Message)
 	{
 		Messages.Enqueue(Message);
 	}
@@ -428,6 +435,11 @@ public:
 
 TUniquePtr<FMessagesDialog> MessagesDialog;
 
+void LogInfoDialog(const TCHAR* Msg)
+{
+	MessagesDialog->AddInfo(Msg);
+}
+
 void LogWarningDialog(const FString& Msg)
 {
 	LogWarningDialog(*Msg);
diff --git a/Engine/Source/Programs/Enterprise/Datasmith/DatasmithMaxExporter/Private/DatasmithMaxDirectLinkMesh.cpp b/Engine/Source/Programs/Enterprise/Datasmith/DatasmithMaxExporter/Private/DatasmithMaxDirectLinkMesh.cpp
--- a/Engine/Source/Programs/Enterprise/Datasmith/DatasmithMaxExporter/Private/DatasmithMaxDirectLinkMesh.cpp
+++ b/Engine/Source/Programs/Enterprise/Datasmith/DatasmithMaxExporter/Private/DatasmithMaxDirectLinkMesh.cpp
@@ -444,7 +444,7 @@ bool ConvertMaxMeshToDatasmith(ISceneTracker& Scene, TSharedPtr<IDatasmithMeshEl
 		}
 		else
 		{
-			LogWarningDialog(FString(TEXT("Invalid object: ")) + Node->GetName());
+			LogWarning(FString(TEXT("Invalid object: ")) + Node->GetName());
 			return false;
 		}
 	}
@@ -453,7 +453,7 @@ bool ConvertMaxMeshToDatasmith(ISceneTracker& Scene, TSharedPtr<IDatasmithMeshEl
 	TMap<int32, int32> UVChannelsMap;
 	if (!CreateDatasmithMeshFromMaxMesh(DatasmithMesh, Node, MeshName, RenderMesh, bConsolidateMaterialIds, SupportedChannels, UVChannelsMap))
 	{
-		LogWarningDialog(FString(TEXT("Invalid object: ")) + Node->GetName());
+		LogWarning(FString(TEXT("Invalid object: ")) + Node->GetName());
 		return false;
 	}
 
@@ -471,7 +471,7 @@ bool ConvertMaxMeshToDatasmith(ISceneTracker& Scene, TSharedPtr<IDatasmithMeshEl
 		}
 		else if (SelectedLightmapUVChannel != DefaultValue)
 		{
-			LogWarningDialog(*FString::Printf(TEXT("%s won't use the channel %i for its lightmap because it's not supported by the mesh. A new channel will be generated.")
+			LogWarning(*FString::Printf(TEXT("%s won't use the channel %i for its lightmap because it's not supported by the mesh. A new channel will be generated.")
 				, static_cast<const TCHAR*>(Node->GetName())
 				, SelectedLightmapUVChannel));
 		}
diff --git a/Engine/Source/Programs/Enterprise/Datasmith/DatasmithMaxExporter/Private/DatasmithMaxExporterWithDirectLink.cpp b/Engine/Source/Programs/Enterprise/Datasmith/DatasmithMaxExporter/Private/DatasmithMaxExporterWithDirectLink.cpp
--- a/Engine/Source/Programs/Enterprise/Datasmith/DatasmithMaxExporter/Private/DatasmithMaxExporterWithDirectLink.cpp
+++ b/Engine/Source/Programs/Enterprise/Datasmith/DatasmithMaxExporter/Private/DatasmithMaxExporterWithDirectLink.cpp
@@ -198,15 +198,56 @@ public:
 	bool bIsInvalidated = true;
 };
 
-
-
 class FUpdateProgress
 {
-	TUniquePtr<FDatasmithMaxProgressManager> ProgressManager;
-	int32 StageIndex = 0;
-	int32 StageCount;
 public:
-	FUpdateProgress(bool bShowProgressBar, int32 InStageCount) : StageCount(InStageCount)
+	class FStage
+	{
+	public:
+		FUpdateProgress& UpdateProgress;
+
+		FStage(FUpdateProgress& InUpdateProgress, const TCHAR* InName, int32 InStageCount) : UpdateProgress(InUpdateProgress), Name(InName), StageCount(InStageCount)
+		{
+			TimeStart =  FDateTime::UtcNow();
+		}
+
+		void Finished()
+		{
+			TimeFinish =  FDateTime::UtcNow();
+		}
+
+		FStage& ProgressStage(const TCHAR* SubstageName, int32 InStageCount)
+		{
+			LogDebug(SubstageName);
+			if (UpdateProgress.ProgressManager)
+			{
+				StageIndex++;
+				UpdateProgress.ProgressManager->SetMainMessage(*FString::Printf(TEXT("%s (%d of %d)"), SubstageName, StageIndex, StageCount));
+				UpdateProgress.ProgressManager->ProgressEvent(0, TEXT(""));
+			}
+			return Stages.Emplace_GetRef(UpdateProgress, SubstageName, InStageCount);
+		}
+
+		void ProgressEvent(float Progress, const TCHAR* Message)
+		{
+			LogDebug(FString::Printf(TEXT("%f %s"), Progress, Message));
+			if (UpdateProgress.ProgressManager)
+			{
+				UpdateProgress.ProgressManager->ProgressEvent(Progress, Message);
+			}
+		}
+
+		FString Name;
+		int32 StageCount;
+		int32 StageIndex = 0;
+
+		FDateTime TimeStart;
+		FDateTime TimeFinish;
+
+		TArray<FStage> Stages;
+	};
+
+	FUpdateProgress(bool bShowProgressBar, int32 InStageCount) : MainStage(*this, TEXT("Total"), InStageCount)
 	{
 		if (bShowProgressBar)
 		{
@@ -214,33 +255,37 @@ public:
 		}
 	}
 
-	void ProgressStage(const TCHAR* Name)
+	void PrintStatisticss()
 	{
-		LogDebug(Name);
-		if (ProgressManager)
-		{
-			StageIndex++;
-			ProgressManager->SetMainMessage(*FString::Printf(TEXT("%s (%d of %d)"), Name, StageIndex, StageCount));
-			ProgressManager->ProgressEvent(0, TEXT(""));
-		}
+		LogInfo(TEXT("Timings:"));
+		PrintStage(MainStage);
 	}
 
-	void ProgressEvent(float Progress, const TCHAR* Message)
+	void PrintStage(FStage& Stage, FString Indent=TEXT(""))
 	{
-		LogDebug(FString::Printf(TEXT("%f %s"), Progress, Message));
-		if (ProgressManager)
+		LogInfo(Indent + FString::Printf(TEXT("    %s - %s"), *Stage.Name, *(Stage.TimeFinish-Stage.TimeStart).ToString()));
+		for(FStage& ChildStage: Stage.Stages)
 		{
-			ProgressManager->ProgressEvent(Progress, Message);
+			PrintStage(ChildStage, Indent + TEXT("  "));
 		}
 	}
+
+	void Finished()
+	{
+		MainStage.Finished();
+	}
+
+	TUniquePtr<FDatasmithMaxProgressManager> ProgressManager;
+
+	FStage MainStage;
 };
 
 class FProgressCounter
 {
 public:
 
-	FProgressCounter(FUpdateProgress& InProgressManager, int32 InCount)
-		: ProgressManager(InProgressManager)
+	FProgressCounter(FUpdateProgress::FStage& InProgressStage, int32 InCount)
+		: ProgressStage(InProgressStage)
 		, Count(InCount)
 		, SecondsOfLastUpdate(FPlatformTime::Seconds())
 	{
@@ -251,19 +296,36 @@ public:
 		double CurrentTime = FPlatformTime::Seconds();
 		if (CurrentTime - SecondsOfLastUpdate > UpdateIntervalMin) // Don't span progress bar
 		{
-			ProgressManager.ProgressEvent(float(Index) / Count, *FString::Printf(TEXT("%d of %d"), Index, Count) );
+			ProgressStage.ProgressEvent(float(Index) / Count, *FString::Printf(TEXT("%d of %d"), Index, Count) );
 			SecondsOfLastUpdate = CurrentTime;
 		}
 		Index++;
 	}
 private:
-	FUpdateProgress& ProgressManager;
+	FUpdateProgress::FStage& ProgressStage;
 	int32 Count;
 	int32 Index = 0;
 	const double UpdateIntervalMin = 0.05; // Don't update progress it last update was just recently
 	double SecondsOfLastUpdate;
 };
 
+class FProgressStageGuard
+{
+public:
+	FUpdateProgress::FStage& Stage;
+	FProgressStageGuard(FUpdateProgress::FStage& ParentStage, const TCHAR* InName, int32 Count=0) : Stage(ParentStage.ProgressStage(InName, Count))
+	{
+	}
+
+	~FProgressStageGuard()
+	{
+		Stage.Finished();
+	}
+};
+
+#define PROGRESS_STAGE(Name) FProgressStageGuard ProgressStage(MainStage, TEXT(Name));
+#define PROGRESS_STAGE_COUNTER(Count) FProgressCounter ProgressCounter(ProgressStage.Stage, Count);
+
 // Convert various node data to Datasmith tags
 class FTagsConverter
 {
@@ -390,13 +452,14 @@ struct FExportOptions
 	// Default options for DirectLink change-tracking
 	bool bSelectedOnly = false;
 	bool bAnimatedTransforms = false;
+
+	bool bStatSync = false;
 };
 
 // Global export options, stored in preferences
 class FPersistentExportOptions: public IPersistentExportOptions
 {
 public:
-
 	void Load()
 	{
 		if (bLoaded)
@@ -456,6 +519,17 @@ public:
 		return FPaths::Combine(PlugCfgPath, TEXT("UnrealDatasmithMax.ini"));
 	}
 
+	virtual void SetStatSync(bool bValue) override
+	{
+		Options.bStatSync = bValue;
+		SetBool(TEXT("StatExport"), bValue);
+	}
+
+	virtual bool GetStatSync() override
+	{
+		return Options.bStatSync;
+	}
+
 
 	FExportOptions Options;
 	bool bLoaded = false;
@@ -614,7 +688,7 @@ public:
 
 			TUniquePtr<FLayerTracker>& LayerTracker = LayersForAnimHandle.FindOrAdd(Handle);
 
-			BOOL bIsHidden = Layer->IsHidden(TRUE);
+			bool bIsHidden = Layer->IsHidden(TRUE);
 			FString Name = Layer->GetName().data();
 
 			if (!LayerTracker)
@@ -642,7 +716,7 @@ public:
 	}
 
 	// Applies all recorded changes to Datasmith scene
-	bool Update(bool bQuiet, bool bRenderQuality)
+	bool Update(FUpdateProgress::FStage& ProgressStage, bool bRenderQuality)
 	{
 		// Disable Undo, editing, redraw, messages during export/sync so that nothing changes the scene
 		GetCOREInterface()->EnableUndo(false);
@@ -661,7 +735,7 @@ public:
 		NodesPreparer.Start(GetCOREInterface()->GetTime(), bRenderQuality);
 
 		bUpdateInProgress = true;
-		bool bResult = UpdateInternalSafe(bQuiet);
+		bool bResult = UpdateInternalSafe(ProgressStage);
 		bUpdateInProgress = false;
 
 		NodesPreparer.Finish();
@@ -673,22 +747,23 @@ public:
 		return bResult;
 	}
 
-	bool UpdateInternalSafe(bool bQuiet)
+	bool UpdateInternalSafe(FUpdateProgress::FStage& ProgressStage)
 	{
 		__try
 		{
-			return UpdateInternal(bQuiet);
+			return UpdateInternal(ProgressStage);
 		}
 		__except(EXCEPTION_EXECUTE_HANDLER)
 		{
-			LogInfo(TEXT("Update finished with exception"));
+			LogWarning(TEXT("Update finished with exception"));
 		}
 		return false;
 	}
 
-	bool UpdateInternal(bool bQuiet)
+	bool UpdateInternal(FUpdateProgress::FStage& ParentStage)
 	{
-		FUpdateProgress ProgressManager(!bQuiet, 6); // Will shutdown on end of Update
+		FProgressStageGuard MainProgressStageGuard(ParentStage, TEXT("Update"), 10);
+		FUpdateProgress::FStage& MainStage = MainProgressStageGuard.Stage;
 
 		bool bChangeEncountered = false;
 
@@ -697,8 +772,8 @@ public:
 			ParseScene();
 		}
 
-		ProgressManager.ProgressStage(TEXT("Refresh layers"));
 		{
+			PROGRESS_STAGE("Refresh layers")
 			bChangeEncountered = UpdateLayers() && bChangeEncountered;
 		}
 
@@ -707,8 +782,8 @@ public:
 		bChangeEncountered |= !MaterialsCollectionTracker.GetInvalidatedMaterials().IsEmpty();
 
 
-		ProgressManager.ProgressStage(TEXT("Remove deleted nodes"));
 		{
+			PROGRESS_STAGE("Remove deleted nodes")
 			TArray<FNodeTracker*> DeletedNodeTrackers;
 			for (FNodeTracker* NodeTracker : InvalidatedNodeTrackers)
 			{
@@ -724,22 +799,25 @@ public:
 			}
 		}
 
-		ProgressManager.ProgressStage(TEXT("Update node names"));
 		// todo: move to NameChanged and NodeAdded?
-		for (FNodeTracker* NodeTracker : InvalidatedNodeTrackers)
 		{
-			FString Name = NodeTracker->Node->GetName();
-			if (Name != NodeTracker->Name)
+			PROGRESS_STAGE("Update node names")
+			for (FNodeTracker* NodeTracker : InvalidatedNodeTrackers)
 			{
-				NodeTrackersNames[NodeTracker->Name].Remove(NodeTracker);
-				NodeTracker->Name = Name;
-				NodeTrackersNames.FindOrAdd(NodeTracker->Name).Add(NodeTracker);
+				FString Name = NodeTracker->Node->GetName();
+				if (Name != NodeTracker->Name)
+				{
+					NodeTrackersNames[NodeTracker->Name].Remove(NodeTracker);
+					NodeTracker->Name = Name;
+					NodeTrackersNames.FindOrAdd(NodeTracker->Name).Add(NodeTracker);
+				}
 			}
 		}
 
-		ProgressManager.ProgressStage(TEXT("Refresh collisions")); // Update set of nodes used for collision 
 		{
-			FProgressCounter ProgressCounter(ProgressManager, InvalidatedNodeTrackers.Num());
+			PROGRESS_STAGE("Refresh collisions") // Update set of nodes used for collision
+			PROGRESS_STAGE_COUNTER(InvalidatedNodeTrackers.Num());
+			
 			TSet<FNodeTracker*> NodesWithChangedCollisionStatus; // Need to invalidate these nodes to make them renderable or to keep them from renderable depending on collision status
 			for (FNodeTracker* NodeTracker : InvalidatedNodeTrackers)
 			{
@@ -749,9 +827,9 @@ public:
 			InvalidatedNodeTrackers.Append(NodesWithChangedCollisionStatus);
 		}
 
-		ProgressManager.ProgressStage(TEXT("Process invalidated nodes"));
 		{
-			FProgressCounter ProgressCounter(ProgressManager, InvalidatedNodeTrackers.Num());
+			PROGRESS_STAGE("Process invalidated nodes")
+			PROGRESS_STAGE_COUNTER(InvalidatedNodeTrackers.Num());
 			for (FNodeTracker* NodeTracker : InvalidatedNodeTrackers)
 			{
 				ProgressCounter.Next();
@@ -759,9 +837,9 @@ public:
 			}
 		}
 		
-		ProgressManager.ProgressStage(TEXT("Process invalidated instances"));
 		{
-			FProgressCounter ProgressCounter(ProgressManager, InvalidatedInstances.Num());
+			PROGRESS_STAGE("Process invalidated instances")
+			PROGRESS_STAGE_COUNTER(InvalidatedInstances.Num());
 			for (FInstances* Instances : InvalidatedInstances)
 			{
 				ProgressCounter.Next();
@@ -771,41 +849,45 @@ public:
 			InvalidatedInstances.Reset();
 		}
 
-		ProgressManager.ProgressStage(TEXT("Reparent Datasmith Actors"));
-		for (FNodeTracker* NodeTracker : InvalidatedNodeTrackers)
 		{
-			AttachNodeToDatasmithScene(*NodeTracker);
+			PROGRESS_STAGE("Reparent Datasmith Actors");
+			for (FNodeTracker* NodeTracker : InvalidatedNodeTrackers)
+			{
+				AttachNodeToDatasmithScene(*NodeTracker);
+			}
+			InvalidatedNodeTrackers.Reset();
 		}
-		InvalidatedNodeTrackers.Reset();
 
 		TSet<Mtl*> ActualMaterialToUpdate;
 		TSet<Texmap*> ActualTexmapsToUpdate;
 
-		ProgressManager.ProgressStage(TEXT("Process invalidated materials"));
 		{
-			FProgressCounter ProgressCounter(ProgressManager, MaterialsCollectionTracker.GetInvalidatedMaterials().Num());
-			for (FMaterialTracker* MaterialTracker : MaterialsCollectionTracker.GetInvalidatedMaterials())
+			PROGRESS_STAGE("Process invalidated materials");
 			{
-				ProgressCounter.Next();
+				PROGRESS_STAGE_COUNTER(MaterialsCollectionTracker.GetInvalidatedMaterials().Num());
+				for (FMaterialTracker* MaterialTracker : MaterialsCollectionTracker.GetInvalidatedMaterials())
+				{
+					ProgressCounter.Next();
 
-				MaterialsCollectionTracker.UpdateMaterial(MaterialTracker);
+					MaterialsCollectionTracker.UpdateMaterial(MaterialTracker);
 
-				for (Mtl* ActualMaterial : MaterialTracker->GetActualMaterials())
-				{
-					ActualMaterialToUpdate.Add(ActualMaterial);
-				}
-				MaterialTracker->bInvalidated = false;
-				for (Texmap* Texture: MaterialTracker->Textures)
-				{
-					ActualTexmapsToUpdate.Add(Texture);
+					for (Mtl* ActualMaterial : MaterialTracker->GetActualMaterials())
+					{
+						ActualMaterialToUpdate.Add(ActualMaterial);
+					}
+					MaterialTracker->bInvalidated = false;
+					for (Texmap* Texture: MaterialTracker->Textures)
+					{
+						ActualTexmapsToUpdate.Add(Texture);
+					}
 				}
+				MaterialsCollectionTracker.ResetInvalidatedMaterials();
 			}
-			MaterialsCollectionTracker.ResetInvalidatedMaterials();
 		}
 
-		ProgressManager.ProgressStage(TEXT("Update textures"));
 		{
-			FProgressCounter ProgressCounter(ProgressManager, ActualTexmapsToUpdate.Num());
+			PROGRESS_STAGE("Update textures")
+			PROGRESS_STAGE_COUNTER(ActualTexmapsToUpdate.Num());
 			for (Texmap* Texture : ActualTexmapsToUpdate)
 			{
 				ProgressCounter.Next();
@@ -813,9 +895,9 @@ public:
 			}
 		}
 
-		ProgressManager.ProgressStage(TEXT("Update materials"));
 		{
-			FProgressCounter ProgressCounter(ProgressManager, ActualMaterialToUpdate.Num());
+			PROGRESS_STAGE("Update materials")
+			PROGRESS_STAGE_COUNTER(ActualMaterialToUpdate.Num());
 			for (Mtl* ActualMaterial : ActualMaterialToUpdate)
 			{
 				ProgressCounter.Next();
@@ -1675,7 +1757,7 @@ public:
 				}
 				else
 				{
-					LogWarningDialog(TEXT("Unsupported light: ") + DatasmithMaxLogger::Get().GetLightDescription(NodeTracker.Node));
+					LogWarning(TEXT("Unsupported light: ") + DatasmithMaxLogger::Get().GetLightDescription(NodeTracker.Node));
 				}
 				return false;
 			}
@@ -2149,7 +2231,17 @@ public:
 
 	virtual bool UpdateScene(bool bQuiet) override
 	{
-		return SceneTracker.Update(bQuiet, false);
+		FUpdateProgress ProgressManager(!bQuiet, 1);
+
+		bool bResult = SceneTracker.Update(ProgressManager.MainStage, false);
+
+		ProgressManager.Finished();
+
+		if (Options.bStatSync)
+		{
+			ProgressManager.PrintStatisticss();
+		}
+		return bResult;
 	}
 
 	virtual void ResetSceneTracking() override
@@ -2172,7 +2264,7 @@ public:
 		return SceneTracker;
 	}
 
-	FExportOptions Options;
+	FExportOptions& Options;
 
 	FDatasmith3dsMaxScene ExportedScene;
 	TUniquePtr<FDatasmithDirectLink> DirectLinkImpl;
@@ -2209,11 +2301,8 @@ bool CreateExporter(bool bEnableUI, const TCHAR* EnginePath)
 		return false;
 	}
 
-	static FExportOptions ExporterOptions; // Default options
-	Exporter = MakeUnique<FExporter>(ExporterOptions);
-
 	PersistentExportOptions.Load(); // Access GConfig only after FDatasmithExporterManager::Initialize finishes, which ensures that Unreal game thread was initialized(GConfig is created there)
-
+	Exporter = MakeUnique<FExporter>(PersistentExportOptions.Options);
 	return true;
 }
 
@@ -2244,20 +2333,32 @@ void FExporter::Shutdown()
 
 bool Export(const TCHAR* Name, const TCHAR* OutputPath, bool bQuiet)
 {
+	FUpdateProgress ProgressManager(!bQuiet, 3);
+	FUpdateProgress::FStage& MainStage = ProgressManager.MainStage;
+
 	FDatasmith3dsMaxScene ExportedScene;
 	ExportedScene.SetupScene();
 	ExportedScene.SetName(Name);
 	ExportedScene.SetOutputPath(OutputPath);
 
 	FSceneTracker SceneTracker(PersistentExportOptions.Options, ExportedScene, nullptr);
-	SceneTracker.Update(bQuiet, true);
+
+	SceneTracker.Update(MainStage, true);
 
 	if (PersistentExportOptions.Options.bAnimatedTransforms)
 	{
+		PROGRESS_STAGE("Export Animations");
 		SceneTracker.ExportAnimations();
 	}
 
-	ExportedScene.GetSceneExporter().Export(ExportedScene.GetDatasmithScene(), false);
+	{
+		PROGRESS_STAGE("Save Datasmith Scene");
+		ExportedScene.GetSceneExporter().Export(ExportedScene.GetDatasmithScene(), false);
+	}
+
+	ProgressManager.Finished();
+
+	ProgressManager.PrintStatisticss();
 
 	return true;
 }
diff --git a/Engine/Source/Programs/Enterprise/Datasmith/DatasmithMaxExporter/Public/DatasmithMaxDirectLink.h b/Engine/Source/Programs/Enterprise/Datasmith/DatasmithMaxExporter/Public/DatasmithMaxDirectLink.h
--- a/Engine/Source/Programs/Enterprise/Datasmith/DatasmithMaxExporter/Public/DatasmithMaxDirectLink.h
+++ b/Engine/Source/Programs/Enterprise/Datasmith/DatasmithMaxExporter/Public/DatasmithMaxDirectLink.h
@@ -48,6 +48,10 @@ public:
 	// Whether to export animation or not
 	virtual void SetAnimatedTransforms(bool) = 0;
 	virtual bool GetAnimatedTransforms() = 0;
+
+	// Whether to output export statistics to listener/log
+	virtual void SetStatSync(bool) = 0;
+	virtual bool GetStatSync() = 0;
 };
 
 //---- Main class for export/change tracking
@@ -452,9 +456,13 @@ void LogInfo(const FString& Msg);
 void LogDebugNode(const FString& Name, class INode* Node);
 void LogNodeEvent(const MCHAR* Name, INodeEventCallback::NodeKeyTab& nodes);
 void LogFlush();
+void LogWarning(const TCHAR* Msg);
+void LogWarning(const FString& Msg);
 
 void LogWarningDialog(const TCHAR* Msg);
 void LogWarningDialog(const FString& Msg);
+void LogInfoDialog(const TCHAR* Msg);
+
 
 }
 
