[3ds Max] DirectLink: CoronaPhysicalMtl ClearCoat normal

diff --git a/Engine/Source/Programs/Enterprise/Datasmith/DatasmithMaxExporter/Private/MaxMaterialsToUEPbr/DatasmithMaxCoronaMaterialsToUEPbr.cpp b/Engine/Source/Programs/Enterprise/Datasmith/DatasmithMaxExporter/Private/MaxMaterialsToUEPbr/DatasmithMaxCoronaMaterialsToUEPbr.cpp
--- a/Engine/Source/Programs/Enterprise/Datasmith/DatasmithMaxExporter/Private/MaxMaterialsToUEPbr/DatasmithMaxCoronaMaterialsToUEPbr.cpp
+++ b/Engine/Source/Programs/Enterprise/Datasmith/DatasmithMaxExporter/Private/MaxMaterialsToUEPbr/DatasmithMaxCoronaMaterialsToUEPbr.cpp
@@ -392,6 +392,7 @@ namespace DatasmithMaxCoronaMaterialsToUEPbrImpl
 		DatasmithMaxTexmapParser::FMapParameter TranslucencyFractionTexmap;
 		FLinearColor ThinAbsorptionColor;
 		DatasmithMaxTexmapParser::FMapParameter ThinAbsorptionTexmap;
+		DatasmithMaxTexmapParser::FMapParameter ClearcoatBumpTexmap;
 		DatasmithMaxTexmapParser::FMapParameter MetalnessTexmap;
 		int RoughnessMode;
 		FLinearColor TranslucencyColor;
@@ -670,6 +671,18 @@ namespace DatasmithMaxCoronaMaterialsToUEPbrImpl
 					{
 						ThinAbsorptionTexmap.Weight = ParamBlock2->GetFloat(ParamDefinition.ID, CurrentTime);
 					}
+					else if (FCString::Stricmp(ParamDefinition.int_name, TEXT("clearcoatBumpTexmap")) == 0)
+					{
+						ClearcoatBumpTexmap.Map = ParamBlock2->GetTexmap(ParamDefinition.ID, GetCOREInterface()->GetTime());
+					}
+					else if (FCString::Stricmp(ParamDefinition.int_name, TEXT("clearcoatBumpTexmapOn")) == 0)
+					{
+						ClearcoatBumpTexmap.bEnabled = ParamBlock2->GetInt(ParamDefinition.ID, CurrentTime) != 0;
+					}
+					else if (FCString::Stricmp(ParamDefinition.int_name, TEXT("clearcoatBumpMapAmount")) == 0)
+					{
+						ClearcoatBumpTexmap.Weight = ParamBlock2->GetFloat(ParamDefinition.ID, CurrentTime);
+					}
 					else if (FCString::Stricmp(ParamDefinition.int_name, TEXT("metalnessTexmap")) == 0)
 					{
 						MetalnessTexmap.Map = ParamBlock2->GetTexmap(ParamDefinition.ID, GetCOREInterface()->GetTime());
@@ -1431,31 +1444,31 @@ void FDatasmithMaxCoronaPhysicalMaterialToUEPbr::Convert(TSharedRef<IDatasmithSc
 
 
 	// Clear Coat
-	if (MaterialProperties.ClearcoatAmount > 0 || MaterialProperties.ClearcoatAmountTexmap.IsMapPresentAndEnabled())
+	bool bHasClearCoat = MaterialProperties.ClearcoatAmount > 0 || MaterialProperties.ClearcoatAmountTexmap.IsMapPresentAndEnabled();
+	if (bHasClearCoat)
 	{
 		IDatasmithMaterialExpression* ClearCoatExpression = TextureOrScalar(TEXT("Clear Coat Amount"), MaterialProperties.ClearcoatAmountTexmap, MaterialProperties.ClearcoatAmount);
 
 		if (ClearCoatExpression)
 		{
-			if (MaterialProperties.ClearcoatIorTexmap.IsMapPresentAndEnabled())
-			{
-				IDatasmithMaterialExpression* ClearCoatIorExpression = TextureOrScalar(TEXT("Clear Coat Ior"), MaterialProperties.ClearcoatIorTexmap, MaterialProperties.ClearcoatIor);
-				Connect(PbrMaterialElement->GetClearCoat(), CalcIORComplex(*ClearCoatIorExpression, Scalar(0), *ClearCoatExpression, Multiply(*ClearCoatExpression, Scalar(0.1))));
-			}
-			else
-			{
-				Connect(PbrMaterialElement->GetClearCoat(), CalcIORComplex(MaterialProperties.ClearcoatIor, 0, *ClearCoatExpression, Multiply(*ClearCoatExpression, Scalar(0.1))));
-			}
+			IDatasmithMaterialExpression& ClearCoatAmountByFresnelExpression = MaterialProperties.ClearcoatIorTexmap.IsMapPresentAndEnabled()
+				? CalcIORComplex(*TextureOrScalar(TEXT("Clear Coat Ior"), MaterialProperties.ClearcoatIorTexmap, MaterialProperties.ClearcoatIor), Scalar(0), *ClearCoatExpression, Multiply(*ClearCoatExpression, Scalar(0.1)))
+				: CalcIORComplex(MaterialProperties.ClearcoatIor, 0, *ClearCoatExpression, Multiply(*ClearCoatExpression, Scalar(0.1)));
+
+			// Don't use Ior in Path Tracer(doesn't work well). May implement when Path Tracer implements CLear Coat Ior input
+			Connect(PbrMaterialElement->GetClearCoat(), PathTracingQualitySwitch(ClearCoatAmountByFresnelExpression, *ClearCoatExpression));
 		}
+		
 		Connect(PbrMaterialElement->GetClearCoatRoughness(), TextureOrScalar(TEXT("Clear Coat Roughness"), MaterialProperties.ClearcoatRoughnessTexmap, MaterialProperties.ClearcoatRoughness));
 		PbrMaterialElement->SetShadingModel(EDatasmithShadingModel::ClearCoat);
 	}
 
+	// Convert all bump/normal maps
+	ConvertState.DefaultTextureMode = EDatasmithTextureMode::Bump; // Will change to normal if we pass through a normal map texmap
+	ConvertState.bCanBake = false; // Current baking fails to produce proper normal maps
+
 	if (MaterialProperties.BaseBumpTexmap.IsMapPresentAndEnabled())
 	{
-		ConvertState.DefaultTextureMode = EDatasmithTextureMode::Bump; // Will change to normal if we pass through a normal map texmap
-		ConvertState.bCanBake = false; // Current baking fails to produce proper normal maps
-	
 		IDatasmithMaterialExpression* BumpExpression = FDatasmithMaxTexmapToUEPbrUtils::MapOrValue(this, MaterialProperties.BaseBumpTexmap, PARAM_NAME_BUMPMAP, TOptional<FLinearColor>(), TOptional<float>());
 	
 		if ( BumpExpression )
@@ -1463,8 +1476,19 @@ void FDatasmithMaxCoronaPhysicalMaterialToUEPbr::Convert(TSharedRef<IDatasmithSc
 			BumpExpression->ConnectExpression(PbrMaterialElement->GetNormal());
 			BumpExpression->SetName(PARAM_NAME_BUMPMAP);
 		}
+	}
 
-		ConvertState.bCanBake = true;
+	if (bHasClearCoat)
+	{
+		IDatasmithMaterialExpression* ClearCoatBumpExpression = FDatasmithMaxTexmapToUEPbrUtils::MapOrValue(this, MaterialProperties.ClearcoatBumpTexmap, PARAM_NAME_BUMPMAP, TOptional<FLinearColor>(), TOptional<float>());
+
+		if (ClearCoatBumpExpression)
+		{
+			// Clear Coat uses additional output
+			IDatasmithMaterialExpressionGeneric* MultiplyDiffuseLevelExpression = PbrMaterialElement->AddMaterialExpression< IDatasmithMaterialExpressionGeneric >();
+			MultiplyDiffuseLevelExpression->SetExpressionName( TEXT("ClearCoatNormalCustomOutput") );
+			Connect(*MultiplyDiffuseLevelExpression->GetInput(0), ClearCoatBumpExpression);
+		}
 	}
 	
 	MaterialElement = PbrMaterialElement;
