[3ds Max] DirectLink: Stats output specifies kind of Export/Sync/AutoSync

diff --git a/Engine/Source/Programs/Enterprise/Datasmith/DatasmithMaxExporter/Private/DatasmithMaxDirectLinkMaxScript.cpp b/Engine/Source/Programs/Enterprise/Datasmith/DatasmithMaxExporter/Private/DatasmithMaxDirectLinkMaxScript.cpp
--- a/Engine/Source/Programs/Enterprise/Datasmith/DatasmithMaxExporter/Private/DatasmithMaxDirectLinkMaxScript.cpp
+++ b/Engine/Source/Programs/Enterprise/Datasmith/DatasmithMaxExporter/Private/DatasmithMaxDirectLinkMaxScript.cpp
@@ -540,8 +540,7 @@ public:
 			{
 				if (GetExporter())
 				{
-					GetExporter()->UpdateScene(false);
-					GetExporter()->UpdateDirectLinkScene();
+					GetExporter()->PerformSync(false);
 				}
 				return true;
 			}
diff --git a/Engine/Source/Programs/Enterprise/Datasmith/DatasmithMaxExporter/Private/DatasmithMaxExporterWithDirectLink.cpp b/Engine/Source/Programs/Enterprise/Datasmith/DatasmithMaxExporter/Private/DatasmithMaxExporterWithDirectLink.cpp
--- a/Engine/Source/Programs/Enterprise/Datasmith/DatasmithMaxExporter/Private/DatasmithMaxExporterWithDirectLink.cpp
+++ b/Engine/Source/Programs/Enterprise/Datasmith/DatasmithMaxExporter/Private/DatasmithMaxExporterWithDirectLink.cpp
@@ -201,7 +201,7 @@ public:
 class FUpdateProgress
 {
 public:
-	class FStage
+	class FStage: FNoncopyable
 	{
 	public:
 		FUpdateProgress& UpdateProgress;
@@ -257,7 +257,6 @@ public:
 
 	void PrintStatisticss()
 	{
-		LogInfo(TEXT("Timings:"));
 		PrintStage(MainStage);
 	}
 
@@ -2177,12 +2176,7 @@ public:
 
 			if (IdlePeriod > FMath::RoundToInt(AutoSyncIdleDelaySeconds*1000))
 			{
-				// Don't create progress bar for autosync - it steals focus, closes listener and what else
-				// todo: consider creating progress when a big change in scene is detected, e.g. number of nodes?
-				if (UpdateScene(true)) // Don't sent redundant update if scene change wasn't detected
-				{
-					UpdateDirectLinkScene();
-				}
+				PerformAutoSync();
 			}
 		}
 	}
@@ -2201,8 +2195,7 @@ public:
 		else
 		{
 			// Perform full Sync when AutoSync is first enabled
-			UpdateScene(false);
-			UpdateDirectLinkScene();
+			PerformSync(false);
 
 			const uint32 AutoSyncCheckIntervalMs = FMath::RoundToInt(AutoSyncDelaySeconds*1000);
 			SetTimer(GetCOREInterface()->GetMAXHWnd(), reinterpret_cast<UINT_PTR>(this), AutoSyncCheckIntervalMs, AutoSyncTimerProc);
@@ -2244,6 +2237,50 @@ public:
 		return bResult;
 	}
 
+	virtual void PerformSync(bool bQuiet) override
+	{
+		FUpdateProgress ProgressManager(!bQuiet, 1);
+		FUpdateProgress::FStage& MainStage = ProgressManager.MainStage;
+
+		SceneTracker.Update(MainStage, false);
+		{
+			PROGRESS_STAGE("Sync With DirectLink")
+			UpdateDirectLinkScene();
+		}
+
+		ProgressManager.Finished();
+
+		if (Options.bStatSync)
+		{
+			LogInfo("Sync completed:");
+			ProgressManager.PrintStatisticss();
+		}
+	}
+
+	virtual void PerformAutoSync()
+	{
+		// Don't create progress bar for autosync - it steals focus, closes listener and what else
+		// todo: consider creating progress when a big change in scene is detected, e.g. number of nodes?
+		bool bQuiet = true;
+
+		FUpdateProgress ProgressManager(!bQuiet, 1);
+		FUpdateProgress::FStage& MainStage = ProgressManager.MainStage;
+
+		if (SceneTracker.Update(MainStage, false))  // Don't sent redundant update if scene change wasn't detected
+		{
+			PROGRESS_STAGE("Sync With DirectLink")
+			UpdateDirectLinkScene();
+		}
+
+		ProgressManager.Finished();
+
+		if (Options.bStatSync)
+		{
+			LogInfo("AutoSync completed:");
+			ProgressManager.PrintStatisticss();
+		}
+	}
+
 	virtual void ResetSceneTracking() override
 	{
 		NotificationsHandler.StopSceneChangeTracking();
@@ -2358,6 +2395,7 @@ bool Export(const TCHAR* Name, const TCHAR* OutputPath, bool bQuiet)
 
 	ProgressManager.Finished();
 
+	LogInfo(TEXT("Export fihished:"));
 	ProgressManager.PrintStatisticss();
 
 	return true;
diff --git a/Engine/Source/Programs/Enterprise/Datasmith/DatasmithMaxExporter/Public/DatasmithMaxDirectLink.h b/Engine/Source/Programs/Enterprise/Datasmith/DatasmithMaxExporter/Public/DatasmithMaxDirectLink.h
--- a/Engine/Source/Programs/Enterprise/Datasmith/DatasmithMaxExporter/Public/DatasmithMaxDirectLink.h
+++ b/Engine/Source/Programs/Enterprise/Datasmith/DatasmithMaxExporter/Public/DatasmithMaxDirectLink.h
@@ -85,6 +85,8 @@ class IExporter
 	virtual bool IsAutoSyncEnabled() = 0;
 	virtual void SetAutoSyncDelay(float Seconds) = 0;
 	virtual void SetAutoSyncIdleDelay(float Seconds) = 0;
+
+	virtual void PerformSync(bool bQuiet) = 0;
 };
 
 bool CreateExporter(bool bEnableUI, const TCHAR* EnginePath); // Create exporter with ability for DirectLink change tracking
