[3ds Max] DirectLink: typo fix

diff --git a/Engine/Source/Programs/Enterprise/Datasmith/DatasmithMaxExporter/Private/DatasmithMaxExporterWithDirectLink.cpp b/Engine/Source/Programs/Enterprise/Datasmith/DatasmithMaxExporter/Private/DatasmithMaxExporterWithDirectLink.cpp
--- a/Engine/Source/Programs/Enterprise/Datasmith/DatasmithMaxExporter/Private/DatasmithMaxExporterWithDirectLink.cpp
+++ b/Engine/Source/Programs/Enterprise/Datasmith/DatasmithMaxExporter/Private/DatasmithMaxExporterWithDirectLink.cpp
@@ -2395,7 +2395,7 @@ bool Export(const TCHAR* Name, const TCHAR* OutputPath, bool bQuiet)
 
 	ProgressManager.Finished();
 
-	LogInfo(TEXT("Export fihished:"));
+	LogInfo(TEXT("Export completed:"));
 	ProgressManager.PrintStatisticss();
 
 	return true;
