[Solidworks Exporter] Implement support for lights

diff --git a/Engine/Source/Programs/Enterprise/Datasmith/DatasmithSolidworksExporter/DatasmithSolidworks/DatasmithSolidworks.csproj b/Engine/Source/Programs/Enterprise/Datasmith/DatasmithSolidworksExporter/DatasmithSolidworks/DatasmithSolidworks.csproj
--- a/Engine/Source/Programs/Enterprise/Datasmith/DatasmithSolidworksExporter/DatasmithSolidworks/DatasmithSolidworks.csproj
+++ b/Engine/Source/Programs/Enterprise/Datasmith/DatasmithSolidworksExporter/DatasmithSolidworks/DatasmithSolidworks.csproj
@@ -160,6 +160,7 @@
     <Compile Include="..\Private\DatasmithExporter.cs" />
     <Compile Include="..\Private\Document.cs" />
     <Compile Include="..\Private\GeometryChunk.cs" />
+    <Compile Include="..\Private\Lights.cs" />
     <Compile Include="..\Private\Material.cs" />
     <Compile Include="..\Private\MaterialUtils.cs" />
     <Compile Include="..\Private\MathUtils.cs" />
@@ -194,6 +195,6 @@
 "$(EngineDir)\Source\ThirdParty\WiX\3.8\heat.exe" file "$(TargetDir)$(TargetName)$(TargetExt)" -gg -o "$(TargetDir)asm_SolidworksToUnreal.wxs"
 "$(EngineDir)\Source\ThirdParty\WiX\3.8\heat.exe" file "$(TargetDir)$(TargetName).tlb" -gg -o "$(TargetDir)tlb_SolidworksToUnreal.wxs"
 "$(ProjectDir)..\bin\$(Configuration)\PreprocessWxs.exe" "$(TargetDir)$(TargetName)$(TargetExt)" "$(TargetDir)asm_SolidworksToUnreal.wxs" "$(TargetDir)tlb_SolidworksToUnreal.wxs" "$(InstallerDir)" "$(EngineDir)"
-    </PostBuildEvent>
+</PostBuildEvent>
   </PropertyGroup>
 </Project>
\ No newline at end of file
diff --git a/Engine/Source/Programs/Enterprise/Datasmith/DatasmithSolidworksExporter/Private/DatasmithExporter.cs b/Engine/Source/Programs/Enterprise/Datasmith/DatasmithSolidworksExporter/Private/DatasmithExporter.cs
--- a/Engine/Source/Programs/Enterprise/Datasmith/DatasmithSolidworksExporter/Private/DatasmithExporter.cs
+++ b/Engine/Source/Programs/Enterprise/Datasmith/DatasmithSolidworksExporter/Private/DatasmithExporter.cs
@@ -12,7 +12,10 @@ namespace DatasmithSolidworks
 	public enum EActorType
 	{
 		SimpleActor,
-		MeshActor
+		MeshActor,
+		PointLightActor,
+		SpotLightActor,
+		DirLightActor
 	};
 
 	public class FDatasmithActorExportInfo
@@ -51,7 +54,7 @@ namespace DatasmithSolidworks
 			return null;
 		}
 
-		public void ExportOrUpdateActor(FDatasmithActorExportInfo InExportInfo)
+		public FDatasmithFacadeActor ExportOrUpdateActor(FDatasmithActorExportInfo InExportInfo)
 		{
 			FDatasmithFacadeActor Actor = null;
 
@@ -77,8 +80,11 @@ namespace DatasmithSolidworks
 				{
 					case EActorType.SimpleActor: Actor = new FDatasmithFacadeActor(InExportInfo.Name); break;
 					case EActorType.MeshActor: Actor = new FDatasmithFacadeActorMesh(InExportInfo.Name); break;
+					case EActorType.PointLightActor: Actor = new FDatasmithFacadePointLight(InExportInfo.Name); break;
+					case EActorType.SpotLightActor: Actor = new FDatasmithFacadeSpotLight(InExportInfo.Name); break;
+					case EActorType.DirLightActor: Actor = new FDatasmithFacadeDirectionalLight(InExportInfo.Name); break;
 				}
-				
+
 				Actor.AddTag(InExportInfo.Name);
 
 				ExportedActorsMap[InExportInfo.Name] = new Tuple<EActorType, FDatasmithFacadeActor>(InExportInfo.Type, Actor);
@@ -105,6 +111,8 @@ namespace DatasmithSolidworks
 				FDatasmithFacadeActorMesh MeshActor = Actor as FDatasmithFacadeActorMesh;
 				MeshActor.SetMesh(InExportInfo.MeshName);
 			}
+
+			return Actor;
 		}
 
 		public void RemoveActor(string InActorName)
@@ -118,6 +126,66 @@ namespace DatasmithSolidworks
 			}
 		}
 
+		public void ExportLight(FLight InLight)
+		{
+			FDatasmithActorExportInfo ExportInfo = new FDatasmithActorExportInfo();
+			ExportInfo.Label = InLight.LightName;
+			ExportInfo.Name = InLight.LightName;
+			ExportInfo.bVisible = true;
+		
+			FVec3 LightPosition = null;
+			FVec3 LightDirection = null;
+
+			switch (InLight.LightType)
+			{
+				case FLight.EType.Directional:
+				{
+					LightDirection = InLight.DirLightDirection;
+					ExportInfo.Type = EActorType.DirLightActor;
+				}
+				break;
+				case FLight.EType.Point:
+				{
+					LightPosition = InLight.PointLightPosition;
+					ExportInfo.Type = EActorType.PointLightActor;
+				}
+				break;
+				case FLight.EType.Spot:
+				{
+					LightPosition = InLight.SpotLightPosition;
+					LightDirection = (InLight.SpotLightTarget - InLight.SpotLightPosition).Normalized();
+					ExportInfo.Type = EActorType.SpotLightActor;
+				}
+				break;
+
+				default: return; // Unsupported light type
+			}
+
+			if (LightDirection != null)
+			{
+				ExportInfo.Transform = MathUtils.LookAt(LightDirection, LightPosition, 100f);
+			}
+			else if (LightPosition != null)
+			{
+				ExportInfo.Transform = MathUtils.Translation(LightPosition, 100f);
+			}
+
+			FDatasmithFacadeActorLight LightActor = ExportOrUpdateActor(ExportInfo) as FDatasmithFacadeActorLight;
+
+			const float MaxIntensity = 500f; // Map from SW (normlized) to Datasmith intensity
+
+			LightActor.SetIntensity(InLight.Intensity * MaxIntensity);
+			LightActor.SetColor(InLight.Color.X, InLight.Color.Y, InLight.Color.Z, 1f);
+			LightActor.SetEnabled(InLight.bIsEnabled);
+
+			if (LightActor is FDatasmithFacadeSpotLight SpotLight)
+			{
+				// Solidworks spot light has only one cone angle
+				SpotLight.SetInnerConeAngle(InLight.SpotLightConeAngle);
+				SpotLight.SetOuterConeAngle(InLight.SpotLightConeAngle);
+			}
+		}
+
 		public string ExportMesh(string InMeshName, FMeshData InData, string InUpdateMeshActor, out Tuple<FDatasmithFacadeMeshElement, FDatasmithFacadeMesh> OutMeshPair)
 		{
 			OutMeshPair = null;
diff --git a/Engine/Source/Programs/Enterprise/Datasmith/DatasmithSolidworksExporter/Private/Document.cs b/Engine/Source/Programs/Enterprise/Datasmith/DatasmithSolidworksExporter/Private/Document.cs
--- a/Engine/Source/Programs/Enterprise/Datasmith/DatasmithSolidworksExporter/Private/Document.cs
+++ b/Engine/Source/Programs/Enterprise/Datasmith/DatasmithSolidworksExporter/Private/Document.cs
@@ -124,6 +124,16 @@ namespace DatasmithSolidworks
 			}
 		}
 
+		private void ExportLights()
+		{
+			List<FLight> Lights  = FLightExporter.ExportLights(SwDoc);
+
+			foreach (FLight Light in Lights)
+			{
+				Exporter.ExportLight(Light);
+			}
+		}
+
 		private void RunExport(bool bIsDirectLinkExport)
 		{
 			if (bIsDirectLinkExport)
@@ -141,6 +151,8 @@ namespace DatasmithSolidworks
 
 				ExportToDatasmithScene();
 
+				ExportLights();
+
 #if DEBUG
 				Watch.Stop();
 				Debug.WriteLine($"EXPORT TIME: {(double)Watch.ElapsedMilliseconds / 1000.0}");
@@ -169,7 +181,8 @@ namespace DatasmithSolidworks
 				Exporter = new FDatasmithExporter(DatasmithScene);
 
 				ExportToDatasmithScene();
-
+			
+				ExportLights();
 
 				DatasmithScene.PreExport();
 				DatasmithScene.ExportScene(DatasmithFileExportPath);
diff --git a/Engine/Source/Programs/Enterprise/Datasmith/DatasmithSolidworksExporter/Private/Lights.cs b/Engine/Source/Programs/Enterprise/Datasmith/DatasmithSolidworksExporter/Private/Lights.cs
new file mode 100644
--- /dev/null
+++ b/Engine/Source/Programs/Enterprise/Datasmith/DatasmithSolidworksExporter/Private/Lights.cs
@@ -0,0 +1,214 @@
+// Copyright Epic Games, Inc. All Rights Reserved.
+
+using SolidWorks.Interop.sldworks;
+using System.Collections.Generic;
+using System.Runtime.InteropServices;
+
+namespace DatasmithSolidworks
+{
+	public class FLight
+	{
+		public enum EType
+		{
+			Point,
+			Spot,
+			Directional
+		};
+
+		public EType LightType { get; private set; }
+
+		public string LightName { get; private set; }
+
+		public bool bIsEnabled { get; private set; }
+
+		public float Intensity { get; private set; }
+		public FVec3 Color { get; private set; }
+
+		public FVec3 PointLightPosition
+		{
+			get
+			{
+				return LightType == EType.Point ? new FVec3(Properties[0], Properties[1], Properties[2]) : null;
+			}
+		}
+
+		public FVec3 SpotLightPosition
+		{
+			get
+			{
+				return LightType == EType.Spot ? new FVec3(Properties[0], Properties[1], Properties[2]) : null;
+			}
+		}
+		public FVec3 SpotLightTarget
+		{
+			get
+			{
+				return LightType == EType.Spot ? new FVec3(Properties[3], Properties[4], Properties[5]) : null;
+			}
+		}
+		public float SpotLightConeAngle
+		{
+			get
+			{
+				return LightType == EType.Spot ? Properties[6] : 0f;
+			}
+		}
+
+		public FVec3 DirLightDirection
+		{
+			get
+			{
+				return LightType == EType.Directional ? new FVec3(Properties[0], Properties[1], Properties[2]) : null;
+			}
+		}
+
+		private FLight(EType InType, string InName, bool bInIsEnabled, double InIntensisty, int InColorref)
+		{
+			LightType = InType;
+			LightName = InName;
+			bIsEnabled = bInIsEnabled;
+			Intensity = (float)InIntensisty;
+
+			int R = InColorref & 0xff;
+			int G = (InColorref >> 8) & 0xff;
+			int B = (InColorref >> 16) & 0xff;
+			int A = (InColorref >> 24) & 0xff;
+
+			Color = new FVec3(R/255f, G/255f, B/255f);
+		}
+
+		// Specific properties
+		private float[] Properties = new float[7];
+	
+		static public FLight MakeSpotLight(string InName, bool bInIsEnabled, FVec3 InPosition, FVec3 InTarget, double InConeAngle, double InIntensity, int InColor)
+		{
+			FLight SpotLight = new FLight(EType.Spot, InName, bInIsEnabled, InIntensity, InColor);
+
+			// Position
+			SpotLight.Properties[0] = InPosition.X;
+			SpotLight.Properties[1] = InPosition.Y;
+			SpotLight.Properties[2] = InPosition.Z;
+			// Target
+			SpotLight.Properties[3] = InTarget.X;
+			SpotLight.Properties[4] = InTarget.Y;
+			SpotLight.Properties[5] = InTarget.Z;
+			// Angle
+			SpotLight.Properties[6] = (float)InConeAngle;
+
+			return SpotLight;
+		}
+
+		static public FLight MakePointLight(string InName, bool bInIsEnabled, FVec3 InPosition, double InIntensity, int InColor)
+		{
+			FLight PointLight = new FLight(EType.Point, InName, bInIsEnabled, InIntensity, InColor);
+
+			// Position
+			PointLight.Properties[0] = InPosition.X;
+			PointLight.Properties[1] = InPosition.Y;
+			PointLight.Properties[2] = InPosition.Z;
+
+			return PointLight;
+		}
+
+		static public FLight MakeDirLight(string InName, bool bInIsEnabled, FVec3 InDirection, double InIntensity, int InColor)
+		{
+			FLight DirLight = new FLight(EType.Directional, InName, bInIsEnabled, InIntensity, InColor);
+
+			// Position
+			DirLight.Properties[0] = InDirection.X;
+			DirLight.Properties[1] = InDirection.Y;
+			DirLight.Properties[2] = InDirection.Z;
+
+			return DirLight;
+		}
+	};
+
+	public class FLightExporter
+	{
+		private const int SPOT_LIGHT = 2;
+		private const int POINT_LIGHT = 3;
+		private const int DIR_LIGHT = 4;
+
+		// Helper struct to pack/unpack integers to/from doubles.
+		// (some Solidworks API methods returns integers pack into doubles)
+		[StructLayout(LayoutKind.Explicit)]
+		private struct FDoubleIntConv
+		{
+			// An 8-byte double contains 2 4-byte ints.
+			[FieldOffset(0)] private int Int;
+			[FieldOffset(0)] private double Double;
+
+			public static void Unpack(double InValue, out int OutInt)
+			{
+				FDoubleIntConv CV = new FDoubleIntConv();
+				CV.Double = InValue;
+				OutInt = CV.Int;
+			}
+		}
+
+		public static List<FLight> ExportLights(ModelDoc2 InDoc)
+		{
+			List<FLight> SceneLights = new List<FLight>();
+
+			for (int LightIndex = 0; LightIndex < InDoc.GetLightSourceCount(); ++LightIndex)
+			{
+				string LightName = InDoc.GetLightSourceName(LightIndex);
+				double[] Props = InDoc.LightSourcePropertyValues[LightIndex] as double[];
+
+				int LightType;
+				FDoubleIntConv.Unpack(Props[0], out LightType);
+
+				double AmbientIntensity = 0.0;
+				double DiffuseIntensity = 0.0;
+				double SpecularIntensity = 0.0;
+				int Color = 0;
+				bool bIsEnabled = true;
+				bool bIsFixed = true;
+
+				FLight Light = null;
+
+				switch (LightType)
+				{
+					case SPOT_LIGHT:
+					{
+						double PosX = 0.0, PosY = 0.0, PosZ = 0.0;
+						double DirX = 0.0, DirY = 0.0, DirZ = 0.0;
+						double ConeAngle = 0.0;
+
+						if (InDoc.GetSpotlightProperties(LightName, ref AmbientIntensity, ref DiffuseIntensity, ref SpecularIntensity, ref Color, ref bIsEnabled, ref bIsFixed, ref PosX, ref PosY, ref PosZ, ref DirX, ref DirY, ref DirZ, ref ConeAngle))
+						{
+							Light = FLight.MakeSpotLight(LightName, bIsEnabled, new FVec3(PosX, PosY, PosZ), new FVec3(DirX, DirY, DirZ), ConeAngle, DiffuseIntensity, Color);
+						}
+					}
+					break;
+					case POINT_LIGHT:
+					{
+						double PosX = 0.0, PosY = 0.0, PosZ = 0.0;
+
+						if (InDoc.GetPointLightProperties(LightName, ref AmbientIntensity, ref DiffuseIntensity, ref SpecularIntensity, ref Color, ref bIsEnabled, ref bIsFixed, ref PosX, ref PosY, ref PosZ))
+						{
+							Light = FLight.MakePointLight(LightName, bIsEnabled, new FVec3(PosX, PosY, PosZ), DiffuseIntensity, Color);
+						}
+					}
+					break;
+					case DIR_LIGHT:
+					{
+						double TargetX = 0.0, TargetY = 0.0, TargetZ = 0.0;
+						if (InDoc.GetDirectionLightProperties(LightName, ref AmbientIntensity, ref DiffuseIntensity, ref SpecularIntensity, ref Color, ref bIsEnabled, ref bIsFixed, ref TargetX, ref TargetY, ref TargetZ))
+						{
+							Light = FLight.MakeDirLight(LightName, bIsEnabled, new FVec3(TargetX, TargetY, TargetZ), DiffuseIntensity, Color);
+						}
+					}
+					break;
+				}
+
+				if (Light != null)
+				{
+					SceneLights.Add(Light);
+				}
+			}
+
+			return SceneLights;
+		}
+	}
+}
\ No newline at end of file
diff --git a/Engine/Source/Programs/Enterprise/Datasmith/DatasmithSolidworksExporter/Private/MathUtils.cs b/Engine/Source/Programs/Enterprise/Datasmith/DatasmithSolidworksExporter/Private/MathUtils.cs
--- a/Engine/Source/Programs/Enterprise/Datasmith/DatasmithSolidworksExporter/Private/MathUtils.cs
+++ b/Engine/Source/Programs/Enterprise/Datasmith/DatasmithSolidworksExporter/Private/MathUtils.cs
@@ -7,7 +7,7 @@ using System.Runtime.InteropServices;
 
 namespace DatasmithSolidworks
 {
-    [ComVisible(false)]
+	[ComVisible(false)]
     public static class MathUtils
     {
 		public static float Rad2Deg { get { return (float)(180.0 / Math.PI); } }
@@ -78,6 +78,72 @@ namespace DatasmithSolidworks
 			return Result;
 		}
 
+		public static float[] LookAt(FVec3 InDirection, FVec3 InFromPoint, float InGeomScale)
+		{
+			float[] Ret = new float[16];
+
+			if (InDirection != null)
+			{
+				FVec3 Forward = InDirection.Normalized();
+				FVec3 Right = FVec3.Cross(FVec3.YAxis, Forward);
+				FVec3 Up = FVec3.Cross(Forward, Right);
+
+				Ret[0] = Forward.X;
+				Ret[1] = Forward.Y;
+				Ret[2] = Forward.Z;
+
+				Ret[4] = Right.X;
+				Ret[5] = Right.Y;
+				Ret[6] = Right.Z;
+
+				Ret[8] = Up.X;
+				Ret[9] = Up.Y;
+				Ret[10] = Up.Z;
+
+				Ret[3] = 0f;
+				Ret[7] = 0f;
+				Ret[11] = 0f;
+				Ret[15] = 1f;
+
+				if (InFromPoint != null)
+				{
+					Ret[12] = InFromPoint.X * InGeomScale;
+					Ret[13] = InFromPoint.Y * InGeomScale;
+					Ret[14] = InFromPoint.Z * InGeomScale;
+				}
+			}
+
+			return Ret;
+		}
+
+		public static float[] Translation(FVec3 InTranslation, float InGeomScale)
+		{
+			float[] Ret = new float[16];
+
+			Ret[0] = 1f;
+			Ret[1] = 0f;
+			Ret[2] = 0f;
+
+			Ret[4] = 0f;
+			Ret[5] = 1f;
+			Ret[6] = 0f;
+
+			Ret[8] = 0f;
+			Ret[9] = 0f;
+			Ret[10] = 1f;
+
+			Ret[3] = 0f;
+			Ret[7] = 0f;
+			Ret[11] = 0f;
+			Ret[15] = 1f;
+
+			Ret[12] = InTranslation.X * InGeomScale;
+			Ret[13] = InTranslation.Y * InGeomScale;
+			Ret[14] = InTranslation.Z * InGeomScale;
+
+			return Ret;
+		}
+
 		public static bool TransformsAreEqual(float[] InTransformA, float[] InTransformB)
 		{
 			Debug.Assert(InTransformA.Length == 16 && InTransformA.Length == InTransformB.Length);
diff --git a/Engine/Source/Programs/Enterprise/Datasmith/DatasmithSolidworksExporter/Private/Matrix4.cs b/Engine/Source/Programs/Enterprise/Datasmith/DatasmithSolidworksExporter/Private/Matrix4.cs
--- a/Engine/Source/Programs/Enterprise/Datasmith/DatasmithSolidworksExporter/Private/Matrix4.cs
+++ b/Engine/Source/Programs/Enterprise/Datasmith/DatasmithSolidworksExporter/Private/Matrix4.cs
@@ -385,19 +385,5 @@ namespace DatasmithSolidworks
 		{
 			return base.Equals(InObj);
 		}
-
-		public static bool operator ==(FMatrix4 InA, FMatrix4 InB)
-		{
-			if ((object)InA == null)
-			{
-				return (object)InB == null;
-			}
-			return InA.Equals(InB);
-		}
-
-		public static bool operator !=(FMatrix4 InA, FMatrix4 InB)
-		{
-			return !(InA == InB);
-		}
 	}
 }
diff --git a/Engine/Source/Programs/Enterprise/Datasmith/DatasmithSolidworksExporter/Private/Vec3.cs b/Engine/Source/Programs/Enterprise/Datasmith/DatasmithSolidworksExporter/Private/Vec3.cs
--- a/Engine/Source/Programs/Enterprise/Datasmith/DatasmithSolidworksExporter/Private/Vec3.cs
+++ b/Engine/Source/Programs/Enterprise/Datasmith/DatasmithSolidworksExporter/Private/Vec3.cs
@@ -110,19 +110,6 @@ namespace DatasmithSolidworks
                 (Z == float.NaN) ? 0.0f : Z);
         }
 
-        public static bool operator !=(FVec3 A, FVec3 B)
-        {
-            return !(A == B);
-        }
-
-        public static bool operator ==(FVec3 V1, FVec3 V2)
-        {
-            return 
-				MathUtils.Equals(V1.X, V2.X) && 
-				MathUtils.Equals(V1.Y, V2.Y) &&
-				MathUtils.Equals(V1.Z, V2.Z);
-        }
-
         public static bool operator <(FVec3 V1, FVec3 V2)
         {
 			return
