[Solidworks exporter] implement support for display state and configuration

diff --git a/Engine/Source/Programs/Enterprise/Datasmith/DatasmithSolidworksExporter/DatasmithSolidworks/DatasmithSolidworks.csproj b/Engine/Source/Programs/Enterprise/Datasmith/DatasmithSolidworksExporter/DatasmithSolidworks/DatasmithSolidworks.csproj
--- a/Engine/Source/Programs/Enterprise/Datasmith/DatasmithSolidworksExporter/DatasmithSolidworks/DatasmithSolidworks.csproj
+++ b/Engine/Source/Programs/Enterprise/Datasmith/DatasmithSolidworksExporter/DatasmithSolidworks/DatasmithSolidworks.csproj
@@ -149,13 +149,14 @@
     <Compile Include="Properties\AssemblyInfo.cs">
       <SubType>Code</SubType>
     </Compile>
+    <Compile Include="..\Private\ConfigurationTree.cs" />
+    <Compile Include="..\Private\ConfigurationExporter.cs" />
     <Compile Include="..\Private\Addin.cs" />
     <Compile Include="..\Private\Animation.cs" />
     <Compile Include="..\Private\AnimationExtractor.cs" />
     <Compile Include="..\Private\AssemblyDocument.cs" />
     <Compile Include="..\Private\Body.cs" />
     <Compile Include="..\Private\BoundingBox.cs" />
-    <Compile Include="..\Private\ConfigurationManager.cs" />
     <Compile Include="..\Private\ConstructionGeometry.cs" />
     <Compile Include="..\Private\DatasmithExporter.cs" />
     <Compile Include="..\Private\Document.cs" />
diff --git a/Engine/Source/Programs/Enterprise/Datasmith/DatasmithSolidworksExporter/Private/ConfigurationExporter.cs b/Engine/Source/Programs/Enterprise/Datasmith/DatasmithSolidworksExporter/Private/ConfigurationExporter.cs
new file mode 100644
--- /dev/null
+++ b/Engine/Source/Programs/Enterprise/Datasmith/DatasmithSolidworksExporter/Private/ConfigurationExporter.cs
@@ -0,0 +1,279 @@
+// Copyright Epic Games, Inc. All Rights Reserved.
+
+using System;
+using System.Collections.Generic;
+using SolidWorks.Interop.sldworks;
+using SolidWorks.Interop.swconst;
+using System.Runtime.InteropServices;
+using System.Collections.Concurrent;
+
+namespace DatasmithSolidworks
+{
+	public class FConfigurationData
+	{
+		public string Name;
+		public bool bIsDisplayStateConfiguration = false;
+		public Dictionary<string, float[]> ComponentTransform = new Dictionary<string, float[]>();
+		public Dictionary<string, bool> ComponentVisibility = new Dictionary<string, bool>();
+		public Dictionary<string, FObjectMaterials> ComponentMaterials = new Dictionary<string, FObjectMaterials>();
+
+		public bool IsEmpty()
+		{
+			return (
+				ComponentTransform.Count == 0 && 
+				ComponentVisibility.Count == 0 && 
+				ComponentMaterials.Count == 0);
+		}
+	};
+
+	[ComVisible(false)]
+	public class FConfigurationExporter
+	{
+		public static List<FConfigurationData> ExportConfigurations(FDocument InDoc)
+		{
+			string[] CfgNames = InDoc.SwDoc?.GetConfigurationNames();
+
+			if (CfgNames == null || CfgNames.Length <= 1)
+			{
+				return null;
+			}
+
+			FConfigurationTree.FComponentTreeNode CombinedTree = new FConfigurationTree.FComponentTreeNode();
+			CombinedTree.ComponentName = "CombinedTree";
+
+			// Ensure recursion will not stop on the root node (this may happen if it is not explicitly marked as visible)
+			CombinedTree.bVisibilitySame = true;
+			CombinedTree.CommonConfig.bVisible = true;
+
+			ConfigurationManager ConfigManager = InDoc.SwDoc.ConfigurationManager;
+			IConfiguration OriginalConfiguration = ConfigManager.ActiveConfiguration;
+
+			List<string> ExportedConfigurationNames = new List<string>();
+
+			foreach (string CfgName in CfgNames)
+			{
+				IConfiguration swConfiguration = InDoc.SwDoc.GetConfigurationByName(CfgName) as IConfiguration;
+
+				InDoc.SwDoc.ShowConfiguration2(CfgName);
+
+				string ConfigName = CfgName;
+
+				int DisplayStateCount = swConfiguration.GetDisplayStatesCount();
+				string[] DisplayStates = null;
+
+				if (ConfigManager.LinkDisplayStatesToConfigurations && DisplayStateCount > 1)
+				{
+					DisplayStates = swConfiguration.GetDisplayStates();
+					ConfigName = $"{CfgName}_{DisplayStates[0]}";
+				}
+
+				FConfigurationTree.FComponentTreeNode ConfigNode = new FConfigurationTree.FComponentTreeNode();
+				ConfigNode.Children = new List<FConfigurationTree.FComponentTreeNode>();
+				ConfigNode.ComponentName = ConfigName;
+
+				// Build the tree and get default materials (which aren't affected by any configuration)
+				// Use GetRootComponent3() with Resolve = true to ensure suppressed components will be loaded
+				CollectComponentsRecursive(InDoc, swConfiguration.GetRootComponent3(true), ConfigNode);
+
+				ExportedConfigurationNames.Add(ConfigName);
+
+				Dictionary<string, FObjectMaterials> MaterialsMap = GetComponentMaterials(InDoc, DisplayStates != null ? DisplayStates[0] : null, swConfiguration);
+				SetComponentTreeMaterials(ConfigNode, MaterialsMap, null, false);
+
+				if (DisplayStates != null)
+				{
+					// Export display states here: they are linked to configurations.
+					// Skip first as it was already exported above.
+					for (int Index = 1; Index < DisplayStateCount; ++Index)
+					{
+						string DisplayState = DisplayStates[Index];
+						MaterialsMap = GetComponentMaterials(InDoc, DisplayState, swConfiguration);
+						string DisplayStateTreeName = $"{CfgName}_DisplayState_{FDatasmithExporter.SanitizeName(DisplayState)}";
+						SetComponentTreeMaterials(ConfigNode, MaterialsMap, DisplayStateTreeName, false);
+					}
+				}
+
+				// Combine separate scene trees into the single one with configuration-specific data
+				FConfigurationTree.Merge(CombinedTree, ConfigNode, CfgName);
+			}
+
+			if (OriginalConfiguration != null)
+			{
+				InDoc.SwDoc.ShowConfiguration2(OriginalConfiguration.Name);
+			}
+
+			List<string> DisplayStateConfigurations = new List<string>();
+
+			if (!ConfigManager.LinkDisplayStatesToConfigurations)
+			{
+				// Export display states as separate configurations
+				string[] DisplayStates = OriginalConfiguration.GetDisplayStates();
+
+				foreach (string DisplayState in DisplayStates)
+				{
+					string DisplayStateConfigName = $"DisplayState_{FDatasmithExporter.SanitizeName(DisplayState)}";
+					DisplayStateConfigurations.Add(DisplayStateConfigName);
+					Dictionary<string, FObjectMaterials> MaterialsMap = GetComponentMaterials(InDoc, DisplayState, OriginalConfiguration);
+
+					SetComponentTreeMaterials(CombinedTree, MaterialsMap, DisplayStateConfigName, true);
+				}
+			}
+
+			// Remove configuration data when it's the same
+			FConfigurationTree.Compress(CombinedTree);
+
+			List<FConfigurationData> FlatConfigurationData = new List<FConfigurationData>();
+
+			// Add roots for configurations
+			foreach (string CfgName in ExportedConfigurationNames)
+			{
+				FConfigurationData CfgData = new FConfigurationData();
+				CfgData.Name = CfgName;
+				CfgData.bIsDisplayStateConfiguration = false;
+				FConfigurationTree.FillConfigurationData(CombinedTree, CfgName, CfgData, false);
+
+				if (!CfgData.IsEmpty())
+				{
+					FlatConfigurationData.Add(CfgData);
+				}
+			}
+			foreach (string DisplayStateConfigName in DisplayStateConfigurations)
+			{
+				FConfigurationData CfgData = new FConfigurationData();
+				CfgData.Name = DisplayStateConfigName;
+				CfgData.bIsDisplayStateConfiguration = true;
+				FConfigurationTree.FillConfigurationData(CombinedTree, DisplayStateConfigName, CfgData, true);
+
+				if (!CfgData.IsEmpty())
+				{
+					FlatConfigurationData.Add(CfgData);
+				}
+			}
+
+			return FlatConfigurationData;
+		}
+
+		private static Dictionary<string, FObjectMaterials> GetComponentMaterials(FDocument InDoc, string InDisplayState, IConfiguration InConfiguration)
+		{
+			Dictionary<string, FObjectMaterials> MaterialsMap = new Dictionary<string, FObjectMaterials>();
+
+			swDisplayStateOpts_e Option = InDisplayState != null ? swDisplayStateOpts_e.swSpecifyDisplayState : swDisplayStateOpts_e.swThisDisplayState;
+			string[] DisplayStates = InDisplayState != null ? new string[] { InDisplayState } : null;
+
+			if (InDoc is FAssemblyDocument)
+			{
+				FAssemblyDocument AsmDoc = InDoc as FAssemblyDocument;
+
+				HashSet<string> InComponentsSet = new HashSet<string>();
+
+				foreach (string CompName in AsmDoc.SyncState.ExportedComponentsMap.Keys)
+				{
+					InComponentsSet.Add(CompName);
+				}
+
+				ConcurrentDictionary<string, FObjectMaterials> ComponentMaterials =
+					FObjectMaterials.LoadAssemblyMaterials(AsmDoc, InComponentsSet, Option, DisplayStates);
+
+				foreach (var KVP in ComponentMaterials)
+				{
+					MaterialsMap.Add(KVP.Key, KVP.Value);
+				}
+			}
+			else
+			{
+				FPartDocument PartDocument = InDoc as FPartDocument;
+				FObjectMaterials PartMaterials = FObjectMaterials.LoadPartMaterials(InDoc, PartDocument.SwPartDoc, Option, DisplayStates);
+				Component2 Comp = InConfiguration.GetRootComponent3(true);
+				MaterialsMap.Add(Comp?.Name2 ?? "", PartMaterials);
+			}
+
+			return MaterialsMap;
+		}
+
+		private static void SetComponentTreeMaterials(FConfigurationTree.FComponentTreeNode InComponentTree, Dictionary<string, FObjectMaterials> InComponentMaterialsMap, string InConfigurationName, bool bIsDisplayState)
+		{
+			FConfigurationTree.FComponentConfig TargetConfig = null;
+
+			if (InConfigurationName != null)
+			{
+				TargetConfig = InComponentTree.GetConfiguration(InConfigurationName, bIsDisplayState);
+				if (TargetConfig == null)
+				{
+					TargetConfig = InComponentTree.AddConfiguration(InConfigurationName, bIsDisplayState);
+				}
+			}
+			else
+			{
+				TargetConfig = InComponentTree.CommonConfig;
+			}
+
+			if (InComponentMaterialsMap.ContainsKey(InComponentTree.ComponentName))
+			{
+				FObjectMaterials Materials = InComponentMaterialsMap[InComponentTree.ComponentName];
+				TargetConfig.Materials = Materials;
+			}
+
+			if (InComponentTree.Children != null)
+			{
+				foreach (FConfigurationTree.FComponentTreeNode Child in InComponentTree.Children)
+				{
+					SetComponentTreeMaterials(Child, InComponentMaterialsMap, InConfigurationName, bIsDisplayState);
+				}
+			}
+		}
+
+		private static void CollectComponentsRecursive(FDocument InDoc, Component2 InComponent, FConfigurationTree.FComponentTreeNode InParentNode)
+		{
+			FConfigurationTree.FComponentTreeNode NewNode = new FConfigurationTree.FComponentTreeNode();
+			InParentNode.Children.Add(NewNode);
+
+			// Basic properties
+			NewNode.ComponentName = InComponent.Name2;
+			NewNode.ComponentID = InComponent.GetID();
+			NewNode.CommonConfig.bVisible = InComponent.Visible != (int)swComponentVisibilityState_e.swComponentHidden;
+			NewNode.CommonConfig.bSuppressed = InComponent.IsSuppressed();
+
+			// Read transform
+			MathTransform ComponentTransform = InComponent.GetTotalTransform(true);
+			if (ComponentTransform == null)
+			{
+				ComponentTransform = InComponent.Transform2;
+			}
+
+			if (ComponentTransform != null)
+			{
+				NewNode.CommonConfig.Transform = MathUtils.ConvertFromSolidworksTransform(ComponentTransform, 100f /*GeomScale*/);
+
+				if (InParentNode.CommonConfig.Transform != null)
+				{
+					// Convert transform to parent space (original transform value fetched from Solidworks
+					// is in the root component's space). Datasmith wants relative transform for variants.
+					FMatrix4 ParentTransform = new FMatrix4(InParentNode.CommonConfig.Transform);
+					FMatrix4 InverseParentTransform = ParentTransform.Inverse();
+					NewNode.CommonConfig.RelativeTransform = FMatrix4.FMatrix4x4Multiply(InverseParentTransform, NewNode.CommonConfig.Transform);
+				}
+				else
+				{
+					NewNode.CommonConfig.RelativeTransform = NewNode.CommonConfig.Transform;
+				}
+			}
+
+			// Process children components
+			var Children = (Object[])InComponent.GetChildren();
+			if (Children != null && Children.Length > 0)
+			{
+				NewNode.Children = new List<FConfigurationTree.FComponentTreeNode>();
+				foreach (object ObjChild in Children)
+				{
+					Component2 Child = (Component2)ObjChild;
+					CollectComponentsRecursive(InDoc, Child, NewNode);
+				}
+
+				NewNode.Children.Sort(delegate (FConfigurationTree.FComponentTreeNode InA, FConfigurationTree.FComponentTreeNode InB)
+				{
+					return InA.ComponentID - InB.ComponentID;
+				});
+			}
+		}
+	}
+}
\ No newline at end of file
diff --git a/Engine/Source/Programs/Enterprise/Datasmith/DatasmithSolidworksExporter/Private/ConfigurationManager.cs b/Engine/Source/Programs/Enterprise/Datasmith/DatasmithSolidworksExporter/Private/ConfigurationManager.cs
deleted file mode 100644
diff --git a/Engine/Source/Programs/Enterprise/Datasmith/DatasmithSolidworksExporter/Private/ConfigurationTree.cs b/Engine/Source/Programs/Enterprise/Datasmith/DatasmithSolidworksExporter/Private/ConfigurationTree.cs
new file mode 100644
--- /dev/null
+++ b/Engine/Source/Programs/Enterprise/Datasmith/DatasmithSolidworksExporter/Private/ConfigurationTree.cs
@@ -0,0 +1,341 @@
+// Copyright Epic Games, Inc. All Rights Reserved.
+
+using SolidWorks.Interop.sldworks;
+using System.Collections.Generic;
+using System.Linq;
+using System.Runtime.InteropServices;
+
+namespace DatasmithSolidworks
+{
+	[ComVisible(false)]
+	public class FConfigurationTree
+	{
+		// Component state for a single configuration
+		public class FComponentConfig
+		{
+			public string ConfigName = null;
+
+			public bool bVisible = true;
+			public bool bSuppressed = false;
+			
+			// Node transform. Null value if transform is not changed in this configuration.
+			public float[] Transform = null;
+			public float[] RelativeTransform = null;
+
+			// Null value if configuration doesn't override the material.
+			public FObjectMaterials Materials = null;
+
+			public List<Body2> Bodies = null;
+
+			public void CopyFrom(FComponentConfig InOther)
+			{
+				bVisible = InOther.bVisible;
+				bSuppressed = InOther.bSuppressed;
+				Materials = InOther.Materials;
+				Transform = InOther.Transform;
+				RelativeTransform = InOther.RelativeTransform;
+				Bodies = InOther.Bodies;
+			}
+		};
+
+		// A node contains all the configuration data for a single component
+		public class FComponentTreeNode
+		{
+			public string ComponentName;
+			public int ComponentID;
+
+			// Common configuration data
+			public FComponentConfig CommonConfig = new FComponentConfig();
+
+			// Per-configuration data
+			public List<FComponentConfig> Configurations = null;
+			public List<FComponentConfig> DisplayStates = null;
+
+			public bool bVisibilitySame = true;
+			public bool bSuppressionSame = true;
+			public bool bMaterialSame = true;
+
+			public List<FComponentTreeNode> Children;
+
+			public FComponentConfig AddConfiguration(string InConfigurationName, bool bIsDisplayState)
+			{
+				List<FComponentConfig> TargetList = null;
+
+				if (bIsDisplayState)
+				{
+					if (DisplayStates == null)
+					{
+						DisplayStates = new List<FComponentConfig>();
+					}
+					TargetList = DisplayStates;
+				}
+				else
+				{
+					if (Configurations == null)
+					{
+						Configurations = new List<FComponentConfig>();
+					}
+					TargetList = Configurations;
+				}
+				FComponentConfig Result = new FComponentConfig();
+				Result.ConfigName = InConfigurationName;
+				TargetList.Add(Result);
+
+				return Result;
+			}
+
+			public FComponentConfig GetConfiguration(string InConfigurationName, bool bIsDisplayState)
+			{
+				List<FComponentConfig> TargetList = null;
+
+				if (bIsDisplayState)
+				{
+					TargetList = DisplayStates;
+				}
+				else
+				{
+					TargetList = Configurations;
+				}
+
+				if (TargetList == null)
+				{
+					return null;
+				}
+
+				foreach (FComponentConfig Config in TargetList)
+				{
+					if (Config.ConfigName == InConfigurationName)
+					{
+						return Config;
+					}
+				}
+				return null;
+			}
+
+			// Add all parameters except those which could be configuration-specific
+			public void AddParametersFrom(FComponentTreeNode InInput)
+			{
+				ComponentName = InInput.ComponentName;
+				ComponentID = InInput.ComponentID;
+			}
+		};
+
+		static public void Merge(FComponentTreeNode OutCombined, FComponentTreeNode InTree, string InConfigurationName)
+		{
+			if (InTree.Children == null)
+			{
+				return;
+			}
+
+			if (OutCombined.Children == null)
+			{
+				OutCombined.Children = new List<FComponentTreeNode>();
+			}
+
+			foreach (FComponentTreeNode Child in InTree.Children)
+			{
+				// Find the same node in 'combined', merge parameters
+				FComponentTreeNode CombinedChild = OutCombined.Children.FirstOrDefault(X => X.ComponentName == Child.ComponentName);
+
+				if (CombinedChild == null)
+				{
+					// The node doesn't exist yet, so add it
+					CombinedChild = new FComponentTreeNode();
+					CombinedChild.AddParametersFrom(Child);
+					OutCombined.Children.Add(CombinedChild);
+
+					// Copy common materials, which should be "default" for the component. Do not propagate these
+					// material to configurations, so configuration will only have material when it is changed.
+					OutCombined.CommonConfig.Materials = Child.CommonConfig.Materials;
+				}
+
+				// Make a NodeConfig and copy parameter values from 'tree' node
+				FComponentConfig NodeConfig = CombinedChild.AddConfiguration(InConfigurationName, false);
+				NodeConfig.CopyFrom(Child.CommonConfig);
+
+				// Recurse to children
+				Merge(CombinedChild, Child, InConfigurationName);
+			}
+		}
+
+		static public void Compress(FComponentTreeNode InNode)
+		{
+			bool CheckMaterialsEqual(List<FComponentConfig> ConfigList)
+			{
+				if (ConfigList == null || ConfigList.Count <= 1)
+				{
+					return true;
+				}
+
+				FObjectMaterials Materials = ConfigList[0].Materials;
+				bool bAllMaterialsAreSame = true;
+				for (int Idx = 1; Idx < ConfigList.Count; Idx++)
+				{
+					if (Materials != null && (!ConfigList[Idx].Materials?.EqualMaterials(Materials) ?? false))
+					{
+						bAllMaterialsAreSame = false;
+						break;
+					}
+				}
+
+				return bAllMaterialsAreSame;
+			}
+
+			if (InNode.Configurations != null && InNode.Configurations.Count > 0)
+			{
+				// Check transform
+				float[] Transform = InNode.Configurations[0].RelativeTransform;
+				bool bAllTransformsAreSame = true;
+				if (Transform != null)
+				{
+					// There could be components without a transform, so we're checking if for null
+					for (int Idx = 1; Idx < InNode.Configurations.Count; Idx++)
+					{
+						if (!InNode.Configurations[Idx].RelativeTransform.SequenceEqual(Transform))
+						{
+							bAllTransformsAreSame = false;
+							break;
+						}
+					}
+					if (bAllTransformsAreSame)
+					{
+						InNode.CommonConfig.RelativeTransform = Transform;
+						foreach (FComponentConfig Config in InNode.Configurations)
+						{
+							Config.RelativeTransform = null;
+						}
+					}
+				}
+
+				bool bAllMaterialsAreSame = CheckMaterialsEqual(InNode.Configurations);
+
+				if (bAllMaterialsAreSame && InNode.Configurations[0].Materials != null)
+				{
+					// We're explicitly checking for 'material != null' to not erase default
+					// material when no overrides detected.
+					InNode.CommonConfig.Materials = InNode.Configurations[0].Materials;
+					foreach (FComponentConfig Config in InNode.Configurations)
+					{
+						Config.Materials = null;
+					}
+				}
+
+				bool bVisible = InNode.Configurations[0].bVisible;
+				bool bSuppressed = InNode.Configurations[0].bSuppressed;
+				bool bVisibilitySame = true;
+				bool bSuppressionSame = true;
+				for (int Idx = 1; Idx < InNode.Configurations.Count; Idx++)
+				{
+					if (InNode.Configurations[Idx].bVisible != bVisible)
+					{
+						bVisibilitySame = false;
+					}
+
+					if (InNode.Configurations[Idx].bSuppressed != bSuppressed)
+					{
+						bSuppressionSame = false;
+					}
+				}
+				// Propagate common values
+				InNode.bMaterialSame = bAllMaterialsAreSame;
+				InNode.bVisibilitySame = bVisibilitySame;
+				InNode.bSuppressionSame = bSuppressionSame;
+
+				if (bVisibilitySame)
+				{
+					InNode.CommonConfig.bVisible = bVisible;
+				}
+				if (bSuppressionSame)
+				{
+					InNode.CommonConfig.bSuppressed = bSuppressed;
+				}
+
+				//todo: store bAll...Same in ConfigData
+
+				// If EVERYTHING is same, just remove all configurations at all
+				if (bAllTransformsAreSame && bAllMaterialsAreSame && bVisibilitySame && bSuppressionSame)
+				{
+					InNode.Configurations = null;
+				}
+			}
+
+			if (InNode.DisplayStates != null && InNode.DisplayStates.Count > 0)
+			{
+				bool bAllDisplayStateMaterialsAreSame = CheckMaterialsEqual(InNode.DisplayStates);
+
+				InNode.bMaterialSame = bAllDisplayStateMaterialsAreSame;
+
+				if (bAllDisplayStateMaterialsAreSame)
+				{
+					InNode.DisplayStates = null;
+				}
+			}
+
+			// Recurse to children
+			if (InNode.Children != null)
+			{
+				foreach (FComponentTreeNode Child in InNode.Children)
+				{
+					Compress(Child);
+				}
+			}
+		}
+
+		static public void FillConfigurationData(FComponentTreeNode InNode, string InConfigurationName, FConfigurationData OutData, bool bIsDisplayState)
+		{
+			FComponentConfig NodeConfig = InNode.GetConfiguration(InConfigurationName, bIsDisplayState);
+
+			// Visibility or suppression flags are set per node, and not propagated to childre.
+			// Also, suppression doesn't mark node as invisible. We should process these separately
+			// to exclude any variant information from invisible nodes and their children.
+			if ((InNode.bVisibilitySame && !InNode.CommonConfig.bVisible) || (NodeConfig != null && !NodeConfig.bVisible) ||
+				(InNode.bSuppressionSame && InNode.CommonConfig.bSuppressed) || (NodeConfig != null && NodeConfig.bSuppressed))
+			{
+				// This node is not visible
+				OutData.ComponentVisibility.Add(InNode.ComponentName, false);
+				// Do not process children of this node
+				return;
+			}
+
+			// Process current configuration
+			if (NodeConfig != null)
+			{
+				// Only add variant information if attribute is not the same in all configurations
+				if (NodeConfig.RelativeTransform != null)
+				{
+					OutData.ComponentTransform.Add(InNode.ComponentName, NodeConfig.RelativeTransform);
+				}
+				if (!InNode.bVisibilitySame)
+				{
+					OutData.ComponentVisibility.Add(InNode.ComponentName, NodeConfig.bVisible);
+				}
+				if (!InNode.bSuppressionSame)
+				{
+					OutData.ComponentVisibility.Add(InNode.ComponentName, !NodeConfig.bSuppressed);
+				}
+				if (!InNode.bMaterialSame)
+				{
+					FObjectMaterials Materials = NodeConfig.Materials;
+					if (Materials == null)
+					{
+						// There's no material set for the config, use default one
+						Materials = InNode.CommonConfig.Materials;
+					}
+					if (Materials != null)
+					{
+						OutData.ComponentMaterials.Add(InNode.ComponentName, Materials);
+					}
+				}
+			}
+
+			// Recurse to children
+			if (InNode.Children != null)
+			{
+				foreach (FComponentTreeNode Child in InNode.Children)
+				{
+					FillConfigurationData(Child, InConfigurationName, OutData, bIsDisplayState);
+				}
+			}
+		}
+	}
+}
\ No newline at end of file
diff --git a/Engine/Source/Programs/Enterprise/Datasmith/DatasmithSolidworksExporter/Private/DatasmithExporter.cs b/Engine/Source/Programs/Enterprise/Datasmith/DatasmithSolidworksExporter/Private/DatasmithExporter.cs
--- a/Engine/Source/Programs/Enterprise/Datasmith/DatasmithSolidworksExporter/Private/DatasmithExporter.cs
+++ b/Engine/Source/Programs/Enterprise/Datasmith/DatasmithSolidworksExporter/Private/DatasmithExporter.cs
@@ -36,6 +36,7 @@ namespace DatasmithSolidworks
 		private ConcurrentDictionary<int, FDatasmithFacadeMasterMaterial> ExportedMaterialsMap = new ConcurrentDictionary<int, FDatasmithFacadeMasterMaterial>();
 		private ConcurrentDictionary<string, FDatasmithFacadeTexture> ExportedTexturesMap = new ConcurrentDictionary<string, FDatasmithFacadeTexture>();
 		private Dictionary<string, FDatasmithFacadeActorBinding> ExportedActorBindingsMap = new Dictionary<string, FDatasmithFacadeActorBinding>();
+		private Dictionary<string, FDatasmithFacadeVariant> ExportedVariantsMap = new Dictionary<string, FDatasmithFacadeVariant>();
 
 		private FDatasmithFacadeScene DatasmithScene = null;
 
@@ -316,45 +317,142 @@ namespace DatasmithSolidworks
 			}
 		}
 
-		public void ExportConfiguration(string InConfigurationsSetName, FConfigurationData InConfig)
+		FDatasmithFacadeActorBinding GetActorBinding(string InActorName, FDatasmithFacadeVariant InVariant)
 		{
-			FDatasmithFacadeActorBinding GetActorBinding(string InActorName, FDatasmithFacadeVariant InVariant)
+			for (int BindingIndex = 0; BindingIndex < InVariant.GetActorBindingsCount(); ++BindingIndex)
 			{
-				FDatasmithFacadeActorBinding Binding = null;
-
-				if (!ExportedActorBindingsMap.ContainsKey(InActorName))
+				FDatasmithFacadeActorBinding Binding = InVariant.GetActorBinding(BindingIndex);
+				if (Binding.GetName() == InActorName)
 				{
-					// Find a datasmith actor
-					FDatasmithFacadeActor Actor = null;
+					return Binding;
+				}
+			}
+
+			// No binding was found, add one
+
+			FDatasmithFacadeActor Actor = null;
+
+			if (ExportedActorsMap.ContainsKey(InActorName))
+			{
+				Tuple<EActorType, FDatasmithFacadeActor> ActorInfo = ExportedActorsMap[InActorName];
+				Actor = ActorInfo.Item2;
+			}
+			else
+			{
+				// Actor was not found, should not happen
+				return null;
+			}
+
+			// Make a new binding
+			FDatasmithFacadeActorBinding NewBinding = new FDatasmithFacadeActorBinding(Actor);
+			InVariant.AddActorBinding(NewBinding);
+
+			return NewBinding;
+		}
+
+		private void ExportMaterialVariants(List<Tuple<FConfigurationData, FDatasmithFacadeVariant>> InVariants)
+		{
+			foreach (Tuple<FConfigurationData, FDatasmithFacadeVariant> KVP in InVariants)
+			{
+				FConfigurationData Config = KVP.Item1;
+				FDatasmithFacadeVariant Variant = KVP.Item2;
 
-					if (ExportedActorsMap.ContainsKey(InActorName))
+				// Iterate over all material assignments
+				foreach (var MatKVP in Config.ComponentMaterials)
+				{
+					FDatasmithFacadeActorBinding Binding = GetActorBinding(MatKVP.Key, Variant);
+					if (Binding != null)
 					{
-						Tuple<EActorType, FDatasmithFacadeActor> ActorInfo = ExportedActorsMap[InActorName];
-						Actor = ActorInfo.Item2;
+						HashSet<int> UniqueMaterialsSet = new HashSet<int>();
+						FObjectMaterials Materials = MatKVP.Value;
+
+						if (Materials.ComponentMaterialID != -1 && !UniqueMaterialsSet.Contains(Materials.ComponentMaterialID))
+						{
+							UniqueMaterialsSet.Add(Materials.ComponentMaterialID);
+						}
+						if (Materials.PartMaterialID != -1 && !UniqueMaterialsSet.Contains(Materials.ComponentMaterialID))
+						{
+							UniqueMaterialsSet.Add(Materials.PartMaterialID);
+						}
+						foreach (var MatID in Materials.BodyMaterialsMap.Values)
+						{
+							if (!UniqueMaterialsSet.Contains(Materials.ComponentMaterialID))
+							{
+								UniqueMaterialsSet.Add(MatID);
+							}
+						}
+						foreach (var MatID in Materials.FeatureMaterialsMap.Values)
+						{
+							if (!UniqueMaterialsSet.Contains(Materials.ComponentMaterialID))
+							{
+								UniqueMaterialsSet.Add(MatID);
+							}
+						}
+						foreach (var MatID in Materials.FaceMaterialsMap.Values)
+						{
+							if (!UniqueMaterialsSet.Contains(Materials.ComponentMaterialID))
+							{
+								UniqueMaterialsSet.Add(MatID);
+							}
+						}
+
+						foreach (var MatID in UniqueMaterialsSet)
+						{
+							FMaterial Material = Materials.GetMaterial(MatID);
+							FDatasmithFacadeMasterMaterial DatasmithMaterial = null;
+							if (Material != null && ExportedMaterialsMap.TryGetValue(Material.ID, out DatasmithMaterial))
+							{
+								Binding.AddMaterialCapture(DatasmithMaterial);
+							}
+						}
 					}
-					else
+				}
+			}
+		}
+
+		private void ExportTransformVariants(List<Tuple<FConfigurationData, FDatasmithFacadeVariant>> InVariants)
+		{
+			foreach (Tuple<FConfigurationData, FDatasmithFacadeVariant> KVP in InVariants)
+			{
+				FConfigurationData Config = KVP.Item1;
+				FDatasmithFacadeVariant Variant = KVP.Item2;
+
+				// Provide transform variants
+				foreach (var TransformMap in Config.ComponentTransform)
+				{
+					FDatasmithFacadeActorBinding Binding = GetActorBinding(TransformMap.Key, Variant);
+					if (Binding != null)
 					{
-						// Actor was not found, should not happen
-						return null;
+						Binding.AddRelativeTransformCapture(TransformMap.Value);
 					}
-
-					// Make a new binding
-					Binding = new FDatasmithFacadeActorBinding(Actor);
-					ExportedActorBindingsMap.Add(InActorName, Binding);
-					InVariant.AddActorBinding(Binding);
 				}
-				else
+			}
+		}
+
+		private void ExportActorVisibilityVariants(List<Tuple<FConfigurationData, FDatasmithFacadeVariant>> InVariants)
+		{
+			foreach (Tuple<FConfigurationData, FDatasmithFacadeVariant> KVP in InVariants)
+			{
+				FConfigurationData Config = KVP.Item1;
+				FDatasmithFacadeVariant Variant = KVP.Item2;
+
+				// Build a visibility variant data
+				foreach (var VisibilityMap in Config.ComponentVisibility)
 				{
-					// Get an existing binding
-					Binding = ExportedActorBindingsMap[InActorName];
+					FDatasmithFacadeActorBinding Binding = GetActorBinding(VisibilityMap.Key, Variant);
+					if (Binding != null)
+					{
+						Binding.AddVisibilityCapture(VisibilityMap.Value);
+					}
 				}
-				return Binding;
 			}
+		}
 
+		public void ExportLevelVariantSets(List<FConfigurationData> InConfigs)
+		{
 			// Request existing VariantSet, or create a new one
 			FDatasmithFacadeLevelVariantSets LevelVariantSets = null;
-			FDatasmithFacadeVariantSet VariantSet = null;
-
+			
 			if (DatasmithScene.GetLevelVariantSetsCount() == 0)
 			{
 				LevelVariantSets = new FDatasmithFacadeLevelVariantSets("LevelVariantSets");
@@ -365,62 +463,72 @@ namespace DatasmithSolidworks
 				LevelVariantSets = DatasmithScene.GetLevelVariantSets(0);
 			}
 
-			int VariantSetsCount = LevelVariantSets.GetVariantSetsCount();
-			for (int VariantSetIndex = 0; VariantSetIndex < VariantSetsCount; ++VariantSetIndex)
+			FDatasmithFacadeVariantSet GetOrCreateLevelVariantSet(string InName)
 			{
-				FDatasmithFacadeVariantSet VSet = LevelVariantSets.GetVariantSet(VariantSetIndex);
+				int VariantSetsCount = LevelVariantSets.GetVariantSetsCount();
+				FDatasmithFacadeVariantSet VariantSet = null;
 
-				if (VSet.GetName() == InConfigurationsSetName)
+				for (int VariantSetIndex = 0; VariantSetIndex < VariantSetsCount; ++VariantSetIndex)
 				{
-					VariantSet = VSet;
-					break;
+					FDatasmithFacadeVariantSet VSet = LevelVariantSets.GetVariantSet(VariantSetIndex);
+
+					if (VSet.GetName() == InName)
+					{
+						VariantSet = VSet;
+						break;
+					}
 				}
-			}
 
-			if (VariantSet == null)
-			{
-				VariantSet = new FDatasmithFacadeVariantSet(InConfigurationsSetName);
-				LevelVariantSets.AddVariantSet(VariantSet);
+				if (VariantSet == null)
+				{
+					VariantSet = new FDatasmithFacadeVariantSet(InName);
+					LevelVariantSets.AddVariantSet(VariantSet);
+				}
+				return VariantSet;
 			}
 
-			// Add a new variant
-			FDatasmithFacadeVariant Variant = new FDatasmithFacadeVariant(InConfig.Name);
-			VariantSet.AddVariant(Variant);
+			FDatasmithFacadeVariantSet ConfigurationsVariantSet = null;
+			FDatasmithFacadeVariantSet DisplayStatesVariantSet = null;
+
+			List<Tuple<FConfigurationData, FDatasmithFacadeVariant>> ConfigurationVariants = null;
+			List<Tuple<FConfigurationData, FDatasmithFacadeVariant>> DisplayStateVariants = null;
 
-			// Build a visibility variant data
-			foreach (var VisibilityMap in InConfig.ComponentVisibility)
+			foreach (FConfigurationData Config in InConfigs)
 			{
-				FDatasmithFacadeActorBinding Binding = GetActorBinding(VisibilityMap.Key, Variant);
-				if (Binding != null)
+				FDatasmithFacadeVariant Variant = new FDatasmithFacadeVariant(Config.Name);
+
+				if (Config.bIsDisplayStateConfiguration)
 				{
-					Binding.AddVisibilityCapture(VisibilityMap.Value);
+					if (DisplayStatesVariantSet == null)
+					{
+						DisplayStatesVariantSet = GetOrCreateLevelVariantSet("DisplayStates");
+						DisplayStateVariants = new List<Tuple<FConfigurationData, FDatasmithFacadeVariant>>();
+					}
+					DisplayStatesVariantSet.AddVariant(Variant);
+					DisplayStateVariants.Add(new Tuple<FConfigurationData, FDatasmithFacadeVariant>(Config, Variant));
 				}
-			}
-
-			// Provide transform variants
-			foreach (var TransformMap in InConfig.ComponentTransform)
-			{
-				FDatasmithFacadeActorBinding Binding = GetActorBinding(TransformMap.Key, Variant);
-				if (Binding != null)
+				else
 				{
-					Binding.AddRelativeTransformCapture(TransformMap.Value);
+					if (ConfigurationsVariantSet == null)
+					{
+						ConfigurationsVariantSet = GetOrCreateLevelVariantSet("Configurations");
+						ConfigurationVariants = new List<Tuple<FConfigurationData, FDatasmithFacadeVariant>>();
+					}
+					ConfigurationsVariantSet.AddVariant(Variant);
+					ConfigurationVariants.Add(new Tuple<FConfigurationData, FDatasmithFacadeVariant>(Config, Variant));
 				}
 			}
 
-			// Iterate over all material assignments
-			foreach (var KVP in InConfig.ComponentMaterials)
+			if (ConfigurationVariants != null)
 			{
-				FDatasmithFacadeActorBinding Binding = GetActorBinding(KVP.Key, Variant);
-				if (Binding != null)
-				{
-					FObjectMaterials Materials = KVP.Value;
+				ExportActorVisibilityVariants(ConfigurationVariants);
+				ExportMaterialVariants(ConfigurationVariants);
+				ExportTransformVariants(ConfigurationVariants);
+			}
 
-					FMaterial TopMaterial = Materials.GetComponentMaterial();
-					if (TopMaterial != null)
-					{
-						Binding.AddMaterialCapture(ExportedMaterialsMap[TopMaterial.ID]);
-					}
-				}
+			if (DisplayStateVariants != null)
+			{
+				ExportMaterialVariants(DisplayStateVariants);
 			}
 		}
 
diff --git a/Engine/Source/Programs/Enterprise/Datasmith/DatasmithSolidworksExporter/Private/Document.cs b/Engine/Source/Programs/Enterprise/Datasmith/DatasmithSolidworksExporter/Private/Document.cs
--- a/Engine/Source/Programs/Enterprise/Datasmith/DatasmithSolidworksExporter/Private/Document.cs
+++ b/Engine/Source/Programs/Enterprise/Datasmith/DatasmithSolidworksExporter/Private/Document.cs
@@ -99,6 +99,23 @@ namespace DatasmithSolidworks
 			}
 		}
 
+		private void ExportConfigurations()
+		{
+			// Remove any existing Datasmith LevelVariantSets as we'll re-add data
+			while (DatasmithScene.GetLevelVariantSetsCount() > 0)
+			{
+				for (int Index = 0; Index < DatasmithScene.GetLevelVariantSetsCount(); Index++)
+				{
+					DatasmithScene.RemoveLevelVariantSets(DatasmithScene.GetLevelVariantSets(Index));
+				}
+			}
+
+			List<FConfigurationData> Configs = FConfigurationExporter.ExportConfigurations(this);
+
+			Exporter.ExportMaterials(ExportedMaterialsMap);
+			Exporter.ExportLevelVariantSets(Configs);
+		}
+
 		private void CheckForMaterialUpdatesProc()
 		{
 			while (!bExitMaterialUpdateThread)
@@ -181,8 +198,10 @@ namespace DatasmithSolidworks
 				Exporter = new FDatasmithExporter(DatasmithScene);
 
 				ExportToDatasmithScene();
+
 			
 				ExportLights();
+				ExportConfigurations();
 
 				DatasmithScene.PreExport();
 				DatasmithScene.ExportScene(DatasmithFileExportPath);
diff --git a/Engine/Source/Programs/Enterprise/Datasmith/DatasmithSolidworksExporter/Private/ObjectMaterials.cs b/Engine/Source/Programs/Enterprise/Datasmith/DatasmithSolidworksExporter/Private/ObjectMaterials.cs
--- a/Engine/Source/Programs/Enterprise/Datasmith/DatasmithSolidworksExporter/Private/ObjectMaterials.cs
+++ b/Engine/Source/Programs/Enterprise/Datasmith/DatasmithSolidworksExporter/Private/ObjectMaterials.cs
@@ -19,12 +19,12 @@ namespace DatasmithSolidworks
 		public PartDoc PartDocument {  get; private set; }
 		public FDocument OwnerDoc {  get; private set; }
 
-		private int ComponentMaterialID = -1;
-		private int PartMaterialID = -1;
+		public int ComponentMaterialID { get; private set; } = -1;
+		public int PartMaterialID { get; private set; } = -1;
 
-		private Dictionary<string, int> BodyMaterialsMap = new Dictionary<string, int>();
-		private Dictionary<string, int> FaceMaterialsMap = new Dictionary<string, int>();
-		private Dictionary<string, int> FeatureMaterialsMap = new Dictionary<string, int>();
+		public Dictionary<string, int> BodyMaterialsMap { get; private set; } = new Dictionary<string, int>();
+		public Dictionary<string, int> FaceMaterialsMap { get; private set; } = new Dictionary<string, int>();
+		public Dictionary<string, int> FeatureMaterialsMap { get; private set; } = new Dictionary<string, int>();
 
 		public ConcurrentDictionary<int, FMaterial> GlobalMaterialsMap { get; private set; } = null;
 
diff --git a/Engine/Source/Programs/Enterprise/Datasmith/DatasmithSolidworksExporter/Private/PartDocument.cs b/Engine/Source/Programs/Enterprise/Datasmith/DatasmithSolidworksExporter/Private/PartDocument.cs
--- a/Engine/Source/Programs/Enterprise/Datasmith/DatasmithSolidworksExporter/Private/PartDocument.cs
+++ b/Engine/Source/Programs/Enterprise/Datasmith/DatasmithSolidworksExporter/Private/PartDocument.cs
@@ -10,7 +10,7 @@ namespace DatasmithSolidworks
 {
 	public class FPartDocument : FDocument
 	{
-		private PartDoc SwPartDoc = null;
+		public PartDoc SwPartDoc { get; private set; } = null;
 		private FAssemblyDocument AsmDoc = null;
 		private string ComponentName = null;
 
