[Datasmith Solidworks] Animation are not exported from solidworks using datasmith

diff --git a/Engine/Source/Programs/Enterprise/Datasmith/DatasmithSolidworksExporter/Private/Animation.cs b/Engine/Source/Programs/Enterprise/Datasmith/DatasmithSolidworksExporter/Private/Animation.cs
--- a/Engine/Source/Programs/Enterprise/Datasmith/DatasmithSolidworksExporter/Private/Animation.cs
+++ b/Engine/Source/Programs/Enterprise/Datasmith/DatasmithSolidworksExporter/Private/Animation.cs
@@ -1,5 +1,6 @@
 // Copyright Epic Games, Inc. All Rights Reserved.
 
+using System;
 using System.Collections.Generic;
 using SolidWorks.Interop.sldworks;
 
@@ -132,43 +133,5 @@ namespace DatasmithSolidworks
 
 			return XForm;
 		}
-
-#if false // level sequences are not implemented for 4.27
-		public FDatasmithFacadeLevelSequence Export()
-		{
-			FDatasmithFacadeLevelSequence LevelSeq = new FDatasmithFacadeLevelSequence(Name);
-
-			LevelSeq.SetFrameRate(FPS);
-
-			foreach (var NodePair in ComponentToChannelMap)
-			{
-				FChannel Chan = NodePair.Value;
-				Component2 Component = Chan.Target;
-				FDatasmithFacadeTransformAnimation Anim = new FDatasmithFacadeTransformAnimation(Component.Name2);
-
-				foreach (FKeyframe Keyframe in Chan.Keyframes)
-				{
-					FMatrix4 LocalMatrix = Keyframe.LocalMatrix;
-
-					// Get euler angles in degrees
-					float X = MathUtils.Rad2Deg * (float)Math.Atan2(LocalMatrix[6], LocalMatrix[10]);
-					float Y = MathUtils.Rad2Deg * (float)Math.Atan2(-LocalMatrix[2], Math.Sqrt(LocalMatrix[6] * LocalMatrix[6] + LocalMatrix[10] * LocalMatrix[10]));
-					float Z = MathUtils.Rad2Deg * (float)Math.Atan2(LocalMatrix[1], LocalMatrix[0]);
-
-					float Scale = LocalMatrix[15];
-
-					Vec3 Translation = new Vec3(LocalMatrix[12], LocalMatrix[13], LocalMatrix[14]);
-
-					Anim.AddFrame(EDatasmithFacadeAnimationTransformType.Rotation, Keyframe.Step, X, -Y, -Z);
-					Anim.AddFrame(EDatasmithFacadeAnimationTransformType.Scale, Keyframe.Step, Scale, Scale, Scale);
-					Anim.AddFrame(EDatasmithFacadeAnimationTransformType.Translation, Keyframe.Step, Translation.x, -Translation.y, Translation.z);
-				}
-
-				LevelSeq.AddAnimation(Anim);
-			}
-
-			return LevelSeq;
-		}
-#endif
 	}
 }
diff --git a/Engine/Source/Programs/Enterprise/Datasmith/DatasmithSolidworksExporter/Private/AssemblyDocument.cs b/Engine/Source/Programs/Enterprise/Datasmith/DatasmithSolidworksExporter/Private/AssemblyDocument.cs
--- a/Engine/Source/Programs/Enterprise/Datasmith/DatasmithSolidworksExporter/Private/AssemblyDocument.cs
+++ b/Engine/Source/Programs/Enterprise/Datasmith/DatasmithSolidworksExporter/Private/AssemblyDocument.cs
@@ -131,6 +131,21 @@ namespace DatasmithSolidworks
 			SyncState.DirtyComponents.Clear();
 
 			SyncState = OldSyncState;
+
+			// Export animations (only allow when exporting to file)
+		//	if (bFileExportInProgress)
+			{
+				List<FAnimation> Animations = FAnimationExtractor.ExtractAnimations(SwAsmDoc, Root);
+				if (Animations != null)
+				{
+					SetExportStatus($"Animations");
+
+					foreach (FAnimation Anim in Animations)
+					{
+						Exporter.ExportAnimation(Anim);
+					}
+				}
+			}
 		}
 
 		public override bool HasMaterialUpdates()
@@ -263,11 +278,16 @@ namespace DatasmithSolidworks
 				FDatasmithActorExportInfo ActorExportInfo = new FDatasmithActorExportInfo();
 
 				string ComponentName = FDatasmithExporter.SanitizeName(InComponent.Name2);
-				string[] NameComponents = ComponentName.Split('/');
+				string[] NameComponents = InComponent.Name2.Split('/');
 
 				ActorExportInfo.Label = NameComponents.Last();
 				ActorExportInfo.Name = ComponentName;
-				ActorExportInfo.ParentName = InParent?.Name2;
+
+				if (InParent != null)
+				{
+					ActorExportInfo.ParentName = FDatasmithExporter.SanitizeName(InParent.Name2);
+				}
+				
 				ActorExportInfo.bVisible = true;
 				ActorExportInfo.Type = Exporter.GetExportedActorType(ComponentName) ?? EActorType.SimpleActor;
 				ActorExportInfo.Transform = GetComponentDatasmithTransform(InComponent);
diff --git a/Engine/Source/Programs/Enterprise/Datasmith/DatasmithSolidworksExporter/Private/ConfigurationExporter.cs b/Engine/Source/Programs/Enterprise/Datasmith/DatasmithSolidworksExporter/Private/ConfigurationExporter.cs
--- a/Engine/Source/Programs/Enterprise/Datasmith/DatasmithSolidworksExporter/Private/ConfigurationExporter.cs
+++ b/Engine/Source/Programs/Enterprise/Datasmith/DatasmithSolidworksExporter/Private/ConfigurationExporter.cs
@@ -174,9 +174,12 @@ namespace DatasmithSolidworks
 				ConcurrentDictionary<string, FObjectMaterials> ComponentMaterials =
 					FObjectMaterials.LoadAssemblyMaterials(AsmDoc, InComponentsSet, Option, DisplayStates);
 
-				foreach (var KVP in ComponentMaterials)
+				if (ComponentMaterials != null)
 				{
-					MaterialsMap.Add(KVP.Key, KVP.Value);
+					foreach (var KVP in ComponentMaterials)
+					{
+						MaterialsMap.Add(KVP.Key, KVP.Value);
+					}
 				}
 			}
 			else
diff --git a/Engine/Source/Programs/Enterprise/Datasmith/DatasmithSolidworksExporter/Private/DatasmithExporter.cs b/Engine/Source/Programs/Enterprise/Datasmith/DatasmithSolidworksExporter/Private/DatasmithExporter.cs
--- a/Engine/Source/Programs/Enterprise/Datasmith/DatasmithSolidworksExporter/Private/DatasmithExporter.cs
+++ b/Engine/Source/Programs/Enterprise/Datasmith/DatasmithSolidworksExporter/Private/DatasmithExporter.cs
@@ -677,7 +677,7 @@ namespace DatasmithSolidworks
 			{
 				FAnimation.FChannel Chan = NodePair.Value;
 				Component2 Component = Chan.Target;
-				FDatasmithFacadeTransformAnimation Anim = new FDatasmithFacadeTransformAnimation(Component.Name2);
+				FDatasmithFacadeTransformAnimation Anim = new FDatasmithFacadeTransformAnimation(SanitizeName(Component.Name2));
 
 				foreach (var Keyframe in Chan.Keyframes)
 				{
