[3ds Max] DirectLink: Added handlight to additional change events for geometry

- Topology(e.g. Flip normals on Poligons of a EditPoly)
- Mapping(edit MapVert without touching vertices/polys, e.g. by mxs polyop.setMapVert)

diff --git a/Engine/Source/Programs/Enterprise/Datasmith/DatasmithMaxExporter/Private/DatasmithMaxDirectLinkNotifications.cpp b/Engine/Source/Programs/Enterprise/Datasmith/DatasmithMaxExporter/Private/DatasmithMaxDirectLinkNotifications.cpp
--- a/Engine/Source/Programs/Enterprise/Datasmith/DatasmithMaxExporter/Private/DatasmithMaxDirectLinkNotifications.cpp
+++ b/Engine/Source/Programs/Enterprise/Datasmith/DatasmithMaxExporter/Private/DatasmithMaxDirectLinkNotifications.cpp
@@ -36,6 +36,25 @@ public:
 	{
 		LogNodeEvent(L"GeometryChanged", nodes);
 
+		InvalidateGeometry(nodes);
+	}
+
+	virtual void TopologyChanged(NodeKeyTab& nodes) override
+	{
+		LogNodeEvent(L"TopologyChanged", nodes);
+
+		InvalidateGeometry(nodes);
+	}
+
+	virtual void MappingChanged(NodeKeyTab& nodes) override
+	{
+		LogNodeEvent(L"MappingChanged", nodes);
+
+		InvalidateGeometry(nodes);
+	}
+
+	void InvalidateGeometry(NodeKeyTab& nodes)
+	{
 		for (int NodeIndex = 0; NodeIndex < nodes.Count(); ++NodeIndex)
 		{
 			SceneTracker.NodeGeometryChanged(nodes[NodeIndex]);
@@ -122,6 +141,15 @@ public:
 		}
 	}
 
+	virtual void LinkChanged(NodeKeyTab& nodes) override
+	{
+		LogNodeEvent(L"LinkChanged", nodes);
+		for (int NodeIndex = 0; NodeIndex < nodes.Count(); ++NodeIndex)
+		{
+			SceneTracker.NodeLinkChanged(nodes[NodeIndex]);
+		}
+	}
+
 	// Not used:
 
 	virtual void Added(NodeKeyTab& nodes) override
@@ -134,15 +162,6 @@ public:
 		LogNodeEvent(L"Deleted", nodes);
 	}
 
-	virtual void LinkChanged(NodeKeyTab& nodes) override
-	{
-		LogNodeEvent(L"LinkChanged", nodes);
-		for (int NodeIndex = 0; NodeIndex < nodes.Count(); ++NodeIndex)
-		{
-			SceneTracker.NodeLinkChanged(nodes[NodeIndex]);
-		}
-	}
-
 	virtual void GroupChanged(NodeKeyTab& nodes) override
 	{
 		LogNodeEvent(L"GroupChanged", nodes);
@@ -158,16 +177,6 @@ public:
 		LogNodeEvent(L"ModelStructured", nodes);
 	}
 
-	virtual void TopologyChanged(NodeKeyTab& nodes) override
-	{
-		LogNodeEvent(L"TopologyChanged", nodes);
-	}
-
-	virtual void MappingChanged(NodeKeyTab& nodes) override
-	{
-		LogNodeEvent(L"MappingChanged", nodes);
-	}
-
 	virtual void ExtentionChannelChanged(NodeKeyTab& nodes) override
 	{
 		LogNodeEvent(L"ExtentionChannelChanged", nodes);
@@ -327,9 +336,15 @@ public:
 	RefResult NotifyRefChanged(const Interval& ChangeInterval, RefTargetHandle TargetHandle, PartID& PartId, RefMessage Message, BOOL propagate) override
 	{
 		// todo: remove material handling???
-		ensure(ReferencedItemToIndex.Contains(TargetHandle));
+		if (!ReferencedItemToIndex.Contains(TargetHandle))
+		{
+			LogDebugNode(TEXT("FNodeObserver::NotifyRefChanged(MISSING)"), dynamic_cast<INode*>(TargetHandle));
+		}
+		else
+		{
+			LogDebugNode(TEXT("FNodeObserver::NotifyRefChanged"), dynamic_cast<INode*>(TargetHandle));
+		}
 
-		LOG_DEBUG_HEAVY(FString::Printf(TEXT("FNodeObserver::NotifyRefChanged: %s: %x"), dynamic_cast<INode*>(TargetHandle)->GetName(), Message)); // heavy logging - called a lot
 		return REF_SUCCEED;
 	}
 
