Revit: Pressing Datasmith Sync on Revit leads to error

diff --git a/Engine/Source/Programs/Enterprise/Datasmith/DatasmithRevitExporter/Private/DatasmithRevitDocumentData.cs b/Engine/Source/Programs/Enterprise/Datasmith/DatasmithRevitExporter/Private/DatasmithRevitDocumentData.cs
--- a/Engine/Source/Programs/Enterprise/Datasmith/DatasmithRevitExporter/Private/DatasmithRevitDocumentData.cs
+++ b/Engine/Source/Programs/Enterprise/Datasmith/DatasmithRevitExporter/Private/DatasmithRevitDocumentData.cs
@@ -809,7 +809,7 @@ namespace DatasmithRevitExporter
 					// Hash the Datasmith mesh actor name to shorten it.
 					string HashedActorName = FDatasmithFacadeElement.GetStringHash("A:" + GetActorName(true));
 
-					if (BaseElementType.FamilyName == "Decal")
+					if (BaseElementType != null && BaseElementType.FamilyName == "Decal")
 					{
 						bOptimizeHierarchy = false;
 						bIsDecalElement = true;
