[3ds Max] DirectLink: Textures Sync

- Changes to Texmaps are synced(bitmap and other types, bakeable or not)
- No unused textures left in the udatasmith file(any situation when texture was used but not anymore)
- Only textures referenced in converted material are exported. This is an improvement over old exporter where it exported all texmaps referenced by source Max material and then cleaned up scene.

diff --git a/Engine/Source/Programs/Enterprise/Datasmith/DatasmithMaxExporter/Private/DatasmithMaxCorona.cpp b/Engine/Source/Programs/Enterprise/Datasmith/DatasmithMaxExporter/Private/DatasmithMaxCorona.cpp
--- a/Engine/Source/Programs/Enterprise/Datasmith/DatasmithMaxExporter/Private/DatasmithMaxCorona.cpp
+++ b/Engine/Source/Programs/Enterprise/Datasmith/DatasmithMaxExporter/Private/DatasmithMaxCorona.cpp
@@ -67,7 +67,7 @@ FString FDatasmithMaxMatWriter::DumpBitmapCorona(TSharedPtr<IDatasmithCompositeT
 	return *Base;
 }
 
-void FDatasmithMaxMatWriter::GetCoronaTexmap(TSharedRef< IDatasmithScene > DatasmithScene, BitmapTex* InBitmapTex)
+void FDatasmithMaxMatWriter::GetCoronaTexmap(TSharedRef< IDatasmithScene > DatasmithScene, BitmapTex* InBitmapTex, TArray<TSharedPtr< IDatasmithTextureElement >>* OutTextureElements)
 {
 	FString Path = TEXT("");
 
@@ -117,6 +117,10 @@ void FDatasmithMaxMatWriter::GetCoronaTexmap(TSharedRef< IDatasmithScene > Datas
 	}
 	TextureElement->SetFile(*Path);
 	DatasmithScene->AddTexture(TextureElement);
+	if (OutTextureElements)
+	{
+		OutTextureElements->Add(TextureElement);
+	}
 }
 
 void FDatasmithMaxMatWriter::ExportCoronaMaterial(TSharedRef< IDatasmithScene > DatasmithScene, TSharedPtr< IDatasmithMaterialElement >& MaterialElement, Mtl* Material)
diff --git a/Engine/Source/Programs/Enterprise/Datasmith/DatasmithMaxExporter/Private/DatasmithMaxDirectLinkMaterial.cpp b/Engine/Source/Programs/Enterprise/Datasmith/DatasmithMaxExporter/Private/DatasmithMaxDirectLinkMaterial.cpp
--- a/Engine/Source/Programs/Enterprise/Datasmith/DatasmithMaxExporter/Private/DatasmithMaxDirectLinkMaterial.cpp
+++ b/Engine/Source/Programs/Enterprise/Datasmith/DatasmithMaxExporter/Private/DatasmithMaxDirectLinkMaterial.cpp
@@ -7,6 +7,7 @@
 #include "DatasmithMaxHelper.h"
 #include "DatasmithMaxWriter.h"
 #include "DatasmithMaxClassIDs.h"
+#include "MaxMaterialsToUEPbr/DatasmithMaxMaterialsToUEPbr.h"
 
 
 #include "DatasmithSceneFactory.h"
@@ -87,39 +88,7 @@ public:
 			{
 				MaterialEnum(Material->GetSubMtl(i), bAddRecursively);
 			}
-
-			for (int i = 0; i < Material->NumSubTexmaps(); i++)
-			{
-				Texmap* SubTexture = Material->GetSubTexmap(i);
-				if (SubTexture != NULL)
-				{
-					TexEnum(SubTexture);
-				}
-			}
-		}
-	}
-
-	void TexEnum(Texmap* Texture)
-	{
-		if (Texture == NULL)
-		{
-			return;
-		}
-
-		if (!MaterialsCollectionTracker.EncounteredTextures.Contains(Texture))
-		{
-			MaterialsCollectionTracker.EncounteredTextures.Add(Texture);
 		}
-
-		for (int i = 0; i < Texture->NumSubTexmaps(); i++)
-		{
-			Texmap* SubTexture = Texture->GetSubTexmap(i);
-			if (SubTexture != NULL)
-			{
-				TexEnum(SubTexture);
-			}
-		}
-		MaterialTracker.AddActualTexture(Texture);
 	}
 };
 
@@ -133,8 +102,10 @@ void FMaterialsCollectionTracker::Reset()
 	MaterialNames.Reset();
 
 	UsedMaterialToMaterialTracker.Reset();
-
 	UsedMaterialToDatasmithMaterial.Reset();
+
+	UsedTextureToMaterialTracker.Reset();
+	UsedTextureToDatasmithElement.Reset();
 }
 
 void FMaterialsCollectionTracker::UpdateMaterial(FMaterialTracker* MaterialTracker)
@@ -147,6 +118,53 @@ void FMaterialsCollectionTracker::UpdateMaterial(FMaterialTracker* MaterialTrack
 	}
 }
 
+void FMaterialsCollectionTracker::ConvertMaterial(Mtl* Material, TSharedRef<IDatasmithScene> DatasmithScene, const TCHAR* AssetsPath, TSet<Texmap*>& TexmapsConverted)
+{
+	// todo: reworking material export - composite materials(like VRayBlend) export submaterials too
+	//   check below is needed to make sure that no material is exported twice.
+	//   This is how it was handled before bu needs to be handled without going through the list of all materials.
+	// Names should be unique prior to this
+	FString MaterialName(FDatasmithUtils::SanitizeObjectName(Material->GetName().data()));
+	for (int i = 0; i < DatasmithScene->GetMaterialsCount(); i++)
+	{
+		if (FString(DatasmithScene->GetMaterial(i)->GetName()) == MaterialName)
+		{
+			return;
+		}
+	}
+
+	TArray<Texmap*> TexmapsUsedByMaterial;
+	TSharedPtr<IDatasmithBaseMaterialElement> DatasmithMaterial;
+
+	if (FDatasmithMaxMaterialsToUEPbr* MaterialConverter = FDatasmithMaxMaterialsToUEPbrManager::GetMaterialConverter(Material))
+	{
+
+		TGuardValue<TArray<Texmap*>*> TexmapsConvertedGuard(MaterialConverter->TexmapsConverted, &TexmapsUsedByMaterial);
+		 
+		MaterialConverter->Convert(DatasmithScene, DatasmithMaterial, Material, AssetsPath);
+
+		if (DatasmithMaterial)
+		{
+			DatasmithScene->AddMaterial(DatasmithMaterial);
+			UsedMaterialToDatasmithMaterial.Add(Material, DatasmithMaterial);
+		}
+	}
+
+	// Tie texture used by an actual material to tracked material
+	for (Texmap* Texture : TexmapsUsedByMaterial)
+	{
+
+		for (FMaterialTracker* MaterialTracker : UsedMaterialToMaterialTracker[Material])
+		{
+			MaterialTracker->AddActualTexture(Texture);
+			UsedTextureToMaterialTracker.FindOrAdd(Texture).Add(MaterialTracker);
+		}
+	}
+
+	TexmapsConverted.Append(TexmapsUsedByMaterial);
+
+}
+
 void FMaterialsCollectionTracker::ReleaseMaterial(FMaterialTracker& MaterialTracker)
 {
 	RemoveConvertedMaterial(MaterialTracker);
@@ -154,6 +172,7 @@ void FMaterialsCollectionTracker::ReleaseMaterial(FMaterialTracker& MaterialTrac
 	InvalidatedMaterialTrackers.Remove(&MaterialTracker);
 }
 
+
 void FMaterialsCollectionTracker::RemoveConvertedMaterial(FMaterialTracker& MaterialTracker)
 {
 	for (Mtl* Material: MaterialTracker.GetActualMaterials())
@@ -172,6 +191,23 @@ void FMaterialsCollectionTracker::RemoveConvertedMaterial(FMaterialTracker& Mate
 		}
 	}
 
+	for (Texmap* Texture: MaterialTracker.GetActualTexmaps())
+	{
+		TSet<FMaterialTracker*>& MaterialTrackersForTexture = UsedTextureToMaterialTracker[Texture];
+		MaterialTrackersForTexture.Remove(&MaterialTracker);
+
+		if (!MaterialTrackersForTexture.Num()) // No tracked materials are using this texture anymore
+		{
+			UsedTextureToMaterialTracker.Remove(Texture);
+
+			for (const TSharedPtr<IDatasmithTextureElement>& TextureElement: UsedTextureToDatasmithElement[Texture])
+			{
+				SceneTracker.RemoveTexture(TextureElement);
+			}
+			UsedTextureToDatasmithElement.Remove(Texture);
+		}
+	}
+
 	MaterialTracker.ResetActualMaterialAndTextures();
 }
 
diff --git a/Engine/Source/Programs/Enterprise/Datasmith/DatasmithMaxExporter/Private/DatasmithMaxExporterWithDirectLink.cpp b/Engine/Source/Programs/Enterprise/Datasmith/DatasmithMaxExporter/Private/DatasmithMaxExporterWithDirectLink.cpp
--- a/Engine/Source/Programs/Enterprise/Datasmith/DatasmithMaxExporter/Private/DatasmithMaxExporterWithDirectLink.cpp
+++ b/Engine/Source/Programs/Enterprise/Datasmith/DatasmithMaxExporter/Private/DatasmithMaxExporterWithDirectLink.cpp
@@ -857,9 +857,9 @@ public:
 			InvalidatedNodeTrackers.Reset();
 		}
 
-		TSet<Mtl*> ActualMaterialToUpdate;
-		TSet<Texmap*> ActualTexmapsToUpdate;
-
+		// Each tracked(i.e. assigned to a visible node) Max material can result in multiple actual materials
+		// e.g. an assigned MultiSubObj material is aclually a set of material not directly assigned to a node
+		TSet<Mtl*> ActualMaterialToUpdate; 
 		{
 			PROGRESS_STAGE("Process invalidated materials");
 			{
@@ -875,37 +875,34 @@ public:
 						ActualMaterialToUpdate.Add(ActualMaterial);
 					}
 					MaterialTracker->bInvalidated = false;
-					for (Texmap* Texture: MaterialTracker->Textures)
-					{
-						ActualTexmapsToUpdate.Add(Texture);
-					}
 				}
 				MaterialsCollectionTracker.ResetInvalidatedMaterials();
 			}
 		}
 
+		TSet<Texmap*> ActualTexmapsToUpdate;
 		{
-			PROGRESS_STAGE("Update textures")
-			PROGRESS_STAGE_COUNTER(ActualTexmapsToUpdate.Num());
-			for (Texmap* Texture : ActualTexmapsToUpdate)
+			PROGRESS_STAGE("Update materials")
+			PROGRESS_STAGE_COUNTER(ActualMaterialToUpdate.Num());
+			for (Mtl* ActualMaterial: ActualMaterialToUpdate)
 			{
 				ProgressCounter.Next();
-				FDatasmithMaxMatExport::GetXMLTexture(ExportedScene.GetDatasmithScene(), Texture, ExportedScene.GetSceneExporter().GetAssetsOutputPath());
+
+				MaterialsCollectionTracker.ConvertMaterial(ActualMaterial, ExportedScene.GetDatasmithScene(), ExportedScene.GetSceneExporter().GetAssetsOutputPath(), ActualTexmapsToUpdate);
 			}
 		}
 
 		{
-			PROGRESS_STAGE("Update materials")
-			PROGRESS_STAGE_COUNTER(ActualMaterialToUpdate.Num());
-			for (Mtl* ActualMaterial : ActualMaterialToUpdate)
+			PROGRESS_STAGE("Update textures")
+			PROGRESS_STAGE_COUNTER(ActualTexmapsToUpdate.Num());
+			
+			for (Texmap* Texture : ActualTexmapsToUpdate)
 			{
 				ProgressCounter.Next();
 
-				// todo: make sure not reexport submaterial more than once - i.e. when a submaterial is used in two composite material
-				FDatasmithMaxMatExport::bForceReexport = true;
-				TSharedPtr<IDatasmithBaseMaterialElement> DatastmihMaterial = FDatasmithMaxMatExport::ExportUniqueMaterial(ExportedScene.GetDatasmithScene(), ActualMaterial, ExportedScene.GetSceneExporter().GetAssetsOutputPath());
-
-				MaterialsCollectionTracker.SetDatasmithMaterial(ActualMaterial, DatastmihMaterial);
+				TArray<TSharedPtr<IDatasmithTextureElement>> TextureElements;
+				FDatasmithMaxMatExport::GetXMLTexture(ExportedScene.GetDatasmithScene(), Texture, ExportedScene.GetSceneExporter().GetAssetsOutputPath(), &TextureElements);
+				MaterialsCollectionTracker.UsedTextureToDatasmithElement.Add(Texture, TextureElements);
 			}
 		}
 
@@ -987,6 +984,11 @@ public:
 		ExportedScene.DatasmithSceneRef->RemoveMaterial(DatasmithMaterial);		
 	}
 
+	virtual void RemoveTexture(const TSharedPtr<IDatasmithTextureElement>& DatasmithTextureElement) override
+	{
+		ExportedScene.DatasmithSceneRef->RemoveTexture(DatasmithTextureElement);		
+	}
+
 	void InvalidateNode(FNodeTracker& NodeTracker)
 	{
 		ensure(!bUpdateInProgress);
diff --git a/Engine/Source/Programs/Enterprise/Datasmith/DatasmithMaxExporter/Private/DatasmithMaxScanline.cpp b/Engine/Source/Programs/Enterprise/Datasmith/DatasmithMaxExporter/Private/DatasmithMaxScanline.cpp
--- a/Engine/Source/Programs/Enterprise/Datasmith/DatasmithMaxExporter/Private/DatasmithMaxScanline.cpp
+++ b/Engine/Source/Programs/Enterprise/Datasmith/DatasmithMaxExporter/Private/DatasmithMaxScanline.cpp
@@ -452,7 +452,7 @@ FString FDatasmithMaxMatWriter::DumpAutodeskBitmap(TSharedPtr<IDatasmithComposit
 	return ActualBitmapName;
 }
 
-void FDatasmithMaxMatWriter::GetRegularTexmap(TSharedRef< IDatasmithScene > DatasmithScene, BitmapTex* InBitmapTex)
+void FDatasmithMaxMatWriter::GetRegularTexmap(TSharedRef< IDatasmithScene > DatasmithScene, BitmapTex* InBitmapTex, TArray<TSharedPtr< IDatasmithTextureElement >>* OutTextureElements)
 {
 	if (!InBitmapTex)
 	{
@@ -487,9 +487,13 @@ void FDatasmithMaxMatWriter::GetRegularTexmap(TSharedRef< IDatasmithScene > Data
 
 	TextureElement->SetFile(*Path);
 	DatasmithScene->AddTexture(TextureElement);
+	if (OutTextureElements)
+	{
+		OutTextureElements->Add(TextureElement);
+	}
 }
 
-void FDatasmithMaxMatWriter::GetAutodeskTexmap(TSharedRef< IDatasmithScene > DatasmithScene, Texmap* InTexMap)
+void FDatasmithMaxMatWriter::GetAutodeskTexmap(TSharedRef< IDatasmithScene > DatasmithScene, Texmap* InTexMap, TArray<TSharedPtr< IDatasmithTextureElement >>* OutTextureElements)
 {
 	if (PBBitmap* BitmapSourceFile = DatasmithMaxTexmapParser::ParseAutodeskBitmap(InTexMap).SourceFile)
 	{
@@ -527,6 +531,10 @@ void FDatasmithMaxMatWriter::GetAutodeskTexmap(TSharedRef< IDatasmithScene > Dat
 
 		TextureElement->SetFile(*Path);
 		DatasmithScene->AddTexture(TextureElement);
+		if (OutTextureElements)
+		{
+			OutTextureElements->Add(TextureElement);
+		}
 	}
 }
 
diff --git a/Engine/Source/Programs/Enterprise/Datasmith/DatasmithMaxExporter/Private/DatasmithMaxThea.cpp b/Engine/Source/Programs/Enterprise/Datasmith/DatasmithMaxExporter/Private/DatasmithMaxThea.cpp
--- a/Engine/Source/Programs/Enterprise/Datasmith/DatasmithMaxExporter/Private/DatasmithMaxThea.cpp
+++ b/Engine/Source/Programs/Enterprise/Datasmith/DatasmithMaxExporter/Private/DatasmithMaxThea.cpp
@@ -11,10 +11,10 @@ FString FDatasmithMaxMatWriter::DumpBitmapThea(TSharedPtr<IDatasmithCompositeTex
 	return DumpBitmap(CompTex, Texture, Prefix, bForceInvert, bIsGrayscale);
 }
 
-void FDatasmithMaxMatWriter::GetTheaTexmap(TSharedRef< IDatasmithScene > DatasmithScene, BitmapTex* InBitmapTex)
+void FDatasmithMaxMatWriter::GetTheaTexmap(TSharedRef< IDatasmithScene > DatasmithScene, BitmapTex* InBitmapTex, TArray<TSharedPtr< IDatasmithTextureElement >>* OutTextureElements)
 {
 	BitmapTex* Texture = (BitmapTex*)InBitmapTex->GetReference(0);
-	GetRegularTexmap(DatasmithScene, Texture);
+	GetRegularTexmap(DatasmithScene, Texture, OutTextureElements);
 }
 
 void FDatasmithMaxMatWriter::ExportTheaSubmaterial(TSharedRef< IDatasmithScene > DatasmithScene, TSharedPtr< IDatasmithShaderElement >& MaterialShader, Mtl* Material, EDSMaterialType MaterialType)
diff --git a/Engine/Source/Programs/Enterprise/Datasmith/DatasmithMaxExporter/Private/DatasmithMaxVRay.cpp b/Engine/Source/Programs/Enterprise/Datasmith/DatasmithMaxExporter/Private/DatasmithMaxVRay.cpp
--- a/Engine/Source/Programs/Enterprise/Datasmith/DatasmithMaxExporter/Private/DatasmithMaxVRay.cpp
+++ b/Engine/Source/Programs/Enterprise/Datasmith/DatasmithMaxExporter/Private/DatasmithMaxVRay.cpp
@@ -123,7 +123,7 @@ FString FDatasmithMaxMatWriter::DumpVrayHdri(TSharedPtr<IDatasmithCompositeTextu
 	return *BitmapName;
 }
 
-void FDatasmithMaxMatWriter::GetVrayHdri(TSharedRef< IDatasmithScene > DatasmithScene, BitmapTex* InBitmapTex)
+void FDatasmithMaxMatWriter::GetVrayHdri(TSharedRef< IDatasmithScene > DatasmithScene, BitmapTex* InBitmapTex, TArray<TSharedPtr< IDatasmithTextureElement >>* OutTextureElements)
 {
 	FString Path = TEXT("");
 
@@ -174,6 +174,10 @@ void FDatasmithMaxMatWriter::GetVrayHdri(TSharedRef< IDatasmithScene > Datasmith
 	}
 	TextureElement->SetFile(*Path);
 	DatasmithScene->AddTexture(TextureElement);
+	if (OutTextureElements)
+	{
+		OutTextureElements->Add(TextureElement);
+	}
 }
 
 void FDatasmithMaxMatWriter::ExportVRayMaterial(TSharedRef< IDatasmithScene > DatasmithScene, TSharedPtr< IDatasmithMaterialElement >& MaterialElement, Mtl* Material)
diff --git a/Engine/Source/Programs/Enterprise/Datasmith/DatasmithMaxExporter/Private/DatasmithMaxWriter.cpp b/Engine/Source/Programs/Enterprise/Datasmith/DatasmithMaxExporter/Private/DatasmithMaxWriter.cpp
--- a/Engine/Source/Programs/Enterprise/Datasmith/DatasmithMaxExporter/Private/DatasmithMaxWriter.cpp
+++ b/Engine/Source/Programs/Enterprise/Datasmith/DatasmithMaxExporter/Private/DatasmithMaxWriter.cpp
@@ -446,33 +446,40 @@ bool FDatasmithMaxMatHelper::IsSRGB(Bitmap& InBitmap)
 
 bool FDatasmithMaxMatExport::bForceReexport = false;
 
-void FDatasmithMaxMatExport::GetXMLTexture(TSharedRef< IDatasmithScene > DatasmithScene, Texmap* InTexMap, const TCHAR* AssetsPath)
+// todo: need wa wya to return created texture element(if there was one), maybe use 'last' in DatasmithScene(check that count changed)
+//   verify that only one texture created!
+void FDatasmithMaxMatExport::GetXMLTexture(TSharedRef< IDatasmithScene > DatasmithScene, Texmap* InTexMap, const TCHAR* AssetsPath, TArray<TSharedPtr< IDatasmithTextureElement >>* OutTextureElements)
 {
 	switch (FDatasmithMaxMatHelper::GetTextureClass(InTexMap))
 	{
 	case EDSBitmapType::NotSupported:
-		return;
 		break;
 	case EDSBitmapType::RegularBitmap:
-		FDatasmithMaxMatWriter::GetRegularTexmap(DatasmithScene, (BitmapTex*)InTexMap);
+		FDatasmithMaxMatWriter::GetRegularTexmap(DatasmithScene, (BitmapTex*)InTexMap, OutTextureElements);
 		break;
 	case EDSBitmapType::AutodeskBitmap:
-		FDatasmithMaxMatWriter::GetAutodeskTexmap(DatasmithScene, InTexMap);
+		FDatasmithMaxMatWriter::GetAutodeskTexmap(DatasmithScene, InTexMap, OutTextureElements);
 		break; 
 	case EDSBitmapType::TheaBitmap:
-		FDatasmithMaxMatWriter::GetTheaTexmap(DatasmithScene, (BitmapTex*)InTexMap);
+		FDatasmithMaxMatWriter::GetTheaTexmap(DatasmithScene, (BitmapTex*)InTexMap, OutTextureElements);
 		break;
 	case EDSBitmapType::CoronaBitmap:
-		FDatasmithMaxMatWriter::GetCoronaTexmap(DatasmithScene, (BitmapTex*)InTexMap);
+		FDatasmithMaxMatWriter::GetCoronaTexmap(DatasmithScene, (BitmapTex*)InTexMap, OutTextureElements);
 		break;
 	case EDSBitmapType::VRayHRDI:
-		FDatasmithMaxMatWriter::GetVrayHdri(DatasmithScene, (BitmapTex*)InTexMap);
+		FDatasmithMaxMatWriter::GetVrayHdri(DatasmithScene, (BitmapTex*)InTexMap, OutTextureElements);
 		break;
 	case EDSBitmapType::ColorCorrector:
-		GetXMLTexture(DatasmithScene, InTexMap->GetSubTexmap(0), AssetsPath);
+		GetXMLTexture(DatasmithScene, InTexMap->GetSubTexmap(0), AssetsPath, OutTextureElements);
 		break;
 	case EDSBitmapType::BakeableMap:
-		FDatasmithMaxMatWriter::AddBakeable(DatasmithScene, (BitmapTex*)InTexMap, AssetsPath);
+		{
+			TSharedPtr<IDatasmithTextureElement> TextureElement = FDatasmithMaxMatWriter::AddBakeable(DatasmithScene, (BitmapTex*)InTexMap, AssetsPath);
+			if (OutTextureElements)
+			{
+				OutTextureElements->Add(TextureElement);
+			}
+		}
 		break;
 	case EDSBitmapType::NormalMap:
 	case EDSBitmapType::FallOff:
@@ -485,7 +492,7 @@ void FDatasmithMaxMatExport::GetXMLTexture(TSharedRef< IDatasmithScene > Datasmi
 	case EDSBitmapType::CompositeTex:
 		for (int i = 0; i < InTexMap->NumSubTexmaps(); i++)
 		{
-			GetXMLTexture(DatasmithScene, InTexMap->GetSubTexmap(i), AssetsPath);
+			GetXMLTexture(DatasmithScene, InTexMap->GetSubTexmap(i), AssetsPath, OutTextureElements);
 		}
 		break;
 	default:
diff --git a/Engine/Source/Programs/Enterprise/Datasmith/DatasmithMaxExporter/Private/MaxMaterialsToUEPbr/DatasmithMaxMaterialsToUEPbr.cpp b/Engine/Source/Programs/Enterprise/Datasmith/DatasmithMaxExporter/Private/MaxMaterialsToUEPbr/DatasmithMaxMaterialsToUEPbr.cpp
--- a/Engine/Source/Programs/Enterprise/Datasmith/DatasmithMaxExporter/Private/MaxMaterialsToUEPbr/DatasmithMaxMaterialsToUEPbr.cpp
+++ b/Engine/Source/Programs/Enterprise/Datasmith/DatasmithMaxExporter/Private/MaxMaterialsToUEPbr/DatasmithMaxMaterialsToUEPbr.cpp
@@ -187,6 +187,11 @@ IDatasmithMaterialExpression* FDatasmithMaxMaterialsToUEPbr::ConvertTexmap( cons
 		return nullptr;
 	}
 
+	if (TexmapsConverted)
+	{
+		TexmapsConverted->Add(MapParameter.Map);
+	}
+
 	// Check if the texmap has a texture ouput
 	TextureOutput* TexOutput = nullptr;
 
diff --git a/Engine/Source/Programs/Enterprise/Datasmith/DatasmithMaxExporter/Private/MaxMaterialsToUEPbr/DatasmithMaxMaterialsToUEPbr.h b/Engine/Source/Programs/Enterprise/Datasmith/DatasmithMaxExporter/Private/MaxMaterialsToUEPbr/DatasmithMaxMaterialsToUEPbr.h
--- a/Engine/Source/Programs/Enterprise/Datasmith/DatasmithMaxExporter/Private/MaxMaterialsToUEPbr/DatasmithMaxMaterialsToUEPbr.h
+++ b/Engine/Source/Programs/Enterprise/Datasmith/DatasmithMaxExporter/Private/MaxMaterialsToUEPbr/DatasmithMaxMaterialsToUEPbr.h
@@ -59,6 +59,8 @@ public:
 
 	} ConvertState;
 
+	TArray<Texmap*>* TexmapsConverted = nullptr;
+
 protected:
 	TIndirectArray< IDatasmithMaxTexmapToUEPbr > TexmapConverters;
 
diff --git a/Engine/Source/Programs/Enterprise/Datasmith/DatasmithMaxExporter/Public/DatasmithMaxDirectLink.h b/Engine/Source/Programs/Enterprise/Datasmith/DatasmithMaxExporter/Public/DatasmithMaxDirectLink.h
--- a/Engine/Source/Programs/Enterprise/Datasmith/DatasmithMaxExporter/Public/DatasmithMaxDirectLink.h
+++ b/Engine/Source/Programs/Enterprise/Datasmith/DatasmithMaxExporter/Public/DatasmithMaxDirectLink.h
@@ -234,6 +234,7 @@ public:
 	virtual void SetupActor(FNodeTracker& NodeTracker) = 0;
 	virtual void SetupDatasmithHISMForNode(FNodeTracker& NodeTracker, INode* GeometryNode, const FRenderMeshForConversion& RenderMesh, Mtl* Material, int32 MeshIndex, const TArray<Matrix3>& Transforms) = 0;
 	virtual void RemoveMaterial(const TSharedPtr<IDatasmithBaseMaterialElement>& DatasmithMaterial) = 0;
+	virtual void RemoveTexture(const TSharedPtr<IDatasmithTextureElement>&) = 0;
 	virtual void NodeXRefMerged(INode* Node) = 0;
 };
 
@@ -356,6 +357,11 @@ public:
 		return Materials;
 	}
 
+	TArray<Texmap*>& GetActualTexmaps()
+	{
+		return Textures;
+	}
+
 	void ResetActualMaterialAndTextures()
 	{
 		Materials.Reset();
@@ -402,16 +408,18 @@ public:
 
 	void Reset();
 
-	FMaterialTracker* AddMaterial(Mtl* Material);
+	FMaterialTracker* AddMaterial(Mtl* Material); // Add material to track its changes
 
-	void InvalidateMaterial(Mtl* Material);
+	void InvalidateMaterial(Mtl* Material); // Mark changed
 
-	void UpdateMaterial(FMaterialTracker* MaterialTracker);
+	void UpdateMaterial(FMaterialTracker* MaterialTracker); // Reparse source material
+
+	void ConvertMaterial(Mtl* Material, TSharedRef<IDatasmithScene> DatasmithScene, const TCHAR* AssetsPath, TSet<Texmap*>& TexmapsConverted); // Convert to Datasmith
 
 	// When Material is not used my scene anymore - stop tracking it
 	void ReleaseMaterial(FMaterialTracker& MaterialTracker);
 
-	// Clean all converted data, remove from datasmith scene(before rebuild)
+	// Clean all converted data, remove from datasmith scene(e.g. before rebuilding material)
 	void RemoveConvertedMaterial(FMaterialTracker& MaterialTracker);
 
 	void ResetInvalidatedMaterials()
@@ -424,11 +432,6 @@ public:
 		return InvalidatedMaterialTrackers;
 	}
 
-	void SetDatasmithMaterial(Mtl* ActualMaterial, TSharedPtr<IDatasmithBaseMaterialElement> DatastmihMaterial)
-	{
-		UsedMaterialToDatasmithMaterial.Add(ActualMaterial, DatastmihMaterial);
-	}
-
 	ISceneTracker& SceneTracker;
 
 	TMap<FMaterialKey, FMaterialTrackerHandle> MaterialTrackers; // Tracks all assigned materials
@@ -441,6 +444,11 @@ public:
 
 	TMap<Mtl*, TSet<FMaterialTracker*>> UsedMaterialToMaterialTracker; // Materials uses by nodes keep set of assigned materials they are used for
 	TMap<Mtl*, TSharedPtr<IDatasmithBaseMaterialElement>> UsedMaterialToDatasmithMaterial;
+
+	TMap<Texmap*, TSet<FMaterialTracker*>> UsedTextureToMaterialTracker; // Materials uses by nodes keep set of assigned textures they are used for
+	TMap<Texmap*, TArray<TSharedPtr<IDatasmithTextureElement>>> UsedTextureToDatasmithElement; // Keep track of Datasmith Element created for texmap to simplify update/removal(no need to search DatasmithScene)
+	// note: each texmap can create multiple texture elements
+
 };
 
 //----
diff --git a/Engine/Source/Programs/Enterprise/Datasmith/DatasmithMaxExporter/Public/DatasmithMaxWriter.h b/Engine/Source/Programs/Enterprise/Datasmith/DatasmithMaxExporter/Public/DatasmithMaxWriter.h
--- a/Engine/Source/Programs/Enterprise/Datasmith/DatasmithMaxExporter/Public/DatasmithMaxWriter.h
+++ b/Engine/Source/Programs/Enterprise/Datasmith/DatasmithMaxExporter/Public/DatasmithMaxWriter.h
@@ -103,7 +103,7 @@ class FDatasmithMaxMatExport
 {
 public:
 	// stores the textures without allowing dupe in uniqueTextures
-	static void GetXMLTexture(TSharedRef< IDatasmithScene > DatasmithScene, Texmap* InTexmap, const TCHAR* AssetsPath);
+	static void GetXMLTexture(TSharedRef< IDatasmithScene > DatasmithScene, Texmap* InTexmap, const TCHAR* AssetsPath, TArray<TSharedPtr< IDatasmithTextureElement >>* OutTextureElements=nullptr);
 		
 	// Exports the given material if it hasn't been exported yet and returns the created IDatasmithBaseMaterialElement, else returns an invalid TSharedPtr.
 	static TSharedPtr< IDatasmithBaseMaterialElement > ExportUniqueMaterial(TSharedRef< IDatasmithScene > DatasmithScene, Mtl* Material, const TCHAR* AssetsPath);
@@ -144,8 +144,8 @@ public:
 	static FString CropBitmap(BitmapTex* InBitmapTex);
 	static FString DumpBitmap(TSharedPtr<IDatasmithCompositeTexture>& CompTex, BitmapTex* InBitmapTex, const TCHAR* Prefix, bool bForceInvert, bool bIsGrayscale);
 	static FString DumpAutodeskBitmap(TSharedPtr<IDatasmithCompositeTexture>& CompTex, Texmap* InTexmap, const TCHAR* Prefix, bool bForceInvert, bool bIsGrayscale);
-	static void GetRegularTexmap(TSharedRef< IDatasmithScene > DatasmithScene, BitmapTex* InBitmapTex);
-	static void GetAutodeskTexmap(TSharedRef< IDatasmithScene > DatasmithScene, Texmap* InTexMap);
+	static void GetRegularTexmap(TSharedRef< IDatasmithScene > DatasmithScene, BitmapTex* InBitmapTex, TArray<TSharedPtr< IDatasmithTextureElement >>* OutTextureElements);
+	static void GetAutodeskTexmap(TSharedRef< IDatasmithScene > DatasmithScene, Texmap* InTexMap, TArray<TSharedPtr< IDatasmithTextureElement >>* OutTextureElements);
 	static TSharedPtr< IDatasmithTextureElement > AddBakeable(TSharedRef< IDatasmithScene > DatasmithScene, Texmap* InTexmap, const TCHAR* AssetsPath);
 	static void ExportStandardMaterial(TSharedRef< IDatasmithScene > DatasmithScene, TSharedPtr< IDatasmithMaterialElement >& MaterialElement, Mtl* Material);
 	static void ExportBlendMaterial(TSharedRef< IDatasmithScene > DatasmithScene, TSharedPtr< IDatasmithMaterialElement >& MaterialElement, Mtl* Material);
@@ -159,7 +159,7 @@ public:
 	// THEA DEPENDANT
 	//---------------------------------------------------------------
 	static FString DumpBitmapThea(TSharedPtr<IDatasmithCompositeTexture>& CompTex, BitmapTex* InBitmapTex, const TCHAR* Prefix, bool bForceInvert, bool bIsGrayscale);
-	static void GetTheaTexmap(TSharedRef< IDatasmithScene > DatasmithScene, BitmapTex* InBitmapTex);
+	static void GetTheaTexmap(TSharedRef< IDatasmithScene > DatasmithScene, BitmapTex* InBitmapTex, TArray<TSharedPtr< IDatasmithTextureElement >>* OutTextureElements);
 	static void ExportTheaSubmaterial(TSharedRef< IDatasmithScene > DatasmithScene, TSharedPtr< IDatasmithShaderElement >& MaterialShader, Mtl* Material, EDSMaterialType MaterialType);
 	static void ExportTheaMaterial(TSharedRef< IDatasmithScene > DatasmithScene, TSharedPtr< IDatasmithMaterialElement >& MaterialElement, Mtl* Material);
 
@@ -167,7 +167,7 @@ public:
 	// VRAY DEPENDANT
 	//---------------------------------------------------------------
 	static FString DumpVrayColor(TSharedPtr<IDatasmithCompositeTexture>& CompTex, Texmap* InTexmap, const TCHAR* Prefix, bool bForceInvert);
-	static void GetVrayHdri(TSharedRef< IDatasmithScene > DatasmithScene, BitmapTex* InBitmapTex);
+	static void GetVrayHdri(TSharedRef< IDatasmithScene > DatasmithScene, BitmapTex* InBitmapTex, TArray<TSharedPtr< IDatasmithTextureElement >>* OutTextureElements);
 	static FString DumpVrayHdri(TSharedPtr<IDatasmithCompositeTexture>& CompTex, BitmapTex* InBitmapTex, const TCHAR* Prefix, bool bForceInvert);
 	static void ExportVRayMaterial(TSharedRef< IDatasmithScene > DatasmithScene, TSharedPtr< IDatasmithMaterialElement >& MaterialElement, Mtl* Material);
 	static void ExportVRayLightMaterial(TSharedRef< IDatasmithScene > DatasmithScene, TSharedPtr< IDatasmithMaterialElement >& MaterialElement, Mtl* Material);
@@ -184,7 +184,7 @@ public:
 	//---------------------------------------------------------------
 	static FString DumpBitmapCorona(TSharedPtr<IDatasmithCompositeTexture>& CompTex, BitmapTex* InBitmapTex, const TCHAR* Prefix, bool bForceInvert, bool bIsGrayscale);
 	static FString DumpCoronaColor(TSharedPtr<IDatasmithCompositeTexture>& CompTex, Texmap* InTexmap, const TCHAR* Prefix, bool bForceInvert);
-	static void GetCoronaTexmap(TSharedRef< IDatasmithScene > DatasmithScene, BitmapTex* InBitmapTex);
+	static void GetCoronaTexmap(TSharedRef< IDatasmithScene > DatasmithScene, BitmapTex* InBitmapTex, TArray<TSharedPtr< IDatasmithTextureElement >>* OutTextureElements);
 	static void ExportCoronaMaterial(TSharedRef< IDatasmithScene > DatasmithScene, TSharedPtr< IDatasmithMaterialElement >& MaterialElement, Mtl* Material);
 	static void ExportCoronaLightMaterial(TSharedRef< IDatasmithScene > DatasmithScene, TSharedPtr< IDatasmithMaterialElement >& MaterialElement, Mtl* Material);
 	static void ExportCoronaBlendMaterial(TSharedRef< IDatasmithScene > DatasmithScene,  TSharedPtr< IDatasmithMaterialElement >& MaterialElement, Mtl* Material);
