DatasmithSDK: adopt latest version of warning about includes order

diff --git a/Engine/Source/Programs/Enterprise/Datasmith/DatasmithSDK/DatasmithSDK.Target.cs b/Engine/Source/Programs/Enterprise/Datasmith/DatasmithSDK/DatasmithSDK.Target.cs
--- a/Engine/Source/Programs/Enterprise/Datasmith/DatasmithSDK/DatasmithSDK.Target.cs
+++ b/Engine/Source/Programs/Enterprise/Datasmith/DatasmithSDK/DatasmithSDK.Target.cs
@@ -9,6 +9,8 @@ public class DatasmithSDKTarget : TargetRules
 		: base(Target)
 	{
 		Type = TargetType.Program;
+		IncludeOrderVersion = EngineIncludeOrderVersion.Latest;
+		DefaultBuildSettings = BuildSettingsVersion.Latest;
 		SolutionDirectory = "Programs/Datasmith";
 		bBuildInSolutionByDefault = false;
 
@@ -41,7 +43,7 @@ public class DatasmithSDKTarget : TargetRules
 	public void PostBuildCopy(string SrcPath, string DestPath)
 	{
 		PostBuildSteps.Add(string.Format("echo Copying {0} to {1}", SrcPath, DestPath));
-		PostBuildSteps.Add(string.Format("xcopy \"{0}\" \"{1}\" /R /Y /S", SrcPath, DestPath));
+		PostBuildSteps.Add(string.Format("xcopy \"{0}\" \"{1}\" /R /Y /S /Q", SrcPath, DestPath));
 	}
 
 	public void AddWindowsPostBuildSteps()
