Datasmith: adopt latest version of warning about includes order

diff --git a/Engine/Source/Programs/Enterprise/Datasmith/DatasmithARCHICADExporter/DatsmithUE4ArchiCAD/DatasmithUE4ArchiCAD.Target.cs b/Engine/Source/Programs/Enterprise/Datasmith/DatasmithARCHICADExporter/DatsmithUE4ArchiCAD/DatasmithUE4ArchiCAD.Target.cs
--- a/Engine/Source/Programs/Enterprise/Datasmith/DatasmithARCHICADExporter/DatsmithUE4ArchiCAD/DatasmithUE4ArchiCAD.Target.cs
+++ b/Engine/Source/Programs/Enterprise/Datasmith/DatasmithARCHICADExporter/DatsmithUE4ArchiCAD/DatasmithUE4ArchiCAD.Target.cs
@@ -9,6 +9,7 @@ public class DatasmithUE4ArchiCADTarget : TargetRules
 		: base(Target)
 	{
 		Type = TargetType.Program;
+		IncludeOrderVersion = EngineIncludeOrderVersion.Latest;
 		SolutionDirectory = "Programs/Datasmith";
 		bBuildInSolutionByDefault = false;
 
diff --git a/Engine/Source/Programs/Enterprise/Datasmith/DatasmithARCHICADExporter/UBTModules/DatasmithArchiCAD.Target.cs b/Engine/Source/Programs/Enterprise/Datasmith/DatasmithARCHICADExporter/UBTModules/DatasmithArchiCAD.Target.cs
--- a/Engine/Source/Programs/Enterprise/Datasmith/DatasmithARCHICADExporter/UBTModules/DatasmithArchiCAD.Target.cs
+++ b/Engine/Source/Programs/Enterprise/Datasmith/DatasmithARCHICADExporter/UBTModules/DatasmithArchiCAD.Target.cs
@@ -10,6 +10,7 @@ public class DatasmithArchiCADTarget : TargetRules
 		: base(Target)
 	{
 		Type = TargetType.Program;
+		IncludeOrderVersion = EngineIncludeOrderVersion.Latest;
 		SolutionDirectory = "Programs/Datasmith";
 		bBuildInSolutionByDefault = false;
 
diff --git a/Engine/Source/Programs/Enterprise/Datasmith/DatasmithCADWorker/DatasmithCADWorker.Target.cs b/Engine/Source/Programs/Enterprise/Datasmith/DatasmithCADWorker/DatasmithCADWorker.Target.cs
--- a/Engine/Source/Programs/Enterprise/Datasmith/DatasmithCADWorker/DatasmithCADWorker.Target.cs
+++ b/Engine/Source/Programs/Enterprise/Datasmith/DatasmithCADWorker/DatasmithCADWorker.Target.cs
@@ -10,6 +10,7 @@ public class DatasmithCADWorkerTarget : TargetRules
 	{
 		Type = TargetType.Program;
 		LinkType = TargetLinkType.Monolithic;
+		IncludeOrderVersion = EngineIncludeOrderVersion.Latest;
 		LaunchModuleName = "DatasmithCADWorker";
 		SolutionDirectory = "Programs/Datasmith";
 
diff --git a/Engine/Source/Programs/Enterprise/Datasmith/DatasmithFacadeCSharp/DatasmithFacadeCSharp.Target.cs b/Engine/Source/Programs/Enterprise/Datasmith/DatasmithFacadeCSharp/DatasmithFacadeCSharp.Target.cs
--- a/Engine/Source/Programs/Enterprise/Datasmith/DatasmithFacadeCSharp/DatasmithFacadeCSharp.Target.cs
+++ b/Engine/Source/Programs/Enterprise/Datasmith/DatasmithFacadeCSharp/DatasmithFacadeCSharp.Target.cs
@@ -11,6 +11,7 @@ public class DatasmithFacadeCSharpTarget : TargetRules
 		: base(Target)
 	{
 		Type = TargetType.Program;
+		IncludeOrderVersion = EngineIncludeOrderVersion.Latest;
 		SolutionDirectory = Path.Combine("Programs","Datasmith");
 		bBuildInSolutionByDefault = false;
 
diff --git a/Engine/Source/Programs/Enterprise/Datasmith/DatasmithMaxExporter/DatasmithMax2017.Target.cs b/Engine/Source/Programs/Enterprise/Datasmith/DatasmithMaxExporter/DatasmithMax2017.Target.cs
--- a/Engine/Source/Programs/Enterprise/Datasmith/DatasmithMaxExporter/DatasmithMax2017.Target.cs
+++ b/Engine/Source/Programs/Enterprise/Datasmith/DatasmithMaxExporter/DatasmithMax2017.Target.cs
@@ -9,6 +9,7 @@ public abstract class DatasmithMaxBaseTarget : TargetRules
 		: base(Target)
 	{
 		Type = TargetType.Program;
+		IncludeOrderVersion = EngineIncludeOrderVersion.Latest;
 		SolutionDirectory = "Programs/Datasmith";
 		bBuildInSolutionByDefault = false;
 		bLegalToDistributeBinary = true;
diff --git a/Engine/Source/Programs/Enterprise/Datasmith/DatasmithNavisworksExporter/UBTModules/DatasmithNavisworks2020.Target.cs b/Engine/Source/Programs/Enterprise/Datasmith/DatasmithNavisworksExporter/UBTModules/DatasmithNavisworks2020.Target.cs
--- a/Engine/Source/Programs/Enterprise/Datasmith/DatasmithNavisworksExporter/UBTModules/DatasmithNavisworks2020.Target.cs
+++ b/Engine/Source/Programs/Enterprise/Datasmith/DatasmithNavisworksExporter/UBTModules/DatasmithNavisworks2020.Target.cs
@@ -10,11 +10,12 @@ public abstract class DatasmithNavisworksBaseTarget : TargetRules
 		: base(Target)
 	{
 		Type = TargetType.Program;
+		IncludeOrderVersion = EngineIncludeOrderVersion.Latest;
 		SolutionDirectory = "Programs/Datasmith";
 		bBuildInSolutionByDefault = false;
 
 		string NavisworksVersionString = GetVersion();
-		
+
 		ExeBinariesSubFolder = Path.Combine("Navisworks", NavisworksVersionString);
 		LaunchModuleName = "DatasmithNavisworks" + NavisworksVersionString;
 
diff --git a/Engine/Source/Programs/Enterprise/Datasmith/DatasmithRevitExporter/DatasmithRevit2018.Target.cs b/Engine/Source/Programs/Enterprise/Datasmith/DatasmithRevitExporter/DatasmithRevit2018.Target.cs
--- a/Engine/Source/Programs/Enterprise/Datasmith/DatasmithRevitExporter/DatasmithRevit2018.Target.cs
+++ b/Engine/Source/Programs/Enterprise/Datasmith/DatasmithRevitExporter/DatasmithRevit2018.Target.cs
@@ -10,6 +10,7 @@ public abstract class DatasmithRevitBaseTarget : TargetRules
 		: base(Target)
 	{
 		Type = TargetType.Program;
+		IncludeOrderVersion = EngineIncludeOrderVersion.Latest;
 		SolutionDirectory = "Programs/Datasmith";
 		bBuildInSolutionByDefault = false;
 
diff --git a/Engine/Source/Programs/Enterprise/Datasmith/DatasmithRhinoExporter/DatasmithRhino6.Target.cs b/Engine/Source/Programs/Enterprise/Datasmith/DatasmithRhinoExporter/DatasmithRhino6.Target.cs
--- a/Engine/Source/Programs/Enterprise/Datasmith/DatasmithRhinoExporter/DatasmithRhino6.Target.cs
+++ b/Engine/Source/Programs/Enterprise/Datasmith/DatasmithRhinoExporter/DatasmithRhino6.Target.cs
@@ -11,7 +11,7 @@ public abstract class DatasmithRhinoBaseTarget : TargetRules
 	private const string AutoGeneratedFileDisclaimer = "// Warning. This is an AUTO-GENERATED file used for automatic version management. Please add your changes to GeneratedAssemblyInfo.cs.template.";
 	private Guid UniqueID = Guid.NewGuid();
 
-	private string RhinoEnvVarName 
+	private string RhinoEnvVarName
 	{
 		get { return "RHINO" + GetVersion() + "_PATH"; }
 	}
@@ -20,12 +20,13 @@ public abstract class DatasmithRhinoBaseTarget : TargetRules
 		: base(Target)
 	{
 		Type = TargetType.Program;
+		IncludeOrderVersion = EngineIncludeOrderVersion.Latest;
 		SolutionDirectory = "Programs/Datasmith";
 		bBuildInSolutionByDefault = false;
 
 		string RhinoVersionString = GetVersion();
 		string ProjectName = "DatasmithRhino" + RhinoVersionString;
-		
+
 		ExeBinariesSubFolder = Path.Combine("Rhino", RhinoVersionString);
 		LaunchModuleName = ProjectName;
 
diff --git a/Engine/Source/Programs/Enterprise/Datasmith/DatasmithSketchUpRubyExporter/DatasmithSketchUpRuby2020.Target.cs b/Engine/Source/Programs/Enterprise/Datasmith/DatasmithSketchUpRubyExporter/DatasmithSketchUpRuby2020.Target.cs
--- a/Engine/Source/Programs/Enterprise/Datasmith/DatasmithSketchUpRubyExporter/DatasmithSketchUpRuby2020.Target.cs
+++ b/Engine/Source/Programs/Enterprise/Datasmith/DatasmithSketchUpRubyExporter/DatasmithSketchUpRuby2020.Target.cs
@@ -9,6 +9,7 @@ public abstract class DatasmithSketchUpRubyBaseTarget : TargetRules
 		: base(Target)
 	{
 		Type = TargetType.Program;
+		IncludeOrderVersion = EngineIncludeOrderVersion.Latest;
 		SolutionDirectory = "Programs/Datasmith";
 		bBuildInSolutionByDefault = false;
 
@@ -46,7 +47,7 @@ public abstract class DatasmithSketchUpRubyBaseTarget : TargetRules
 		ReadOnlyBuildVersion BuildVersion = Target.Version;
 		string VersionString = string.Format("{0}.{1}.{2}", BuildVersion.MajorVersion, BuildVersion.MinorVersion, BuildVersion.PatchVersion);
 
-		PostBuildSteps.Add(string.Format("echo {1}> \"{0}\"", 
+		PostBuildSteps.Add(string.Format("echo {1}> \"{0}\"",
 			string.Format(@"$(EngineDir)/Binaries/Win64/{0}/Plugin/UnrealDatasmithSketchUp/version", ExeBinariesSubFolder), VersionString));
 
 		// Copy plugin dll
diff --git a/Engine/Source/Programs/Enterprise/Datasmith/DatasmithSolidworksExporter/DatasmithSolidworks2021.Target.cs b/Engine/Source/Programs/Enterprise/Datasmith/DatasmithSolidworksExporter/DatasmithSolidworks2021.Target.cs
--- a/Engine/Source/Programs/Enterprise/Datasmith/DatasmithSolidworksExporter/DatasmithSolidworks2021.Target.cs
+++ b/Engine/Source/Programs/Enterprise/Datasmith/DatasmithSolidworksExporter/DatasmithSolidworks2021.Target.cs
@@ -47,6 +47,7 @@ public abstract class DatasmithSolidworksBaseTarget : TargetRules
 	public DatasmithSolidworksBaseTarget(TargetInfo Target) : base(Target)
 	{
 		Type = TargetType.Program;
+		IncludeOrderVersion = EngineIncludeOrderVersion.Latest;
 		SolutionDirectory = "Programs/Datasmith";
 		bBuildInSolutionByDefault = false;
 		ExeBinariesSubFolder = "Solidworks/" + GetVersionShort();
