[3ds Max] DirectLink: diagnostics info

- added statistics about exported/sysced scene - nodes, materials, textures

diff --git a/Engine/Source/Programs/Enterprise/Datasmith/DatasmithMaxExporter/Private/DatasmithMaxDirectLinkMaterial.cpp b/Engine/Source/Programs/Enterprise/Datasmith/DatasmithMaxExporter/Private/DatasmithMaxDirectLinkMaterial.cpp
--- a/Engine/Source/Programs/Enterprise/Datasmith/DatasmithMaxExporter/Private/DatasmithMaxDirectLinkMaterial.cpp
+++ b/Engine/Source/Programs/Enterprise/Datasmith/DatasmithMaxExporter/Private/DatasmithMaxDirectLinkMaterial.cpp
@@ -120,6 +120,7 @@ void FMaterialsCollectionTracker::UpdateMaterial(FMaterialTracker* MaterialTrack
 
 void FMaterialsCollectionTracker::ConvertMaterial(Mtl* Material, TSharedRef<IDatasmithScene> DatasmithScene, const TCHAR* AssetsPath, TSet<Texmap*>& TexmapsConverted)
 {
+	SCENE_UPDATE_STAT_INC(UpdateMaterials, Total);
 	// todo: reworking material export - composite materials(like VRayBlend) export submaterials too
 	//   check below is needed to make sure that no material is exported twice.
 	//   This is how it was handled before bu needs to be handled without going through the list of all materials.
@@ -129,6 +130,7 @@ void FMaterialsCollectionTracker::ConvertMaterial(Mtl* Material, TSharedRef<IDat
 	{
 		if (FString(DatasmithScene->GetMaterial(i)->GetName()) == MaterialName)
 		{
+			SCENE_UPDATE_STAT_INC(UpdateMaterials, SkippedAsAlreadyConverted);
 			return;
 		}
 	}
@@ -145,6 +147,7 @@ void FMaterialsCollectionTracker::ConvertMaterial(Mtl* Material, TSharedRef<IDat
 
 		if (DatasmithMaterial)
 		{
+			SCENE_UPDATE_STAT_INC(UpdateMaterials, Converted);
 			DatasmithScene->AddMaterial(DatasmithMaterial);
 			UsedMaterialToDatasmithMaterial.Add(Material, DatasmithMaterial);
 		}
diff --git a/Engine/Source/Programs/Enterprise/Datasmith/DatasmithMaxExporter/Private/DatasmithMaxExporterWithDirectLink.cpp b/Engine/Source/Programs/Enterprise/Datasmith/DatasmithMaxExporter/Private/DatasmithMaxExporterWithDirectLink.cpp
--- a/Engine/Source/Programs/Enterprise/Datasmith/DatasmithMaxExporter/Private/DatasmithMaxExporterWithDirectLink.cpp
+++ b/Engine/Source/Programs/Enterprise/Datasmith/DatasmithMaxExporter/Private/DatasmithMaxExporterWithDirectLink.cpp
@@ -237,6 +237,11 @@ public:
 			}
 		}
 
+		void SetResult(const FString& Text)
+		{
+			Result = Text;
+		}
+
 		FString Name;
 		int32 StageCount;
 		int32 StageIndex = 0;
@@ -244,6 +249,8 @@ public:
 		FDateTime TimeStart;
 		FDateTime TimeFinish;
 
+		FString Result;
+
 		TArray<FStage> Stages;
 	};
 
@@ -263,6 +270,10 @@ public:
 	void PrintStage(FStage& Stage, FString Indent=TEXT(""))
 	{
 		LogInfo(Indent + FString::Printf(TEXT("    %s - %s"), *Stage.Name, *(Stage.TimeFinish-Stage.TimeStart).ToString()));
+		if (!Stage.Result.IsEmpty())
+		{
+			LogInfo(Indent + TEXT("      #") + Stage.Result);
+		}
 		for(FStage& ChildStage: Stage.Stages)
 		{
 			PrintStage(ChildStage, Indent + TEXT("  "));
@@ -319,11 +330,20 @@ public:
 	~FProgressStageGuard()
 	{
 		Stage.Finished();
+		if (ComputeResultDeferred)
+		{
+			Stage.Result = ComputeResultDeferred();
+		}
 	}
+
+	TFunction<FString()> ComputeResultDeferred;
 };
 
 #define PROGRESS_STAGE(Name) FProgressStageGuard ProgressStage(MainStage, TEXT(Name));
 #define PROGRESS_STAGE_COUNTER(Count) FProgressCounter ProgressCounter(ProgressStage.Stage, Count);
+#define PROGRESS_STAGE_RESULT(Text) ProgressStage.Stage.SetResult(Text);
+// Simplily creation of Stage result (called in Guard dtor)
+#define PROGRESS_STAGE_RESULT_DEFERRED ProgressStage.ComputeResultDeferred = [&]()
 
 // Convert various node data to Datasmith tags
 class FTagsConverter
@@ -380,7 +400,7 @@ public:
 			Time = EndTime;
 		}
 
-		virtual int proc(ReferenceMaker* RefMaker) override
+		virtual int32 proc(ReferenceMaker* RefMaker) override
 		{
 			RefMaker->RenderEnd(Time);
 			return REF_ENUM_CONTINUE;
@@ -570,9 +590,12 @@ public:
 		{
 			DWORD XRefFlags = SceneRootNode->GetXRefFlags(XRefChild);
 
+			SCENE_UPDATE_STAT_INC(ParseScene, XRefFileEncountered);
+
 			// XRef is disabled - not shown in viewport/render. Not loaded.
 			if (XRefFlags & XREF_DISABLED)
 			{
+				SCENE_UPDATE_STAT_INC(ParseScene, XRefFileDisabled);
 				// todo: baseline - doesn't check this - exports even disabled and XREF_HIDDEN scenes
 				continue;
 			}
@@ -580,12 +603,14 @@ public:
 			FString Path = FDatasmithMaxSceneExporter::GetActualPath(SceneRootNode->GetXRefFile(XRefChild).GetFileName());
 			if (FPaths::FileExists(Path) == false)
 			{
+				SCENE_UPDATE_STAT_INC(ParseScene, XRefFileMissing);
 				FString Error = FString("XRefScene file \"") + FPaths::GetCleanFilename(*Path) + FString("\" cannot be found");
 				// todo: logging
 				// DatasmithMaxLogger::Get().AddMissingAssetError(*Error);
 			}
 			else
 			{
+				SCENE_UPDATE_STAT_INC(ParseScene, XRefFileToParse);
 				ParseScene(SceneRootNode->GetXRefTree(XRefChild), FXRefScene{SceneRootNode, XRefChild});
 			}
 		}
@@ -605,6 +630,8 @@ public:
 	{
 		LogDebugNode(TEXT("ParseNode"), Node);
 
+		SCENE_UPDATE_STAT_INC(ParseNode, NodesEncountered);
+
 		BOOL bIsNodeHidden = Node->IsNodeHidden(TRUE);
 
 		if (bIsNodeHidden == false)
@@ -715,7 +742,7 @@ public:
 	}
 
 	// Applies all recorded changes to Datasmith scene
-	bool Update(FUpdateProgress::FStage& ProgressStage, bool bRenderQuality)
+	bool Update(FUpdateProgress::FStage& MainStage, bool bRenderQuality)
 	{
 		// Disable Undo, editing, redraw, messages during export/sync so that nothing changes the scene
 		GetCOREInterface()->EnableUndo(false);
@@ -734,7 +761,24 @@ public:
 		NodesPreparer.Start(GetCOREInterface()->GetTime(), bRenderQuality);
 
 		bUpdateInProgress = true;
-		bool bResult = UpdateInternalSafe(ProgressStage);
+
+		bool bResult = false;
+
+		{
+			FProgressStageGuard ProgressStage(MainStage, TEXT("Update"), 10);
+			PROGRESS_STAGE_RESULT_DEFERRED
+			{
+				IDatasmithScene& Scene = *ExportedScene.GetDatasmithScene();
+				return FString::Printf(TEXT("Actors: %d; Meshes: %d, Materials: %d"), 
+					Scene.GetActorsCount(),
+					Scene.GetMeshesCount(),
+					Scene.GetMaterialsCount(),
+					Scene.GetTexturesCount()
+				);
+			};
+
+			UpdateInternalSafe(ProgressStage.Stage);
+		}
 		bUpdateInProgress = false;
 
 		NodesPreparer.Finish();
@@ -746,28 +790,34 @@ public:
 		return bResult;
 	}
 
-	bool UpdateInternalSafe(FUpdateProgress::FStage& ProgressStage)
+	bool UpdateInternalSafe(FUpdateProgress::FStage& MainStage)
 	{
+		bool bResult = false;
 		__try
 		{
-			return UpdateInternal(ProgressStage);
+			bResult = UpdateInternal(MainStage);
 		}
 		__except(EXCEPTION_EXECUTE_HANDLER)
 		{
 			LogWarning(TEXT("Update finished with exception"));
 		}
-		return false;
+		return bResult;
 	}
 
-	bool UpdateInternal(FUpdateProgress::FStage& ParentStage)
+	bool UpdateInternal(FUpdateProgress::FStage& MainStage)
 	{
-		FProgressStageGuard MainProgressStageGuard(ParentStage, TEXT("Update"), 10);
-		FUpdateProgress::FStage& MainStage = MainProgressStageGuard.Stage;
-
 		bool bChangeEncountered = false;
 
+		Stats.Reset();
+
 		if (!bSceneParsed) // Parse whole scene only once
 		{
+			PROGRESS_STAGE("Parse Scene")
+			PROGRESS_STAGE_RESULT_DEFERRED
+			{
+				return FormatStatsParseScene();
+			};
+
 			ParseScene();
 		}
 
@@ -780,9 +830,13 @@ public:
 		bChangeEncountered |= !InvalidatedNodeTrackers.IsEmpty();
 		bChangeEncountered |= !MaterialsCollectionTracker.GetInvalidatedMaterials().IsEmpty();
 
-
 		{
 			PROGRESS_STAGE("Remove deleted nodes")
+			PROGRESS_STAGE_RESULT_DEFERRED
+			{
+				return FormatStatsRemoveDeletedNodes();
+			};
+
 			TArray<FNodeTracker*> DeletedNodeTrackers;
 			for (FNodeTracker* NodeTracker : InvalidatedNodeTrackers)
 			{
@@ -795,12 +849,19 @@ public:
 			for (FNodeTracker* NodeTrackerPtr : DeletedNodeTrackers)
 			{
 				RemoveNodeTracker(*NodeTrackerPtr);
+				SCENE_UPDATE_STAT_INC(RemoveDeletedNodes, Nodes);
 			}
+
 		}
 
 		// todo: move to NameChanged and NodeAdded?
 		{
 			PROGRESS_STAGE("Update node names")
+			PROGRESS_STAGE_RESULT_DEFERRED
+			{
+				return FormatStatsUpdateNodeNames();
+			};
+
 			for (FNodeTracker* NodeTracker : InvalidatedNodeTrackers)
 			{
 				FString Name = NodeTracker->Node->GetName();
@@ -824,10 +885,22 @@ public:
 				UpdateCollisionStatus(NodeTracker, NodesWithChangedCollisionStatus);
 			}
 			InvalidatedNodeTrackers.Append(NodesWithChangedCollisionStatus);
+
+			SCENE_UPDATE_STAT_SET(RefreshCollisions, ChangedNodes, NodesWithChangedCollisionStatus.Num());
+
+			PROGRESS_STAGE_RESULT_DEFERRED
+			{
+				return FormatStatsRefreshCollisions();
+			};
 		}
 
 		{
 			PROGRESS_STAGE("Process invalidated nodes")
+			PROGRESS_STAGE_RESULT_DEFERRED
+			{
+				return FormatStatsProcessInvalidatedNodes();
+			};
+
 			PROGRESS_STAGE_COUNTER(InvalidatedNodeTrackers.Num());
 			for (FNodeTracker* NodeTracker : InvalidatedNodeTrackers)
 			{
@@ -838,6 +911,10 @@ public:
 
 		{
 			PROGRESS_STAGE("Process invalidated instances")
+			PROGRESS_STAGE_RESULT_DEFERRED
+			{
+				return FormatStatsProcessInvalidatedInstances();
+			};
 			PROGRESS_STAGE_COUNTER(InvalidatedInstances.Num());
 			for (FInstances* Instances : InvalidatedInstances)
 			{
@@ -845,15 +922,22 @@ public:
 				UpdateInstances(*Instances);
 				InvalidatedNodeTrackers.Append(Instances->NodeTrackers);
 			}
+
 			InvalidatedInstances.Reset();
 		}
 
 		{
 			PROGRESS_STAGE("Reparent Datasmith Actors");
+			PROGRESS_STAGE_RESULT_DEFERRED
+			{
+				return FormatStatsReparentDatasmithActors();
+			};
+			
 			for (FNodeTracker* NodeTracker : InvalidatedNodeTrackers)
 			{
 				AttachNodeToDatasmithScene(*NodeTracker);
 			}
+			;
 			InvalidatedNodeTrackers.Reset();
 		}
 
@@ -862,27 +946,37 @@ public:
 		TSet<Mtl*> ActualMaterialToUpdate; 
 		{
 			PROGRESS_STAGE("Process invalidated materials");
+			PROGRESS_STAGE_RESULT_DEFERRED
 			{
-				PROGRESS_STAGE_COUNTER(MaterialsCollectionTracker.GetInvalidatedMaterials().Num());
-				for (FMaterialTracker* MaterialTracker : MaterialsCollectionTracker.GetInvalidatedMaterials())
-				{
-					ProgressCounter.Next();
+				return FormatStatsProcessInvalidatedMaterials();
+			};
 
-					MaterialsCollectionTracker.UpdateMaterial(MaterialTracker);
+			PROGRESS_STAGE_COUNTER(MaterialsCollectionTracker.GetInvalidatedMaterials().Num());
+			for (FMaterialTracker* MaterialTracker : MaterialsCollectionTracker.GetInvalidatedMaterials())
+			{
+				ProgressCounter.Next();
+				SCENE_UPDATE_STAT_INC(ProcessInvalidatedMaterials, Invalidated);
 
-					for (Mtl* ActualMaterial : MaterialTracker->GetActualMaterials())
-					{
-						ActualMaterialToUpdate.Add(ActualMaterial);
-					}
-					MaterialTracker->bInvalidated = false;
+				MaterialsCollectionTracker.UpdateMaterial(MaterialTracker);
+
+				for (Mtl* ActualMaterial : MaterialTracker->GetActualMaterials())
+				{
+					SCENE_UPDATE_STAT_INC(ProcessInvalidatedMaterials, ActualToUpdate);
+					ActualMaterialToUpdate.Add(ActualMaterial);
 				}
-				MaterialsCollectionTracker.ResetInvalidatedMaterials();
+				MaterialTracker->bInvalidated = false;
 			}
+
+			MaterialsCollectionTracker.ResetInvalidatedMaterials();
 		}
 
 		TSet<Texmap*> ActualTexmapsToUpdate;
 		{
 			PROGRESS_STAGE("Update materials")
+			PROGRESS_STAGE_RESULT_DEFERRED
+			{
+				return FormatStatsUpdateMaterials();
+			};
 			PROGRESS_STAGE_COUNTER(ActualMaterialToUpdate.Num());
 			for (Mtl* ActualMaterial: ActualMaterialToUpdate)
 			{
@@ -894,11 +988,16 @@ public:
 
 		{
 			PROGRESS_STAGE("Update textures")
+			PROGRESS_STAGE_RESULT_DEFERRED
+			{
+				return FormatStatsUpdateTextures();
+			};
 			PROGRESS_STAGE_COUNTER(ActualTexmapsToUpdate.Num());
-			
+
 			for (Texmap* Texture : ActualTexmapsToUpdate)
 			{
 				ProgressCounter.Next();
+				SCENE_UPDATE_STAT_INC(UpdateTextures, Total);
 
 				TArray<TSharedPtr<IDatasmithTextureElement>> TextureElements;
 				FDatasmithMaxMatExport::GetXMLTexture(ExportedScene.GetDatasmithScene(), Texture, ExportedScene.GetSceneExporter().GetAssetsOutputPath(), &TextureElements);
@@ -928,10 +1027,57 @@ public:
 		//{
 		//	ExportedScene.GetDatasmithScene()->RemoveTexture(Texture);
 		//}
+		return bChangeEncountered;
+	}
 
-		LogDebug("Scene update: done");
+	FString FormatStatsParseScene()
+	{
+		return FString::Printf(TEXT("Nodes: parsed %d"), SCENE_UPDATE_STAT_GET(ParseNode, NodesEncountered));
+	}
 
-		return bChangeEncountered;
+	FString FormatStatsRemoveDeletedNodes()
+	{
+		return FString::Printf(TEXT("Nodes: deleted %d of %d invalidated of total %d"), SCENE_UPDATE_STAT_GET(RemoveDeletedNodes, Nodes), InvalidatedNodeTrackers.Num(), NodeTrackers.Num());
+	}
+
+	FString FormatStatsUpdateNodeNames()
+	{
+		return FString::Printf(TEXT("Nodes: updated %d of total %d"), InvalidatedNodeTrackers.Num(), NodeTrackers.Num());
+	}
+
+	FString FormatStatsRefreshCollisions()
+	{
+		return FString::Printf(TEXT("Nodes: added %d to invalidated %d"), SCENE_UPDATE_STAT_GET(RefreshCollisions, ChangedNodes), InvalidatedNodeTrackers.Num());
+	}
+
+	FString FormatStatsProcessInvalidatedNodes()
+	{
+		return FString::Printf(TEXT("Nodes: %d updated, %d skipped unselected, %d skipped hidden"), SCENE_UPDATE_STAT_GET(UpdateNode, NodesUpdated), SCENE_UPDATE_STAT_GET(UpdateNode, SkippedAsUnselected), SCENE_UPDATE_STAT_GET(UpdateNode, SkippedAsHiddenNode));
+	}
+
+	FString FormatStatsProcessInvalidatedInstances()
+	{
+		return FString::Printf(TEXT("Instances: %ld updated"), InvalidatedInstances.Num());
+	}
+
+	FString FormatStatsReparentDatasmithActors()
+	{
+		return FString::Printf(TEXT("Nodes: %d attached, to root %d, skipped %d"), SCENE_UPDATE_STAT_GET(ReparentActors, Attached), SCENE_UPDATE_STAT_GET(ReparentActors, AttachedToRoot), SCENE_UPDATE_STAT_GET(ReparentActors, SkippedWithoutDatasmithActor));
+	}
+
+	FString FormatStatsProcessInvalidatedMaterials()
+	{
+		return FString::Printf(TEXT("Materials: %d reparsed, found %d actual to update"), SCENE_UPDATE_STAT_GET(ProcessInvalidatedMaterials, Invalidated), SCENE_UPDATE_STAT_GET(ProcessInvalidatedMaterials, ActualToUpdate));
+	}
+
+	FString FormatStatsUpdateMaterials()
+	{
+		return FString::Printf(TEXT("Materials: %d updated, %d converted, %d skipped as already converted"), SCENE_UPDATE_STAT_GET(UpdateMaterials, Total), SCENE_UPDATE_STAT_GET(UpdateMaterials, Converted), SCENE_UPDATE_STAT_GET(UpdateMaterials, SkippedAsAlreadyConverted));
+	}
+
+	FString FormatStatsUpdateTextures()
+	{
+		return FString::Printf(TEXT("Texmaps: %d updated"), SCENE_UPDATE_STAT_GET(UpdateTextures, Total));
 	}
 
 	void ExportAnimations()
@@ -1274,7 +1420,7 @@ public:
 
 	void UpdateNode(FNodeTracker& NodeTracker)
 	{
-
+		SCENE_UPDATE_STAT_INC(UpdateNode, NodesUpdated);
 		// Forget anything that this node was before update: place in datasmith hierarchy, datasmith objects, instances connection. Updating may change anything
 		RemoveFromConverted(NodeTracker);
 		ConvertNodeObject(NodeTracker);
@@ -1296,16 +1442,19 @@ public:
 
 		if (CollisionNodes.Contains(&NodeTracker))
 		{
+			SCENE_UPDATE_STAT_INC(UpdateNode, SkippedAsCollisionNode);
 			return;
 		}
 
 		if (NodeTracker.Node->IsNodeHidden(TRUE) || !NodeTracker.Node->Renderable())
 		{
+			SCENE_UPDATE_STAT_INC(UpdateNode, SkippedAsHiddenNode);
 			return;
 		}
 
 		if (Options.bSelectedOnly && !NodeTracker.Node->Selected())
 		{
+			SCENE_UPDATE_STAT_INC(UpdateNode, SkippedAsUnselected);
 			return;
 		}
 
@@ -1333,6 +1482,7 @@ public:
 		case SHAPE_CLASS_ID:
 		case GEOMOBJECT_CLASS_ID:
 		{
+			SCENE_UPDATE_STAT_INC(UpdateNode, GeomObjEncontered);
 			Class_ID ClassID = ObjState.obj->ClassID();
 			if (ClassID.PartA() == TARGET_CLASS_ID) // Convert camera target as regular actor
 			{
@@ -1438,12 +1588,15 @@ public:
 		return nullptr;
 	}
 
-	void AttachNodeToDatasmithScene(FNodeTracker& NodeTracker)
+	bool AttachNodeToDatasmithScene(FNodeTracker& NodeTracker)
 	{
 		if (!NodeTracker.DatasmithActorElement)
 		{
-			return;
+			SCENE_UPDATE_STAT_INC(ReparentActors, SkippedWithoutDatasmithActor);
+
+			return false;
 		}
+		SCENE_UPDATE_STAT_INC(ReparentActors, Attached);
 
 		if (FNodeTracker* ParentNodeTracker = GetAncestorNodeTrackerWithDatasmithActor(NodeTracker))
 		{
@@ -1452,10 +1605,12 @@ public:
 		}
 		else
 		{
+			SCENE_UPDATE_STAT_INC(ReparentActors, AttachedToRoot);
 			// If there's no ancestor node with DatasmithActor assume node it at root
 			// (node's parent might be node that was skipped - e.g. it was hidden in Max or not selected when exporting only selected objects)
 			ExportedScene.GetDatasmithScene()->AddActor(NodeTracker.DatasmithActorElement);
 		}
+		return true;
 	}
 
 	void GetNodeObjectTransform(FNodeTracker& NodeTracker, FDatasmithConverter Converter, FTransform& ObjectTransform)
@@ -1691,6 +1846,7 @@ public:
 
 	bool ConvertHelper(FNodeTracker& NodeTracker, Object* Obj)
 	{
+		SCENE_UPDATE_STAT_INC(UpdateNode, HelpersEncontered);
 		Helpers.Add(&NodeTracker);
 
 		if(!NodeTracker.DatasmithActorElement)
@@ -1708,6 +1864,7 @@ public:
 
 	bool ConvertCamera(FNodeTracker& NodeTracker, Object* Obj)
 	{
+		SCENE_UPDATE_STAT_INC(UpdateNode, CamerasEncontered);
 		Cameras.Add(&NodeTracker);
 
 		if(!NodeTracker.DatasmithActorElement)
@@ -1735,8 +1892,11 @@ public:
 
 	bool ConvertLight(FNodeTracker& NodeTracker, Object* Obj)
 	{
+		SCENE_UPDATE_STAT_INC(UpdateNode, LightsEncontered);
+
 		if (EMaxLightClass::Unknown == FDatasmithMaxSceneParser::GetLightClass(NodeTracker.Node))
 		{
+			SCENE_UPDATE_STAT_INC(UpdateNode, LightsSkippedAsUnknown);
 			return false;
 		}
 
@@ -1852,9 +2012,11 @@ public:
 		bool bResult = false;
 		if (!Obj->IsRenderable()) // Shape's Enable In Render flag(note - different from Node's Renderable flag)
 		{
+			SCENE_UPDATE_STAT_INC(UpdateNode, GeomObjSkippedAsNonRenderable);
 			return bResult;
 		}
 
+		SCENE_UPDATE_STAT_INC(UpdateNode, GeomObjConverted);
 		// AnimHandle is unique and never reused for new objects
 		// todo: reset instances and nodes when one node of an instance changes??? Check how it should be done actually, dependencies - nodes, object, invalidation place(Update, Event) etc...
 		AnimHandle Handle = Animatable::GetHandleByAnim(Obj);
@@ -2043,6 +2205,11 @@ public:
 		InvalidateNode(NodeKey);
 	}
 
+	virtual FSceneUpdateStats& GetStats() override
+	{
+		return Stats;
+	}
+
 	///////////////////////////////////////////////
 
 	const FExportOptions& Options;
@@ -2052,7 +2219,6 @@ public:
 	bool bSceneParsed = false;
 	bool bUpdateInProgress = false;
 
-
 	TMap<FNodeKey, FNodeTrackerHandle> NodeTrackers; // All scene nodes
 	TMap<FString, TSet<FNodeTracker*>> NodeTrackersNames; // Nodes grouped by name
 	TSet<FNodeTracker*> InvalidatedNodeTrackers; // Nodes that need to be rebuilt
@@ -2084,6 +2250,8 @@ public:
 	TSet<FInstances*> InvalidatedInstances;
 
 	FTagsConverter TagsConverter; // Converts max node information to Datasmith tags
+
+	FSceneUpdateStats Stats;
 };
 
 class FExporter: public IExporter
@@ -2392,7 +2560,16 @@ bool Export(const TCHAR* Name, const TCHAR* OutputPath, bool bQuiet)
 
 	{
 		PROGRESS_STAGE("Save Datasmith Scene");
+
+		IDatasmithScene& Scene = *ExportedScene.GetDatasmithScene();
 		ExportedScene.GetSceneExporter().Export(ExportedScene.GetDatasmithScene(), false);
+
+		PROGRESS_STAGE_RESULT(FString::Printf(TEXT("Actors: %d; Meshes: %d, Materials: %d"), 
+			Scene.GetActorsCount(),
+			Scene.GetMeshesCount(),
+			Scene.GetMaterialsCount(),
+			Scene.GetTexturesCount()
+			));
 	}
 
 	ProgressManager.Finished();
diff --git a/Engine/Source/Programs/Enterprise/Datasmith/DatasmithMaxExporter/Public/DatasmithMaxDirectLink.h b/Engine/Source/Programs/Enterprise/Datasmith/DatasmithMaxExporter/Public/DatasmithMaxDirectLink.h
--- a/Engine/Source/Programs/Enterprise/Datasmith/DatasmithMaxExporter/Public/DatasmithMaxDirectLink.h
+++ b/Engine/Source/Programs/Enterprise/Datasmith/DatasmithMaxExporter/Public/DatasmithMaxDirectLink.h
@@ -210,6 +210,48 @@ private:
 
 };
 
+struct FSceneUpdateStats
+{
+	void Reset()
+	{
+		*this = FSceneUpdateStats();
+	}
+
+	// Explicitly initialize every stat
+	int32 ParseSceneXRefFileEncountered = 0;
+	int32 ParseSceneXRefFileDisabled = 0;
+	int32 ParseSceneXRefFileMissing = 0;
+	int32 ParseSceneXRefFileToParse = 0;
+	int32 ParseNodeNodesEncountered = 0;
+	int32 RemoveDeletedNodesNodes = 0;
+	int32 RefreshCollisionsChangedNodes = 0;
+	int32 UpdateNodeNodesUpdated = 0;
+	int32 UpdateNodeSkippedAsCollisionNode = 0;
+	int32 UpdateNodeSkippedAsHiddenNode = 0;
+	int32 UpdateNodeSkippedAsUnselected = 0;
+	int32 UpdateNodeGeomObjEncontered = 0;
+	int32 UpdateNodeHelpersEncontered = 0;
+	int32 UpdateNodeCamerasEncontered = 0;
+	int32 UpdateNodeLightsEncontered = 0;
+	int32 UpdateNodeLightsSkippedAsUnknown = 0;
+	int32 UpdateNodeGeomObjSkippedAsNonRenderable = 0;
+	int32 UpdateNodeGeomObjConverted = 0;
+	int32 ReparentActorsSkippedWithoutDatasmithActor = 0;
+	int32 ReparentActorsAttached = 0;
+	int32 ReparentActorsAttachedToRoot = 0;
+	int32 ProcessInvalidatedMaterialsInvalidated = 0;
+	int32 ProcessInvalidatedMaterialsActualToUpdate = 0;
+	int32 UpdateMaterialsTotal = 0;
+	int32 UpdateMaterialsSkippedAsAlreadyConverted =0;
+	int32 UpdateMaterialsConverted = 0;
+	int32 UpdateTexturesTotal = 0;
+};
+
+#define SCENE_UPDATE_STAT_INC(Category, Name) {Stats.##Category##Name++;}
+#define SCENE_UPDATE_STAT_SET(Category, Name, Value) Stats.##Category##Name = Value
+#define SCENE_UPDATE_STAT_GET(Category, Name) Stats.##Category##Name
+
+
 // Modifies Datasmith scene in responce to change notification calls
 // Subscription to various Max notification systems is done separately, see  FNotifications
 class ISceneTracker
@@ -236,6 +278,7 @@ public:
 	virtual void RemoveMaterial(const TSharedPtr<IDatasmithBaseMaterialElement>& DatasmithMaterial) = 0;
 	virtual void RemoveTexture(const TSharedPtr<IDatasmithTextureElement>&) = 0;
 	virtual void NodeXRefMerged(INode* Node) = 0;
+	virtual FSceneUpdateStats& GetStats() = 0;
 };
 
 //---- Geometry utility function
@@ -404,7 +447,7 @@ class FMaterialsCollectionTracker
 {
 public:
 
-	FMaterialsCollectionTracker(ISceneTracker& InSceneTracker) : SceneTracker(InSceneTracker) {}
+	FMaterialsCollectionTracker(ISceneTracker& InSceneTracker) : SceneTracker(InSceneTracker), Stats(InSceneTracker.GetStats()) {}
 
 	void Reset();
 
@@ -449,6 +492,7 @@ public:
 	TMap<Texmap*, TArray<TSharedPtr<IDatasmithTextureElement>>> UsedTextureToDatasmithElement; // Keep track of Datasmith Element created for texmap to simplify update/removal(no need to search DatasmithScene)
 	// note: each texmap can create multiple texture elements
 
+	FSceneUpdateStats& Stats;
 };
 
 //----
