Datasmith Revit: Meta Data increments not sent via Direct link

diff --git a/Engine/Source/Programs/Enterprise/Datasmith/DatasmithRevitExporter/Private/DatasmithRevitApplication.cs b/Engine/Source/Programs/Enterprise/Datasmith/DatasmithRevitExporter/Private/DatasmithRevitApplication.cs
--- a/Engine/Source/Programs/Enterprise/Datasmith/DatasmithRevitExporter/Private/DatasmithRevitApplication.cs
+++ b/Engine/Source/Programs/Enterprise/Datasmith/DatasmithRevitExporter/Private/DatasmithRevitApplication.cs
@@ -22,6 +22,7 @@ namespace DatasmithRevitExporter
 
 		private EventHandler<DocumentClosingEventArgs> DocumentClosingHandler;
 		private EventHandler<ViewActivatedEventArgs> ViewActivatedHandler;
+		private EventHandler<IdlingEventArgs> IdlingEventHandler;
 
 		PushButton AutoSyncPushButton;
 		PushButton SyncPushButton;
@@ -122,7 +123,10 @@ namespace DatasmithRevitExporter
 
 			ViewActivatedHandler = new EventHandler<ViewActivatedEventArgs>(OnViewActivated);
 			InApplication.ViewActivated += ViewActivatedHandler;
-			
+
+			IdlingEventHandler = new EventHandler<IdlingEventArgs>(OnIdling);
+			InApplication.Idling += IdlingEventHandler;
+
 			// Setup Direct Link
 
 			string RevitEngineDir = null;
@@ -152,6 +156,11 @@ namespace DatasmithRevitExporter
 			return Result.Succeeded;
 		}
 
+		void OnIdling(object Sender, IdlingEventArgs Args)
+		{
+			FDirectLink.OnApplicationIdle();
+		}
+
 		static void OnDocumentClosing(object sender, DocumentClosingEventArgs e)
 		{
 			FDirectLink.DestroyInstance(FDirectLink.FindInstance(e.Document), e.Document.Application);
diff --git a/Engine/Source/Programs/Enterprise/Datasmith/DatasmithRevitExporter/Private/DatasmithRevitDirectLink.cs b/Engine/Source/Programs/Enterprise/Datasmith/DatasmithRevitExporter/Private/DatasmithRevitDirectLink.cs
--- a/Engine/Source/Programs/Enterprise/Datasmith/DatasmithRevitExporter/Private/DatasmithRevitDirectLink.cs
+++ b/Engine/Source/Programs/Enterprise/Datasmith/DatasmithRevitExporter/Private/DatasmithRevitDirectLink.cs
@@ -7,8 +7,6 @@ using System.IO;
 using Autodesk.Revit.DB;
 using Autodesk.Revit.ApplicationServices;
 using Autodesk.Revit.DB.Events;
-using System.Threading.Tasks;
-using System.Threading;
 using System.Linq;
 using Autodesk.Revit.UI;
 
@@ -20,9 +18,10 @@ namespace DatasmithRevitExporter
 		{
 			public Document															SourceDocument;
 			public Dictionary<ElementId, FDocumentData.FBaseElementData>			CachedElements = new Dictionary<ElementId, FDocumentData.FBaseElementData>();
-			public Queue<KeyValuePair<ElementId, FDocumentData.FBaseElementData>>	ElementsWithoutMetadata = new Queue<KeyValuePair<ElementId, FDocumentData.FBaseElementData>>();
+			public Queue<KeyValuePair<ElementId, FDocumentData.FBaseElementData>>	ElementsWithoutMetadataQueue = new Queue<KeyValuePair<ElementId, FDocumentData.FBaseElementData>>();
+			public HashSet<ElementId>												ElementsWithoutMetadataSet = new HashSet<ElementId>();
 			public HashSet<ElementId>												ExportedElements = new HashSet<ElementId>();
-			public HashSet<ElementId>												ModifiedElements = new HashSet<ElementId>();
+			private HashSet<ElementId>												ModifiedElements = new HashSet<ElementId>();
 			public Dictionary<string, FDatasmithFacadeActor>						ExportedActorsMap = new Dictionary<string, FDatasmithFacadeActor>();
 
 			public Dictionary<ElementId, FCachedDocumentData>						LinkedDocumentsCache = new Dictionary<ElementId, FCachedDocumentData>();
@@ -32,6 +31,41 @@ namespace DatasmithRevitExporter
 				SourceDocument = InDocument;
 			}
 
+			public bool ElementIsModified(ElementId ElemId)
+			{
+				return ModifiedElements.Contains(ElemId);
+			}
+
+			public void SetElementModified(bool bModified, ElementId ElemId)
+			{
+				if (bModified)
+				{
+					if (!ModifiedElements.Contains(ElemId))
+					{
+						ModifiedElements.Add(ElemId);
+						if (CachedElements.ContainsKey(ElemId))
+						{
+							FDocumentData.FBaseElementData ElemData = CachedElements[ElemId];
+
+							if (!ElementsWithoutMetadataSet.Contains(ElemId))
+							{
+								ElementsWithoutMetadataQueue.Enqueue(new KeyValuePair<ElementId, FDocumentData.FBaseElementData>(ElemId, ElemData));
+								ElementsWithoutMetadataSet.Add(ElemId);
+							}
+						}
+					}
+				}
+				else
+				{
+					ModifiedElements.Remove(ElemId);
+				}
+			}
+
+			public int GetModifiedElementsCount()
+			{
+				return ModifiedElements.Count;
+			}
+
 			public void SetAllElementsModified()
 			{
 				foreach (var Link in LinkedDocumentsCache.Values)
@@ -40,9 +74,12 @@ namespace DatasmithRevitExporter
 				}
 
 				ModifiedElements.Clear();
+				ElementsWithoutMetadataSet.Clear();
+				ElementsWithoutMetadataQueue.Clear();
+
 				foreach (var ElemId in CachedElements.Keys)
 				{
-					ModifiedElements.Add(ElemId);
+					SetElementModified(true, ElemId);
 				}
 			}
 
@@ -129,11 +166,6 @@ namespace DatasmithRevitExporter
 
 		private EventHandler<DocumentChangedEventArgs>					DocumentChangedHandler;
 
-		// Metadata related
-		private EventWaitHandle											MetadataEvent = new ManualResetEvent(false);
-		private CancellationTokenSource									MetadataCancelToken = null;
-		private Task													MetadataTask = null;
-
 		private static FDirectLink										ActiveInstance = null;
 		private static List<FDirectLink> Instances = new List<FDirectLink>();
 
@@ -142,6 +174,10 @@ namespace DatasmithRevitExporter
 		private bool bHasChanges = false;
 		private bool bSyncInProgress = false;
 
+		private EventHandler SettingsChangedHandler;
+
+		private bool bSettingsDirty = false;
+
 		private static bool _bAutoSync = false;
 		public static bool bAutoSync
 		{
@@ -217,6 +253,7 @@ namespace DatasmithRevitExporter
 			{
 				ActiveInstance = null;
 			}
+			FSettingsManager.SettingsUpdated -= Instance.SettingsChangedHandler;
 			Instances.Remove(Instance);
 			Instance?.Destroy(InApp);
 		}
@@ -231,6 +268,28 @@ namespace DatasmithRevitExporter
 			Instances.Clear();
 		}
 
+		public static void OnApplicationIdle()
+		{
+			if (ActiveInstance == null || ActiveInstance.SyncCount == 0)
+			{
+				return;
+			}
+
+			// We want to keep this number low in order to ensure functioning Revit GUI
+			int MetadataCountPerIdleEvent = 2;
+
+			string EnvBatchSize = Environment.GetEnvironmentVariable("REVIT_DIRECTLINK_METADATA_BATCH_SIZE");
+			if (!string.IsNullOrEmpty(EnvBatchSize))
+			{
+				if (int.TryParse(EnvBatchSize, out MetadataCountPerIdleEvent))
+				{
+					MetadataCountPerIdleEvent = Math.Max(1, MetadataCountPerIdleEvent);
+				}
+			}
+
+			ActiveInstance.ExportMetadataBatch(MetadataCountPerIdleEvent);
+		}
+
 		public static void OnDocumentChanged(
 		  object InSender,
 		  DocumentChangedEventArgs InArgs) 
@@ -259,12 +318,12 @@ namespace DatasmithRevitExporter
 					{
 						foreach (ElementId DepElemId in DependentElements)
 						{
-							DirectLink.RootCache.ModifiedElements.Add(DepElemId);
+							DirectLink.RootCache.SetElementModified(true, DepElemId);
 						}
 					}
 				}
 
-				DirectLink.RootCache.ModifiedElements.Add(ElemId);
+				DirectLink.RootCache.SetElementModified(true, ElemId);
 			}
 
 			if (DirectLink.bHasChanges && bAutoSync)
@@ -291,6 +350,12 @@ namespace DatasmithRevitExporter
 
 			DocumentChangedHandler = new EventHandler<DocumentChangedEventArgs>(OnDocumentChanged);
 			InDocument.Application.DocumentChanged += DocumentChangedHandler;
+
+			SettingsChangedHandler = new EventHandler((object Sender, EventArgs Args) => 
+			{
+				bSettingsDirty = true;
+			});
+			FSettingsManager.SettingsUpdated += SettingsChangedHandler;
 		}
 
 		private void RunAutoSync()
@@ -326,21 +391,6 @@ namespace DatasmithRevitExporter
 			}
 		}
 
-		private void StopMetadataExport()
-		{
-			if (MetadataTask != null)
-			{
-				MetadataCancelToken.Cancel();
-				MetadataEvent.Set();
-				MetadataTask.Wait();
-				MetadataEvent.Reset();
-			}
-
-			MetadataCancelToken?.Dispose();
-			MetadataCancelToken = null;
-			MetadataTask = null;
-		}
-
 		private void MakeActive(bool bInActive)
 		{
 			if (!bInActive)
@@ -360,8 +410,6 @@ namespace DatasmithRevitExporter
 
 		private void Destroy(Application InApp)
 		{
-			StopMetadataExport();
-
 			InApp.DocumentChanged -= DocumentChangedHandler;
 			DocumentChangedHandler = null;
 
@@ -375,7 +423,7 @@ namespace DatasmithRevitExporter
 		{
 			if (InMaterial != null)
 			{
-				return RootCache.ModifiedElements.Contains(InMaterial.Id);
+				return RootCache.ElementIsModified(InMaterial.Id);
 			}
 			return false;
 		}
@@ -384,7 +432,7 @@ namespace DatasmithRevitExporter
 		{
 			if (InMaterial != null)
 			{
-				RootCache.ModifiedElements.Remove(InMaterial.Id);
+				RootCache.SetElementModified(false, InMaterial.Id);
 			}
 		}
 		public void MarkForExport(Element InElement)
@@ -405,7 +453,7 @@ namespace DatasmithRevitExporter
 		public void ClearModified(Element InElement)
 		{
 			// Clear from modified set since we might get another element with same id and we dont want to skip it.
-			CurrentCache.ModifiedElements.Remove(InElement.Id);
+			CurrentCache.SetElementModified(false, InElement.Id);
 		}
 
 		public void CacheElement(Document InDocument, Element InElement, FDocumentData.FBaseElementData InElementData)
@@ -413,7 +461,12 @@ namespace DatasmithRevitExporter
 			if (!CurrentCache.CachedElements.ContainsKey(InElement.Id))
 			{
 				CurrentCache.CachedElements[InElement.Id] = InElementData;
-				CurrentCache.ElementsWithoutMetadata.Enqueue(new KeyValuePair<ElementId, FDocumentData.FBaseElementData>(InElement.Id, InElementData));
+
+				if (!CurrentCache.ElementsWithoutMetadataSet.Contains(InElement.Id))
+				{
+					CurrentCache.ElementsWithoutMetadataQueue.Enqueue(new KeyValuePair<ElementId, FDocumentData.FBaseElementData>(InElement.Id, InElementData));
+					CurrentCache.ElementsWithoutMetadataSet.Add(InElement.Id);
+				}
 			}
 			CacheActorType(InElementData.ElementActor);
 		}
@@ -473,7 +526,7 @@ namespace DatasmithRevitExporter
 
 		public bool IsElementModified(Element InElement)
 		{
-			return CurrentCache.ModifiedElements.Contains(InElement.Id);
+			return CurrentCache.ElementIsModified(InElement.Id);
 		}
 
 		public void OnBeginLinkedDocument(Element InLinkElement)
@@ -499,9 +552,13 @@ namespace DatasmithRevitExporter
 
 		public void OnBeginExport()
 		{
-			bSyncInProgress = true;
+			if (SyncCount > 0 && bSettingsDirty)
+			{
+				RootCache?.SetAllElementsModified();
+			}
 
-			StopMetadataExport();
+			bSettingsDirty = false;
+			bSyncInProgress = true;
 
 			SetSceneCachePath();
 
@@ -529,7 +586,7 @@ namespace DatasmithRevitExporter
 
 			foreach(SectionBoxInfo CurrentSectionBoxInfo in CurrentSectionBoxes)
 			{
-				if (!RootCache.ModifiedElements.Contains(CurrentSectionBoxInfo.SectionBox.Id))
+				if (!RootCache.ElementIsModified(CurrentSectionBoxInfo.SectionBox.Id))
 				{
 					continue;
 				}
@@ -596,9 +653,9 @@ namespace DatasmithRevitExporter
 
 			foreach (var ElemId in IntersectedElements)
 			{
-				if (!InData.ModifiedElements.Contains(ElemId))
+				if (!InData.ElementIsModified(ElemId))
 				{
-					InData.ModifiedElements.Add(ElemId);
+					InData.SetElementModified(true, ElemId);
 				}
 			}
 
@@ -607,9 +664,9 @@ namespace DatasmithRevitExporter
 			{
 				MarkIntersectedElementsAsModified(LinkedDoc.Value, InSectionBox, InSectionBoxBounds);
 
-				if (LinkedDoc.Value.ModifiedElements.Count > 0)
+				if (LinkedDoc.Value.GetModifiedElementsCount() > 0)
 				{
-					InData.ModifiedElements.Add(LinkedDoc.Key);
+					InData.SetElementModified(true, LinkedDoc.Key);
 					ModifiedLinkedDocuments.Add(LinkedDoc.Value.SourceDocument);
 				}
 			}
@@ -723,39 +780,19 @@ namespace DatasmithRevitExporter
 
 			bHasChanges = false;
 			bSyncInProgress = false;
-
-			Debug.Assert(MetadataTask == null); // We cannot have metadata export running at this point (must be stopped in OnBeginExport)
-			MetadataCancelToken = new CancellationTokenSource();
-			MetadataTask = Task.Run(() => ExportMetadata());
 		}
 
-		void ExportMetadata()
+		void ExportMetadataBatch(int ExportBatchSize)
 		{
-			int DelayExport = 2000;		// milliseconds
-			int ExportBatchSize = 1000;	// After each batch is exported, the process will wait for DelayExport and resume (unless cancelled)
 			int CurrentBatchSize = 0;
 
-			Func<FCachedDocumentData, bool> AddElements = (FCachedDocumentData CacheData) => 
+			Action<FCachedDocumentData> AddElements = (FCachedDocumentData CacheData) => 
 			{
-				while (CacheData.ElementsWithoutMetadata.Count > 0)
+				while (CacheData.ElementsWithoutMetadataQueue.Count > 0 && CurrentBatchSize < ExportBatchSize)
 				{
-					if (CurrentBatchSize == ExportBatchSize)
-					{
-						// Add some delay before exporting next batch.
-						CurrentBatchSize = 0;
+					var Entry = CacheData.ElementsWithoutMetadataQueue.Dequeue();
 
-						// Send metadata to DirectLink.
-						DatasmithDirectLink.UpdateScene(DatasmithScene);
-
-						MetadataEvent.WaitOne(DelayExport);
-					}
-
-					if (MetadataCancelToken.IsCancellationRequested)
-					{
-						return false;
-					}
-
-					var Entry = CacheData.ElementsWithoutMetadata.Dequeue();
+					CacheData.ElementsWithoutMetadataSet.Remove(Entry.Key);
 
 					// Handle the case where element might be deleted in the main export path.
 					if (!CacheData.CachedElements.ContainsKey(Entry.Key))
@@ -784,37 +821,42 @@ namespace DatasmithRevitExporter
 					++CurrentBatchSize;
 
 #if DEBUG
-					Debug.WriteLine($"metadata batch element {CurrentBatchSize}, remain in Q {CacheData.ElementsWithoutMetadata.Count}");
+					Debug.WriteLine($"metadata batch element {CurrentBatchSize}, remain in Q {CacheData.ElementsWithoutMetadataQueue.Count}");
 #endif
 				}
-
-				return true;
 			};
 
 			List<FCachedDocumentData> CachesToExport = new List<FCachedDocumentData>();
 
-			Action<FCachedDocumentData> GetLinkedDocuments = null;
+			Func<FCachedDocumentData, int> GetDocumentCaches = null;
 
-			GetLinkedDocuments = (FCachedDocumentData InParent) =>
+			GetDocumentCaches = (FCachedDocumentData InParent) =>
 			{
+				int ElementsInQueue = InParent.ElementsWithoutMetadataQueue.Count;
+
 				CachesToExport.Add(InParent);
 				foreach (var Cache in InParent.LinkedDocumentsCache.Values) 
 				{
-					GetLinkedDocuments(Cache);
+					ElementsInQueue += GetDocumentCaches(Cache);
 				}
+
+				return ElementsInQueue;
 			};
 
-			GetLinkedDocuments(RootCache);
+			int TotalElementsWithoutMetadata = GetDocumentCaches(RootCache);
+
+			if (TotalElementsWithoutMetadata == 0)
+			{
+				return;
+			}
 
 			foreach (var Cache in CachesToExport)
 			{
-				bool Success = AddElements(Cache);
-				if (!Success)
+				AddElements(Cache);
+
+				if (CurrentBatchSize >= ExportBatchSize)
 				{
-#if DEBUG
-					Debug.WriteLine("metadata cancelled");
-#endif
-					return; // Metadata export was cancelled.
+					break;
 				}
 			}
 
diff --git a/Engine/Source/Programs/Enterprise/Datasmith/DatasmithRevitExporter/Private/DatasmithRevitSettingsDialog.cs b/Engine/Source/Programs/Enterprise/Datasmith/DatasmithRevitExporter/Private/DatasmithRevitSettingsDialog.cs
--- a/Engine/Source/Programs/Enterprise/Datasmith/DatasmithRevitExporter/Private/DatasmithRevitSettingsDialog.cs
+++ b/Engine/Source/Programs/Enterprise/Datasmith/DatasmithRevitExporter/Private/DatasmithRevitSettingsDialog.cs
@@ -93,6 +93,7 @@ namespace DatasmithRevitExporter
 		private Autodesk.Revit.DB.Document Document;
 		private SortedList<string, int> AddedBuiltinParamGroups = new SortedList<string, int>();
 		private List<string> AddedParamNames = new List<string>();
+		private bool bSettingsChanged = false;
 
 		public DatasmithRevitSettingsDialog(Autodesk.Revit.DB.Document InDocument)
 		{
@@ -173,6 +174,10 @@ namespace DatasmithRevitExporter
 				LevelOfTessellation.AutoSize = true;
 				LevelOfTessellation.TabIndex = 3;
 				LevelOfTessellation.TextAlign = HorizontalAlignment.Right;
+				LevelOfTessellation.ValueChanged += (S, E) =>
+				{
+					bSettingsChanged = true;
+				};
 
 				OptionToolTip.SetToolTip(LevelOfTessellationLabel, FormatTooltip(DatasmithRevitResources.Strings.SettingsDialog_LevelOfTesselationTooltip));
 
@@ -229,6 +234,7 @@ namespace DatasmithRevitExporter
 
 						if (bAddedNewGroups)
 						{
+							bSettingsChanged = true;
 							ReloadAddedGroupsList();
 						}
 					}
@@ -250,6 +256,8 @@ namespace DatasmithRevitExporter
 						GroupsList.SelectedIndices.CopyTo(IndicesToRemove, 0);
 						Array.Sort(IndicesToRemove, (A, B) => B.CompareTo(A));
 
+						bSettingsChanged = true;
+
 						foreach (int Index in IndicesToRemove)
 						{
 							GroupsList.Items.RemoveAt(Index);
@@ -356,10 +364,13 @@ namespace DatasmithRevitExporter
 
 		private void OnClosing(object InSender, FormClosingEventArgs InArgs)
 		{
-			Settings.LevelOfTesselation = Decimal.ToInt32(LevelOfTessellation.Value);
-			Settings.InsertionPoint = (FSettings.EInsertionPoint)InsertionPointCombo.SelectedIndex;
-			Settings.MetadataParamGroupsFilter = AddedBuiltinParamGroups.Values;
-			FSettingsManager.WriteSettings(Document, Settings);
+			if (bSettingsChanged)
+			{
+				Settings.LevelOfTesselation = Decimal.ToInt32(LevelOfTessellation.Value);
+				Settings.InsertionPoint = (FSettings.EInsertionPoint)InsertionPointCombo.SelectedIndex;
+				Settings.MetadataParamGroupsFilter = AddedBuiltinParamGroups.Values;
+				FSettingsManager.WriteSettings(Document, Settings);
+			}
 		}
 	}
 }
diff --git a/Engine/Source/Programs/Enterprise/Datasmith/DatasmithRevitExporter/Private/DatasmithRevitSettingsManager.cs b/Engine/Source/Programs/Enterprise/Datasmith/DatasmithRevitExporter/Private/DatasmithRevitSettingsManager.cs
--- a/Engine/Source/Programs/Enterprise/Datasmith/DatasmithRevitExporter/Private/DatasmithRevitSettingsManager.cs
+++ b/Engine/Source/Programs/Enterprise/Datasmith/DatasmithRevitExporter/Private/DatasmithRevitSettingsManager.cs
@@ -138,6 +138,8 @@ namespace DatasmithRevitExporter
 
 		public static FSettings CurrentSettings = null;
 
+		public static event EventHandler SettingsUpdated;
+
 		public static void Init(UIControlledApplication InApplication)
 		{
 			DocumentOpenedHandler = new EventHandler<DocumentOpenedEventArgs>(OnDocumentOpened);
@@ -226,6 +228,8 @@ namespace DatasmithRevitExporter
 					CurrentSettings = Settings;
 
 					T.Commit();
+
+					SettingsUpdated?.Invoke(CurrentSettings, null);
 				}
 				catch { }
 			}
