[Revit Exporter] Ability to create multiple emitters for one RVT file

diff --git a/Engine/Source/Programs/Enterprise/Datasmith/DatasmithRevitExporter/Private/DatasmithRevitApplication.cs b/Engine/Source/Programs/Enterprise/Datasmith/DatasmithRevitExporter/Private/DatasmithRevitApplication.cs
--- a/Engine/Source/Programs/Enterprise/Datasmith/DatasmithRevitExporter/Private/DatasmithRevitApplication.cs
+++ b/Engine/Source/Programs/Enterprise/Datasmith/DatasmithRevitExporter/Private/DatasmithRevitApplication.cs
@@ -1,10 +1,13 @@
 // Copyright Epic Games, Inc. All Rights Reserved.
 
 using System;
+using System.Collections.Generic;
 using System.Diagnostics;
 using System.IO;
+using System.Linq;
 using System.Reflection;
 using System.Windows.Media.Imaging;
+using Autodesk.Revit.ApplicationServices;
 using Autodesk.Revit.DB;
 using Autodesk.Revit.DB.Events;
 using Autodesk.Revit.UI;
@@ -16,6 +19,47 @@ namespace DatasmithRevitExporter
 	// Add-in external application Datasmith Revit Exporter.
 	public class DatasmithRevitApplication : IExternalApplication
 	{
+		// Updater notifying us when 3D views got updated/added/deleted.
+		public class View3DUpdater : IUpdater
+		{
+			static AddInId AppId;
+			static UpdaterId AppUpdaterId;
+
+			public View3DUpdater(AddInId InId)
+			{
+				AppId = InId;
+				AppUpdaterId = new UpdaterId(AppId, new Guid("BE5DA5A4-73C2-42BB-9346-E54F632E2B54"));
+			}
+
+			public void Execute(UpdaterData InData)
+			{
+				if (InData.GetAddedElementIds().Count > 0 || InData.GetDeletedElementIds().Count > 0 || InData.GetModifiedElementIds().Count > 0)
+				{
+					DatasmithRevitApplication.Instance.SetViewList();
+				}
+			}
+
+			public string GetAdditionalInformation()
+			{
+				return "";
+			}
+
+			public ChangePriority GetChangePriority()
+			{
+				return ChangePriority.FloorsRoofsStructuralWalls;
+			}
+
+			public UpdaterId GetUpdaterId()
+			{
+				return AppUpdaterId;
+			}
+
+			public string GetUpdaterName()
+			{
+				return "View3DUpdater";
+			}
+		}
+
 		private static DatasmithRevitExportMessages ExportMessagesDialog = null;
 
 		private static string ExportMessages;
@@ -23,9 +67,17 @@ namespace DatasmithRevitExporter
 		private EventHandler<DocumentClosingEventArgs> DocumentClosingHandler;
 		private EventHandler<ViewActivatedEventArgs> ViewActivatedHandler;
 		private EventHandler<IdlingEventArgs> IdlingEventHandler;
+		private EventHandler<DocumentOpenedEventArgs> DocumentOpenedHandler;
 
-		PushButton AutoSyncPushButton;
-		PushButton SyncPushButton;
+		private PushButton AutoSyncPushButton;
+		private PushButton SyncPushButton;
+
+		private Document ActiveDocument = null;
+
+		private View3DUpdater ViewsUpdater;
+		private ElementId SelectedViewId = ElementId.InvalidElementId;
+		private ComboBox ComboViews;
+		private List<View3D> All3DViews;
 
 		BitmapImage AutoSyncIconOn_Small;
 		BitmapImage AutoSyncIconOn_Large;
@@ -51,6 +103,96 @@ namespace DatasmithRevitExporter
 			SyncPushButton.Enabled = !bToggled;
 		}
 
+		private void OnComboBoxChanged(object s, Autodesk.Revit.UI.Events.ComboBoxCurrentChangedEventArgs e)
+		{
+			if (e.NewValue != null)
+			{
+				View3D SelectedView = All3DViews.Find(View => View.Name == e.NewValue.ItemText);
+
+				if (SelectedView != null)
+				{
+					SelectedViewId = SelectedView.Id;
+					FDirectLink.ActivateInstance(SelectedView);
+				}
+			}
+		}
+
+		private void ClearViews()
+		{
+			All3DViews.Clear();
+
+			foreach (ComboBoxMember Item in ComboViews.GetItems())
+			{
+				Item.Visible = false;
+			}
+
+			SelectedViewId = ElementId.InvalidElementId;
+		}
+
+		private void SetViewList()
+		{
+			if (ActiveDocument == null)
+			{
+				return;
+			}
+
+			// Cache prev selected view
+			All3DViews = new FilteredElementCollector(ActiveDocument).OfClass(typeof(View3D)).Cast<View3D>().ToList();
+			All3DViews.RemoveAll(view => (view.IsTemplate || !view.CanBePrinted));
+
+			// Combo box does not allow removal of items :). This is a workaround to always only add items and 
+			// reuse them.
+			int ViewIndex = 0;
+			for (; ViewIndex < All3DViews.Count; ++ViewIndex)
+			{
+				View3D View = All3DViews[ViewIndex];
+
+				var Items = ComboViews.GetItems();
+				if (Items.Count == ViewIndex)
+				{
+					ComboViews.AddItem(new ComboBoxMemberData($"View_{ViewIndex}", "ViewName"));
+				}
+
+				Items = ComboViews.GetItems();
+				ComboBoxMember Item = Items[ViewIndex];
+				Item.ItemText = View.Name;
+				Item.Visible = true;
+			}
+
+			// Hide the excessive items
+			for (; ViewIndex < ComboViews.GetItems().Count; ++ViewIndex)
+			{
+				ComboViews.GetItems()[ViewIndex].Visible = false;
+			}
+
+			// Set the active view
+			View3D ComboActiveView = null;
+
+			if (SelectedViewId != ElementId.InvalidElementId)
+			{
+				ComboActiveView = All3DViews.Find(View => View.Id == SelectedViewId);
+			}
+
+			if (ComboActiveView == null)
+			{
+				ComboActiveView = All3DViews.Count > 0 ? All3DViews.First() : null;
+			}
+
+			if (ComboActiveView != null)
+			{
+				foreach (ComboBoxMember Item in ComboViews.GetItems())
+				{
+					if (Item.ItemText == ComboActiveView.Name)
+					{
+						ComboViews.Current = Item;
+						SelectedViewId = ComboActiveView.Id;
+						FDirectLink.ActivateInstance(ComboActiveView);
+						break;
+					}
+				}
+			}
+		}
+
 		// Implement the interface to execute some tasks when Revit starts.
 		public Result OnStartup(
 			UIControlledApplication InApplication // handle to the application being started
@@ -77,8 +219,34 @@ namespace DatasmithRevitExporter
 
 			SyncPushButton = DirectLinkRibbonPanel.AddItem(SyncButtonData) as PushButton;
 			AutoSyncPushButton = DirectLinkRibbonPanel.AddItem(AutoSyncButtonData) as PushButton;
+
 			PushButton ManageConnectionsButton = DirectLinkRibbonPanel.AddItem(ManageConnectionsButtonData) as PushButton;
 			PushButton ExportPushButton = FileExportRibbonPanel.AddItem(ExportButtonData) as PushButton;
+
+			// Create the view sync combo
+			ComboBoxData CbData = new ComboBoxData("ComboBoxViews");
+			PushButtonData LabelData = new PushButtonData("ComboLabel", "Select 3D view to sync", Assembly.GetExecutingAssembly().Location, "DatasmithRevitExporter.DatasmithSyncRevitCommand");
+
+			IList<RibbonItem> StackedItems = DatasmithRibbonPanel.AddStackedItems(LabelData, CbData);
+			if (StackedItems.Count > 1)
+			{
+				// We emulate a label on the ribbon with disabled push button (the Revit way to do things)
+				PushButton LabelButton = StackedItems[0] as PushButton;
+				if (LabelButton != null)
+				{
+					LabelButton.ToolTip = "Select 3D view to sync";
+					LabelButton.Enabled = false;
+				}
+
+				ComboViews = StackedItems[1] as ComboBox;
+				if (ComboViews != null)
+				{
+					ComboViews.CurrentChanged += OnComboBoxChanged;
+					ComboViews.ItemText = "3D Views";
+					ComboViews.ToolTip = "Select 3D View to Sync";
+				}
+			}
+
 			PushButton ShowLogButton = DatasmithRibbonPanel.AddItem(LogButtonData) as PushButton;
 			PushButton SettingsButton = DatasmithRibbonPanel.AddItem(SettingsButtonData) as PushButton;
 
@@ -118,12 +286,12 @@ namespace DatasmithRevitExporter
 			ShowLogButton.LargeImage = new BitmapImage(new Uri(DatasmithIconBase + "32.png"));
 			ShowLogButton.ToolTip = DatasmithRevitResources.Strings.ButtonMessagesHint;
 
+			DocumentOpenedHandler = new EventHandler<DocumentOpenedEventArgs>(OnDocumentOpened);
+			InApplication.ControlledApplication.DocumentOpened += DocumentOpenedHandler;
+
 			DocumentClosingHandler = new EventHandler<DocumentClosingEventArgs>(OnDocumentClosing);
 			InApplication.ControlledApplication.DocumentClosing += DocumentClosingHandler;
 
-			ViewActivatedHandler = new EventHandler<ViewActivatedEventArgs>(OnViewActivated);
-			InApplication.ViewActivated += ViewActivatedHandler;
-
 			IdlingEventHandler = new EventHandler<IdlingEventArgs>(OnIdling);
 			InApplication.Idling += IdlingEventHandler;
 
@@ -153,6 +321,16 @@ namespace DatasmithRevitExporter
 
 			FSettingsManager.Init(InApplication);
 
+			// Register updater to react to view modification
+
+			ViewsUpdater = new View3DUpdater(InApplication.ControlledApplication.ActiveAddInId);
+			UpdaterRegistry.RegisterUpdater(ViewsUpdater);
+
+			ElementCategoryFilter Filter = new ElementCategoryFilter( BuiltInCategory.OST_Views);
+			UpdaterRegistry.AddTrigger(ViewsUpdater.GetUpdaterId(), Filter, Element.GetChangeTypeAny());
+			UpdaterRegistry.AddTrigger(ViewsUpdater.GetUpdaterId(), Filter, Element.GetChangeTypeElementAddition());
+			UpdaterRegistry.AddTrigger(ViewsUpdater.GetUpdaterId(), Filter, Element.GetChangeTypeElementDeletion());
+
 			return Result.Succeeded;
 		}
 
@@ -161,20 +339,16 @@ namespace DatasmithRevitExporter
 			FDirectLink.OnApplicationIdle();
 		}
 
-		static void OnDocumentClosing(object sender, DocumentClosingEventArgs e)
+		static void OnDocumentOpened(object sender, DocumentOpenedEventArgs e)
 		{
-			FDirectLink.DestroyInstance(FDirectLink.FindInstance(e.Document), e.Document.Application);
+			Instance.ActiveDocument = e.Document;
+			Instance.SetViewList();
 		}
 
-		static void OnViewActivated(object sender, ViewActivatedEventArgs e)
+		static void OnDocumentClosing(object sender, DocumentClosingEventArgs e)
 		{
-			View Previous = e.PreviousActiveView;
-			View Current = e.CurrentActiveView;
-
-			if (Previous == null || !Previous.Document.Equals(Current.Document))
-			{
-				FDirectLink.ActivateInstance(Current.Document);
-			}
+			Instance.ClearViews();
+			FDirectLink.DestroyInstancesForDocument(e.Document, e.Document.Application);
 		}
 
 		// Implement the interface to execute some tasks when Revit shuts down.
@@ -183,10 +357,13 @@ namespace DatasmithRevitExporter
 		)
 		{
 			InApplication.ControlledApplication.DocumentClosing -= DocumentClosingHandler;
-			InApplication.ViewActivated -= ViewActivatedHandler;
+			InApplication.ControlledApplication.DocumentOpened -= DocumentOpenedHandler;
 
 			DocumentClosingHandler = null;
-			ViewActivatedHandler = null;
+			DocumentOpenedHandler = null;
+
+			UpdaterRegistry.UnregisterUpdater(ViewsUpdater.GetUpdaterId());
+			ViewsUpdater = null;
 
 			if (ExportMessagesDialog != null && !ExportMessagesDialog.IsDisposed)
 			{
diff --git a/Engine/Source/Programs/Enterprise/Datasmith/DatasmithRevitExporter/Private/DatasmithRevitCommand.cs b/Engine/Source/Programs/Enterprise/Datasmith/DatasmithRevitExporter/Private/DatasmithRevitCommand.cs
--- a/Engine/Source/Programs/Enterprise/Datasmith/DatasmithRevitExporter/Private/DatasmithRevitCommand.cs
+++ b/Engine/Source/Programs/Enterprise/Datasmith/DatasmithRevitExporter/Private/DatasmithRevitCommand.cs
@@ -62,11 +62,11 @@ namespace DatasmithRevitExporter
 		{
 			UIDocument UIDoc = InCommandData.Application.ActiveUIDocument;
 			Document Doc = UIDoc.Document;
-			View3D ActiveView = Doc.ActiveView as View3D;
+			View3D ActiveView = FDirectLink.GetSyncView();
 
 			if (ActiveView == null)
 			{
-				string Message = "You must be in a 3D view to export.";
+				string Message = "You must select a 3D view to sync.";
 				MessageBox.Show(Message, DIALOG_CAPTION, MessageBoxButtons.OK, MessageBoxIcon.Warning);
 				return Result.Cancelled;
 			}
@@ -178,7 +178,7 @@ namespace DatasmithRevitExporter
 
 			if (ExportActiveViewOnly)
 			{
-				View3D ActiveView = Doc.ActiveView as View3D;
+				View3D ActiveView = FDirectLink.GetSyncView();
 
 				if (ActiveView == null)
 				{
diff --git a/Engine/Source/Programs/Enterprise/Datasmith/DatasmithRevitExporter/Private/DatasmithRevitDirectLink.cs b/Engine/Source/Programs/Enterprise/Datasmith/DatasmithRevitExporter/Private/DatasmithRevitDirectLink.cs
--- a/Engine/Source/Programs/Enterprise/Datasmith/DatasmithRevitExporter/Private/DatasmithRevitDirectLink.cs
+++ b/Engine/Source/Programs/Enterprise/Datasmith/DatasmithRevitExporter/Private/DatasmithRevitDirectLink.cs
@@ -144,6 +144,8 @@ namespace DatasmithRevitExporter
 		private FCachedDocumentData										RootCache = null;
 		private FCachedDocumentData										CurrentCache = null;
 		
+		private View3D													SyncView = null;
+
 		private HashSet<Document>										ModifiedLinkedDocuments = new HashSet<Document>();
 		private HashSet<ElementId>										ExportedLinkedDocuments = new HashSet<ElementId>();
 		private Stack<FCachedDocumentData>								CacheStack = new Stack<FCachedDocumentData>();
@@ -200,11 +202,16 @@ namespace DatasmithRevitExporter
 			return ActiveInstance;
 		}
 
-		public static void ActivateInstance(Document InDocument)
+		public static View3D GetSyncView()
+		{
+			return ActiveInstance?.SyncView ?? null;
+		}
+
+		public static void ActivateInstance(View3D InView)
 		{
 			if (UIApp == null)
 			{
-				UIApp = new UIApplication(InDocument.Application);
+				UIApp = new UIApplication(InView.Document.Application);
 			}
 
 			// Disable existing instance, if there's active one.
@@ -218,7 +225,12 @@ namespace DatasmithRevitExporter
 
 			foreach (FDirectLink DL in Instances)
 			{
-				if (DL.RootCache.SourceDocument.Equals(InDocument))
+				if (DL.SyncView == null || !DL.SyncView.IsValidObject)
+				{
+					continue;
+				}
+
+				if (DL.RootCache.SourceDocument.Equals(InView.Document) && DL.SyncView.Id == InView.Id)
 				{
 					InstanceToActivate = DL;
 					break;
@@ -227,7 +239,7 @@ namespace DatasmithRevitExporter
 
 			if (InstanceToActivate == null)
 			{
-				InstanceToActivate = new FDirectLink(InDocument);
+				InstanceToActivate = new FDirectLink(InView);
 				Instances.Add(InstanceToActivate);
 			}
 
@@ -235,11 +247,11 @@ namespace DatasmithRevitExporter
 			ActiveInstance = InstanceToActivate;
 		}
 
-		public static FDirectLink FindInstance(Document InDocument)
+		public static FDirectLink FindInstance(View3D InView)
 		{
 			foreach (var Inst in Instances)
 			{
-				if (Inst.RootCache.SourceDocument.Equals(InDocument))
+				if (Inst.SyncView.Id == InView.Id)
 				{
 					return Inst;
 				}
@@ -258,6 +270,27 @@ namespace DatasmithRevitExporter
 			Instance?.Destroy(InApp);
 		}
 
+		public static void DestroyInstancesForDocument(Document InDoc, Application InApp)
+		{
+			for (int InstanceIndex = Instances.Count - 1; InstanceIndex >= 0; --InstanceIndex)
+			{
+				FDirectLink Instance = Instances[InstanceIndex];
+
+				if (!Instance.RootCache.SourceDocument.Equals(InDoc))
+				{
+					continue;
+				}
+
+				if (ActiveInstance == Instance)
+				{
+					ActiveInstance = null;
+				}
+
+				Instances.RemoveAt(InstanceIndex);
+				Instance?.Destroy(InApp);
+			}
+		}
+
 		public static void DestroyAllInstances(Application InApp) 
 		{
 			foreach (FDirectLink DL in Instances)
@@ -332,24 +365,27 @@ namespace DatasmithRevitExporter
 			}
 		}
 
-		private FDirectLink(Document InDocument)
+		private FDirectLink(View3D InView)
 		{
-			RootCache = new FCachedDocumentData(InDocument);
+			RootCache = new FCachedDocumentData(InView.Document);
+			
 			CurrentCache = RootCache;
 
+			SyncView = InView;
+
 			DatasmithScene = new FDatasmithFacadeScene(
 				FDatasmithRevitExportContext.HOST_NAME,
 				FDatasmithRevitExportContext.VENDOR_NAME,
 				FDatasmithRevitExportContext.PRODUCT_NAME,
-				InDocument.Application.VersionNumber);
+				InView.Document.Application.VersionNumber);
 
-			SceneName = Path.GetFileNameWithoutExtension(RootCache.SourceDocument.PathName);
+			SceneName = $"{Path.GetFileNameWithoutExtension(RootCache.SourceDocument.PathName)}_{InView.Name}";
 			string OutputPath = Path.Combine(Path.GetTempPath(), SceneName);
 			DatasmithScene.SetName(SceneName);
 			DatasmithScene.SetLabel(SceneName);
 
 			DocumentChangedHandler = new EventHandler<DocumentChangedEventArgs>(OnDocumentChanged);
-			InDocument.Application.DocumentChanged += DocumentChangedHandler;
+			InView.Document.Application.DocumentChanged += DocumentChangedHandler;
 
 			SettingsChangedHandler = new EventHandler((object Sender, EventArgs Args) => 
 			{
@@ -863,7 +899,7 @@ namespace DatasmithRevitExporter
 			if (CurrentBatchSize > 0)
 			{
 				// Send remaining chunk of metadata.
-				DatasmithDirectLink.UpdateScene(DatasmithScene);
+				DatasmithDirectLink?.UpdateScene(DatasmithScene);
 			}
 
 #if DEBUG
