Fixed issue with wrong Ruby framework linked to SketchUp 2022 exporter
Solution: Removed explicit linkage with Ruby framework using -undefined linker's option since Ruby framework will be already loaded by SketchUp app.
Note: This was already done for the SketchUpAPI framework.

diff --git a/Engine/Source/Programs/Enterprise/Datasmith/DatasmithSketchUpRubyExporter/Build/build_mac.py b/Engine/Source/Programs/Enterprise/Datasmith/DatasmithSketchUpRubyExporter/Build/build_mac.py
--- a/Engine/Source/Programs/Enterprise/Datasmith/DatasmithSketchUpRubyExporter/Build/build_mac.py
+++ b/Engine/Source/Programs/Enterprise/Datasmith/DatasmithSketchUpRubyExporter/Build/build_mac.py
@@ -190,7 +190,6 @@ assert os.path.isfile(DatasmithSDKlib), DatasmithSDKlib
 link_cmd = [
     'xcrun',
     'clang++', 
-#    '-isysroot', '/Applications/Xcode.app/Contents/Developer/Platforms/MacOSX.platform/Developer/SDKs/MacOSX11.1.sdk',
     '-target', 'x86_64-apple-macos10.15',
     '-Xlinker', '-export_dynamic',
     '-Xlinker', '-no_deduplicate',
@@ -201,7 +200,7 @@ link_cmd = [
 
 link_cmd.append('-bundle') # link ruby extension as a MacOs 'bundle' not dynamic lib
 
-link_cmd += ['-framework', 'Ruby'] # link to RUby framework (found by framework_search_paths)
+#link_cmd += ['-framework', 'Ruby'] # link to RUby framework (found by framework_search_paths)
 
 # Handle SketchUp API
 link_cmd += [
