Fixed versioning issue of Datasmith scene element in ArchiCAD exporter

diff --git a/Engine/Source/Programs/Enterprise/Datasmith/DatasmithARCHICADExporter/Private/SyncDatabase.cpp b/Engine/Source/Programs/Enterprise/Datasmith/DatasmithARCHICADExporter/Private/SyncDatabase.cpp
--- a/Engine/Source/Programs/Enterprise/Datasmith/DatasmithARCHICADExporter/Private/SyncDatabase.cpp
+++ b/Engine/Source/Programs/Enterprise/Datasmith/DatasmithARCHICADExporter/Private/SyncDatabase.cpp
@@ -588,8 +588,6 @@ void FSyncDatabase::SetSceneInfo()
 	TheScene.SetVendor(TEXT("Graphisoft"));
 	TheScene.SetProductName(TEXT("ARCHICAD"));
 	TheScene.SetProductVersion(UTF8_TO_TCHAR(UE_AC_STRINGIZE(AC_VERSION)));
-	TheScene.SetExporterVersion(TEXT("4.27.2"));
-	TheScene.SetExporterSDKVersion(TEXT("4.27.2"));
 }
 
 void FSyncDatabase::ResetMeshClasses()
