[3ds Max] DirectLink: fixed AutoSync

diff --git a/Engine/Source/Programs/Enterprise/Datasmith/DatasmithMaxExporter/Private/DatasmithMaxExporterWithDirectLink.cpp b/Engine/Source/Programs/Enterprise/Datasmith/DatasmithMaxExporter/Private/DatasmithMaxExporterWithDirectLink.cpp
--- a/Engine/Source/Programs/Enterprise/Datasmith/DatasmithMaxExporter/Private/DatasmithMaxExporterWithDirectLink.cpp
+++ b/Engine/Source/Programs/Enterprise/Datasmith/DatasmithMaxExporter/Private/DatasmithMaxExporterWithDirectLink.cpp
@@ -765,6 +765,11 @@ public:
 			FProgressStageGuard ProgressStage(MainStage, TEXT("Update"), 10);
 			PROGRESS_STAGE_RESULT_DEFERRED
 			{
+				if (!ExportedScene.DatasmithSceneRef)
+				{
+					return FString(TEXT("<empty>"));
+				}
+
 				IDatasmithScene& Scene = *ExportedScene.GetDatasmithScene();
 				return FString::Printf(TEXT("Actors: %d; Meshes: %d, Materials: %d"), 
 					Scene.GetActorsCount(),
@@ -774,7 +779,7 @@ public:
 				);
 			};
 
-			UpdateInternalSafe(ProgressStage.Stage);
+			bResult = UpdateInternalSafe(ProgressStage.Stage);
 		}
 		bUpdateInProgress = false;
 
