//UE5/Main - Compile Datasmith Revit Exporter - DatasmithRevitApplication.cs warning: The field is never used

diff --git a/Engine/Source/Programs/Enterprise/Datasmith/DatasmithRevitExporter/Private/DatasmithRevitApplication.cs b/Engine/Source/Programs/Enterprise/Datasmith/DatasmithRevitExporter/Private/DatasmithRevitApplication.cs
--- a/Engine/Source/Programs/Enterprise/Datasmith/DatasmithRevitExporter/Private/DatasmithRevitApplication.cs
+++ b/Engine/Source/Programs/Enterprise/Datasmith/DatasmithRevitExporter/Private/DatasmithRevitApplication.cs
@@ -65,7 +65,6 @@ namespace DatasmithRevitExporter
 		private static string ExportMessages;
 
 		private EventHandler<DocumentClosingEventArgs> DocumentClosingHandler;
-		private EventHandler<ViewActivatedEventArgs> ViewActivatedHandler;
 		private EventHandler<IdlingEventArgs> IdlingEventHandler;
 		private EventHandler<DocumentOpenedEventArgs> DocumentOpenedHandler;
 
