[3ds Max] DirectLink: Sync Time Slider changes, Part 1

- Node transforms updated when Max time slider changes

- Optimize Collision Sync - not using Max api to find node by name but instead use plugin's  recorded name->node mapping

- additional refactor to simplify further validity checks(for geometry, materials etc)

diff --git a/Engine/Source/Programs/Enterprise/Datasmith/DatasmithMaxExporter/Private/DatasmithMaxConverters.cpp b/Engine/Source/Programs/Enterprise/Datasmith/DatasmithMaxExporter/Private/DatasmithMaxConverters.cpp
new file mode 100644
--- /dev/null
+++ b/Engine/Source/Programs/Enterprise/Datasmith/DatasmithMaxExporter/Private/DatasmithMaxConverters.cpp
@@ -0,0 +1,205 @@
+// Copyright Epic Games, Inc. All Rights Reserved.
+
+#ifdef NEW_DIRECTLINK_PLUGIN
+
+#include "DatasmithMaxConverters.h"
+
+#include "DatasmithMaxDirectLink.h"
+
+#include "DatasmithSceneFactory.h"
+
+#include "DatasmithMaxSceneExporter.h"
+#include "DatasmithMaxCameraExporter.h"
+#include "DatasmithMaxSceneParser.h"
+#include "DatasmithMaxLogger.h"
+#include "DatasmithMaxHelper.h"
+
+namespace  DatasmithMaxDirectLink
+{
+
+void FHelperNodeConverter::Parse(ISceneTracker& SceneTracker, FNodeTracker& NodeTracker)
+{
+}
+
+void FHelperNodeConverter::ConvertToDatasmith(ISceneTracker& SceneTracker, FNodeTracker& NodeTracker)
+{
+	// note: this is how baseline exporter derives names
+	FString UniqueName = FString::FromInt(NodeTracker.Node->GetHandle());
+	NodeTracker.CreateConverted().DatasmithActorElement = FDatasmithSceneFactory::CreateActor((const TCHAR*)*UniqueName);
+	SceneTracker.SetupActor(NodeTracker);
+
+	// todo: update validity from node
+	NodeTracker.NarrowValidityToInterval(FOREVER);
+}
+
+void FHelperNodeConverter::RemoveFromTracked(ISceneTracker& SceneTracker, FNodeTracker& NodeTracker)
+{
+}
+
+void FCameraNodeConverter::Parse(ISceneTracker& SceneTracker, FNodeTracker& NodeTracker)
+{
+}
+
+void FCameraNodeConverter::ConvertToDatasmith(ISceneTracker& SceneTracker, FNodeTracker& NodeTracker)
+{
+	// note: this is how baseline exporter derives names
+	FString UniqueName = FString::FromInt(NodeTracker.Node->GetHandle());
+	NodeTracker.CreateConverted().DatasmithActorElement = FDatasmithSceneFactory::CreateCameraActor((const TCHAR*)*UniqueName);
+
+	FDatasmithMaxCameraExporter::ExportCamera(*NodeTracker.Node, StaticCastSharedPtr<IDatasmithCameraActorElement>(NodeTracker.GetConverted().DatasmithActorElement).ToSharedRef());
+
+	SceneTracker.SetupActor(NodeTracker);
+
+	// Max camera view direction is Z-, Unreal's X+
+	// Max camera Up is Y+,  Unrela's Z+
+	FQuat Rotation = NodeTracker.GetConverted().DatasmithActorElement->GetRotation();
+	Rotation *= FQuat(0.0, 0.707107, 0.0, 0.707107);
+	Rotation *= FQuat(0.707107, 0.0, 0.0, 0.707107);
+	NodeTracker.GetConverted().DatasmithActorElement->SetRotation(Rotation);
+
+	// todo: update validity from node
+	NodeTracker.NarrowValidityToInterval(FOREVER);
+}
+
+void FCameraNodeConverter::RemoveFromTracked(ISceneTracker& SceneTracker, FNodeTracker& NodeTracker)
+{
+}
+
+void FLightNodeConverter::Parse(ISceneTracker& SceneTracker, FNodeTracker& NodeTracker)
+{
+}
+
+void FLightNodeConverter::ConvertToDatasmith(ISceneTracker& SceneTracker, FNodeTracker& NodeTracker)
+{
+	TSharedPtr<IDatasmithLightActorElement> LightElement;
+	{
+		// note: this is how baseline exporter derives names
+		FString UniqueName = FString::FromInt(NodeTracker.Node->GetHandle());
+
+		LightElement = FDatasmithMaxSceneExporter::CreateLightElementForNode(NodeTracker.Node, *UniqueName);
+
+		if (!LightElement)
+		{
+			if (FDatasmithMaxSceneParser::GetLightClass(NodeTracker.Node) == EMaxLightClass::SkyEquivalent)
+			{
+				SceneTracker.GetDatasmithSceneRef()->SetUsePhysicalSky(true);
+			}
+			else
+			{
+				LogWarning(TEXT("Unsupported light: ") + DatasmithMaxLogger::Get().GetLightDescription(NodeTracker.Node));
+			}
+			return;
+		}
+		else
+		{
+			if ( !FDatasmithMaxSceneExporter::ParseLight(NodeTracker.Node, LightElement.ToSharedRef(), SceneTracker.GetDatasmithSceneRef()) )
+			{
+				return;
+			}
+		}
+
+		NodeTracker.CreateConverted().DatasmithActorElement = LightElement;
+	}
+	SceneTracker.SetupActor(NodeTracker);
+
+	//Cylinder shaped lights don't have the same default orientations, so we recalculate their transform and add the shape information.
+	if (LightElement->IsA(EDatasmithElementType::AreaLight)
+		&& StaticCastSharedPtr<IDatasmithAreaLightElement>(LightElement)->GetLightShape() == EDatasmithLightShape::Cylinder)
+	{
+		FVector Translation, Scale;
+		FQuat Rotation;
+
+		const float UnitMultiplier = (float)GetSystemUnitScale(UNITS_CENTIMETERS);
+		const FMaxLightCoordinateConversionParams LightParams = FMaxLightCoordinateConversionParams(NodeTracker.Node, EDatasmithLightShape::Cylinder);
+
+		Interval ValidityInterval;
+		ValidityInterval.SetInfinite();
+
+		if (NodeTracker.Node->GetWSMDerivedObject() != nullptr)
+		{
+			FDatasmithMaxSceneExporter::MaxToUnrealCoordinates(NodeTracker.Node->GetObjTMAfterWSM(SceneTracker.CurrentSyncPoint.Time, &ValidityInterval), Translation, Rotation, Scale, UnitMultiplier, LightParams);
+		}
+		else
+		{
+			FDatasmithMaxSceneExporter::MaxToUnrealCoordinates(NodeTracker.Node->GetObjectTM(SceneTracker.CurrentSyncPoint.Time, &ValidityInterval), Translation, Rotation, Scale, UnitMultiplier, LightParams);
+		}
+
+		Rotation.Normalize();
+		LightElement->SetTranslation(Translation);
+		LightElement->SetScale(Scale);
+		LightElement->SetRotation(Rotation);
+
+		// todo: update validity from node
+		NodeTracker.NarrowValidityToInterval(ValidityInterval);
+	}
+
+}
+
+void FLightNodeConverter::RemoveFromTracked(ISceneTracker& SceneTracker, FNodeTracker& NodeTracker)
+{
+}
+
+void FMeshNodeConverter::Parse(ISceneTracker& SceneTracker, FNodeTracker& NodeTracker)
+{
+	ObjectState ObjState = NodeTracker.Node->EvalWorldState(SceneTracker.CurrentSyncPoint.Time);
+	Object* Obj = ObjState.obj;
+
+	if (Obj)
+	{
+		// AnimHandle is unique and never reused for new objects
+		InstanceHandle = Animatable::GetHandleByAnim(Obj); // Record instance handle into Converter to identify it later
+
+		SceneTracker.AddGeometryNodeInstance(NodeTracker, *this, Obj);
+	}
+}
+
+void FMeshNodeConverter::RemoveFromTracked(ISceneTracker& SceneTracker, FNodeTracker& NodeTracker)
+{
+	FMeshNodeConverter& Converter = static_cast<FMeshNodeConverter&>(NodeTracker.GetConverter());
+
+	// todo: registerNodeForMaterial is done in UpdateGeometryNode(separate from finding instances for Node)
+	SceneTracker.UnregisterNodeForMaterial(NodeTracker);
+	SceneTracker.RemoveGeometryNodeInstance(NodeTracker);
+}
+
+void FMeshNodeConverter::ConvertToDatasmith(ISceneTracker& SceneTracker, FNodeTracker& NodeTracker)
+{
+	SceneTracker.ConvertGeometryNodeToDatasmith(NodeTracker, *this);
+
+	// Mark node as updated as soon as it is - in order for next nodes to be able to use its DatasmithActor
+	// todo: update validity from node
+	NodeTracker.NarrowValidityToInterval(FOREVER);
+}
+
+void FHismNodeConverter::RemoveFromTracked(ISceneTracker& SceneTracker, FNodeTracker& NodeTracker)
+{
+	// remove static meshes used by the RailClone
+	SceneTracker.UnregisterNodeForMaterial(NodeTracker);
+	
+	for (FMeshConverted& Mesh : Meshes)
+	{
+		SceneTracker.ReleaseMeshElement(Mesh);
+	}
+}
+
+void FRailCloneNodeConverter::Parse(ISceneTracker& SceneTracker, FNodeTracker& NodeTracker)
+{
+}
+
+void FRailCloneNodeConverter::ConvertToDatasmith(ISceneTracker& SceneTracker, FNodeTracker& NodeTracker)
+{
+	GeomUtils::ConvertRailClone(SceneTracker, NodeTracker);
+}
+
+void FForestNodeConverter::Parse(ISceneTracker& SceneTracker, FNodeTracker& NodeTracker)
+{
+}
+
+void FForestNodeConverter::ConvertToDatasmith(ISceneTracker& SceneTracker, FNodeTracker& NodeTracker)
+{
+	GeomUtils::ConvertForest(SceneTracker, NodeTracker);
+}
+	
+}
+
+#endif // NEW_DIRECTLINK_PLUGIN
diff --git a/Engine/Source/Programs/Enterprise/Datasmith/DatasmithMaxExporter/Private/DatasmithMaxDirectLinkItoo.cpp b/Engine/Source/Programs/Enterprise/Datasmith/DatasmithMaxExporter/Private/DatasmithMaxDirectLinkItoo.cpp
--- a/Engine/Source/Programs/Enterprise/Datasmith/DatasmithMaxExporter/Private/DatasmithMaxDirectLinkItoo.cpp
+++ b/Engine/Source/Programs/Enterprise/Datasmith/DatasmithMaxExporter/Private/DatasmithMaxDirectLinkItoo.cpp
@@ -26,7 +26,10 @@ MAX_INCLUDES_END
 namespace DatasmithMaxDirectLink
 {
 
-bool ConvertRailClone(ISceneTracker& SceneTracker, FNodeTracker& NodeTracker, Object* Obj)
+namespace GeomUtils
+{
+
+bool ConvertRailClone(ISceneTracker& SceneTracker, FNodeTracker& NodeTracker)
 {
 	INode* RailCloneNode = NodeTracker.Node;
 
@@ -103,19 +106,22 @@ bool ConvertRailClone(ISceneTracker& SceneTracker, FNodeTracker& NodeTracker, Ob
 			}
 		}
 
-		if(!NodeTracker.DatasmithActorElement)
-		{
-			// note: this is how baseline exporter derives names
-			FString UniqueName = FString::FromInt(NodeTracker.Node->GetHandle());
-			NodeTracker.DatasmithActorElement = FDatasmithSceneFactory::CreateActor((const TCHAR*)*UniqueName);
-		}
+		// note: this is how baseline exporter derives names
+		FString UniqueName = FString::FromInt(NodeTracker.Node->GetHandle());
+		NodeTracker.CreateConverted().DatasmithActorElement = FDatasmithSceneFactory::CreateActor((const TCHAR*)*UniqueName);
 
 		SceneTracker.SetupActor(NodeTracker);
 
 		int32 MeshIndex = 0;
 		for (FHismInstances& Instances: InstancesForMesh)
 		{
-			SceneTracker.SetupDatasmithHISMForNode(NodeTracker, NodeTracker.Node, FRenderMeshForConversion(NodeTracker.Node, Instances.MaxMesh.Get(), false), NodeTracker.Node->GetMtl(), MeshIndex, Instances.Transforms);
+			FMeshConverterSource MeshSource = {
+				NodeTracker.Node, TEXT(""),
+				FRenderMeshForConversion(NodeTracker.Node, Instances.MaxMesh.Get(), false), false,
+				FRenderMeshForConversion()
+			};
+
+			SceneTracker.SetupDatasmithHISMForNode(NodeTracker, MeshSource, NodeTracker.Node->GetMtl(), MeshIndex, Instances.Transforms);
 			MeshIndex++;
 		}
 	}
@@ -128,7 +134,7 @@ bool ConvertRailClone(ISceneTracker& SceneTracker, FNodeTracker& NodeTracker, Ob
 
 }
 
-bool ConvertForest(ISceneTracker& Scene, FNodeTracker& NodeTracker, Object* Obj)
+bool ConvertForest(ISceneTracker& Scene, FNodeTracker& NodeTracker)
 {
 	INode* ForestNode = NodeTracker.Node;
 
@@ -177,12 +183,9 @@ bool ConvertForest(ISceneTracker& Scene, FNodeTracker& NodeTracker, Object* Obj)
 			}
 		}
 
-		if(!NodeTracker.DatasmithActorElement)
-		{
-			// note: this is how baseline exporter derives names
-			FString UniqueName = FString::FromInt(NodeTracker.Node->GetHandle());
-			NodeTracker.DatasmithActorElement = FDatasmithSceneFactory::CreateActor((const TCHAR*)*UniqueName);
-		}
+		// note: this is how baseline exporter derives names
+		FString UniqueName = FString::FromInt(NodeTracker.Node->GetHandle());
+		NodeTracker.CreateConverted().DatasmithActorElement = FDatasmithSceneFactory::CreateActor((const TCHAR*)*UniqueName);
 
 		Scene.SetupActor(NodeTracker);
 
@@ -192,14 +195,18 @@ bool ConvertForest(ISceneTracker& Scene, FNodeTracker& NodeTracker, Object* Obj)
 			INode* GeometryNode = Instances.GeometryNode;
 
 			Object* GeomObj = GeometryNode->EvalWorldState(GetCOREInterface()->GetTime()).obj;
-			
-			FRenderMeshForConversion RenderMesh = GetMeshForGeomObject(GeometryNode, GeomObj);
 
-			if (RenderMesh.IsValid())
+			FMeshConverterSource MeshSource = {
+				GeometryNode, TEXT(""),
+				GetMeshForGeomObject(CurrentTime, GeometryNode, GeomObj), false,
+				FRenderMeshForConversion()
+			};
+
+			if (MeshSource.RenderMesh.IsValid())
 			{
 				for(const TPair<Mtl*, TArray<Matrix3>>& MaterialAndTransforms: Instances.InstancesTransformsForMaterial)
 				{
-					Scene.SetupDatasmithHISMForNode(NodeTracker, GeometryNode, RenderMesh, MaterialAndTransforms.Key, MeshIndex, MaterialAndTransforms.Value);
+					Scene.SetupDatasmithHISMForNode(NodeTracker, MeshSource, MaterialAndTransforms.Key, MeshIndex, MaterialAndTransforms.Value);
 				}
 
 
@@ -216,6 +223,8 @@ bool ConvertForest(ISceneTracker& Scene, FNodeTracker& NodeTracker, Object* Obj)
 }
 
 }
+}
+
 
 #include "Windows/HideWindowsPlatformTypes.h"
 
diff --git a/Engine/Source/Programs/Enterprise/Datasmith/DatasmithMaxExporter/Private/DatasmithMaxDirectLinkLogging.cpp b/Engine/Source/Programs/Enterprise/Datasmith/DatasmithMaxExporter/Private/DatasmithMaxDirectLinkLogging.cpp
--- a/Engine/Source/Programs/Enterprise/Datasmith/DatasmithMaxExporter/Private/DatasmithMaxDirectLinkLogging.cpp
+++ b/Engine/Source/Programs/Enterprise/Datasmith/DatasmithMaxExporter/Private/DatasmithMaxDirectLinkLogging.cpp
@@ -75,16 +75,21 @@ void LogWarning(const FString& Msg)
 {
 	LogWarning(*Msg);
 }
+
+FString GetNodeDesc(INode* Node)
+{
+	return Node ? FString::Printf(TEXT("%s(%llu)"), Node->GetName(), NodeEventNamespace::GetKeyByNode(Node)) : TEXT("<null>");
+}
+
 void LogDebugNode(const FString& Name, INode* Node)
 {
 #ifdef LOG_DEBUG_HEAVY_ENABLE
 
-	LogDebug(FString::Printf(TEXT("%s: %u %s(%d) - %s")
+	LogDebug(FString::Printf(TEXT("%s: %s - %s, parent: %s")
 		, *Name
-		, NodeEventNamespace::GetKeyByNode(Node)
-		, Node ? Node->GetName() : L"<null>"
-		, Node ? Node->GetHandle() : 0
+		, *GetNodeDesc(Node)
 		, (Node && Node->IsNodeHidden(TRUE))? TEXT("HIDDEN") : TEXT("")
+		, *GetNodeDesc(Node ? Node->GetParentNode() : nullptr)
 		));
 	if (Node)
 	{
@@ -108,14 +113,13 @@ void LogNodeEvent(const MCHAR* Name, INodeEventCallback::NodeKeyTab& nodes)
 		FNodeKey NodeKey = nodes[NodeIndex];
 
 		Animatable* anim = Animatable::GetAnimByHandle(NodeKey);
-		INode* Node = NodeEventNamespace::GetNodeByKey(NodeKey);
-		if (Node) // Node sometimes is null. Not sure why
+		if (INode* Node = NodeEventNamespace::GetNodeByKey(NodeKey)) // Node sometimes is null. Not sure why
 		{
-			LogDebug(FString::Printf(TEXT("   %u %s(%d)"), NodeKey, Node->GetName(), Node->GetHandle()));
+			LogDebug(FString::Printf(TEXT("   %s, parent: %s"), *GetNodeDesc(Node), *GetNodeDesc(Node->GetParentNode())));
 		}
 		else
 		{
-			LogDebug(FString::Printf(TEXT("    %u <null>"), NodeKey));
+			LogDebug(FString::Printf(TEXT("   <null>(%llu)"), NodeKey));
 		}
 	}
 #endif
diff --git a/Engine/Source/Programs/Enterprise/Datasmith/DatasmithMaxExporter/Private/DatasmithMaxDirectLinkMaxScript.cpp b/Engine/Source/Programs/Enterprise/Datasmith/DatasmithMaxExporter/Private/DatasmithMaxDirectLinkMaxScript.cpp
--- a/Engine/Source/Programs/Enterprise/Datasmith/DatasmithMaxExporter/Private/DatasmithMaxDirectLinkMaxScript.cpp
+++ b/Engine/Source/Programs/Enterprise/Datasmith/DatasmithMaxExporter/Private/DatasmithMaxDirectLinkMaxScript.cpp
@@ -54,12 +54,8 @@ Value* OnLoad_cf(Value **arg_list, int count)
 
 	InitScripts();
 
-	if (CreateExporter(bEnableUI, EnginePathUnreal))
-	{
-		return &false_value;
-	}
 
-	return bool_result(true);
+	return bool_result(CreateExporter(bEnableUI, EnginePathUnreal));
 }
 
 Value* OnUnload_cf(Value**, int);
diff --git a/Engine/Source/Programs/Enterprise/Datasmith/DatasmithMaxExporter/Private/DatasmithMaxDirectLinkMesh.cpp b/Engine/Source/Programs/Enterprise/Datasmith/DatasmithMaxExporter/Private/DatasmithMaxDirectLinkMesh.cpp
--- a/Engine/Source/Programs/Enterprise/Datasmith/DatasmithMaxExporter/Private/DatasmithMaxDirectLinkMesh.cpp
+++ b/Engine/Source/Programs/Enterprise/Datasmith/DatasmithMaxExporter/Private/DatasmithMaxDirectLinkMesh.cpp
@@ -10,6 +10,7 @@
 #include "DatasmithMaxMeshExporter.h"
 #include "DatasmithMaxSceneExporter.h"
 #include "DatasmithMaxAttributes.h"
+#include "DatasmithMaxSceneParser.h"
 #include "DatasmithMaxLogger.h"
 
 #include "Logging/LogMacros.h"
@@ -24,8 +25,11 @@ MAX_INCLUDES_START
 	#include "MeshNormalSpec.h"
 MAX_INCLUDES_END
 
+
 namespace DatasmithMaxDirectLink
 {
+namespace GeomUtils
+{
 
 class FNullView : public View
 {
@@ -41,26 +45,26 @@ public:
 	}
 };
 
-FRenderMeshForConversion GetMeshForGeomObject(INode* Node, Object* Obj)
+FRenderMeshForConversion GetMeshForGeomObject(TimeValue CurrentTime, INode* Node, Object* Obj)
 {
 	// todo: baseline exporter uses GetBaseObject which takes result of EvalWorldState
 	// and searched down DerivedObject pipeline(by taking GetObjRef) 
 	// This is STRANGE as EvalWorldState shouldn't return DerivedObject in the first place(it should return result of pipeline evaluation)
 
 	BOOL bNeedsDelete;
-	Mesh* RenderMesh = GetMeshFromRenderMesh(Node, bNeedsDelete, GetCOREInterface()->GetTime());
+	Mesh* RenderMesh = GetMeshFromRenderMesh(CurrentTime, Node, bNeedsDelete);
 
 	return FRenderMeshForConversion(Node, RenderMesh, bNeedsDelete);
 }
 
-FRenderMeshForConversion GetMeshForNode(INode* Node, FTransform Pivot)
+FRenderMeshForConversion GetMeshForNode(TimeValue CurrentTime, INode* Node, FTransform Pivot)
 {
 	if (!Node)
 	{
 		return FRenderMeshForConversion();
 	}
 	BOOL bNeedsDelete;
-	Mesh* RenderMesh = GetMeshFromRenderMesh(Node, bNeedsDelete, GetCOREInterface()->GetTime());
+	Mesh* RenderMesh = GetMeshFromRenderMesh(CurrentTime, Node, bNeedsDelete);
 	return FRenderMeshForConversion(Node, RenderMesh, bNeedsDelete, Pivot);
 }
 
@@ -77,13 +81,41 @@ FTransform FTransformFromMatrix3(const Matrix3& MaxTransform, float UnitMultipli
 	return FTransform( Rotation, Translation, Scale );
 }
 
-FRenderMeshForConversion GetMeshForCollision(INode* Node)
+INode* GetCollisionNode(ISceneTracker& SceneTracker, INode* Node, const FDatasmithMaxStaticMeshAttributes* DatasmithAttributes, bool& bOutFromDatasmithAttribute)
+{
+	if (DatasmithAttributes)
+	{
+		INode* ModifierSpecifiedCustomCollisionNode = DatasmithAttributes->GetCustomCollisonNode();
+		if (ModifierSpecifiedCustomCollisionNode)
+		{
+			bOutFromDatasmithAttribute = true;
+			return ModifierSpecifiedCustomCollisionNode;
+		}
+	}
+
+	
+
+	FString OriginalName = Node->GetName();
+	for ( const FString& CollisionNodePrefix : FDatasmithMaxSceneParser::CollisionNodesPrefixes )
+	{
+		if (FNodeTracker* CollisionNode = SceneTracker.GetNodeTrackerByNodeName(*( CollisionNodePrefix + TEXT("_") + OriginalName )))
+		{
+			bOutFromDatasmithAttribute = false;
+			return CollisionNode->Node;
+		}
+	}
+
+	return nullptr;
+}
+
+
+FRenderMeshForConversion GetMeshForCollision(TimeValue CurrentTime, ISceneTracker& SceneTracker, INode* Node)
 {
 	// source: FDatasmithMaxMeshExporter::ExportMesh
 	FDatasmithConverter Converter;
 	bool bIsCollisionFromDatasmithAttributes;
 	TOptional<FDatasmithMaxStaticMeshAttributes> DatasmithAttributes = FDatasmithMaxStaticMeshAttributes::ExtractStaticMeshAttributes(Node);
-	INode* CollisionNode = FDatasmithMaxMeshExporter::GetCollisionNode(Node, DatasmithAttributes ? &DatasmithAttributes.GetValue() : nullptr, bIsCollisionFromDatasmithAttributes);
+	INode* CollisionNode = GetCollisionNode(SceneTracker, Node, DatasmithAttributes ? &DatasmithAttributes.GetValue() : nullptr, bIsCollisionFromDatasmithAttributes);
 	FTransform CollisionPivot;
 	if (CollisionNode)
 	{
@@ -103,11 +135,9 @@ FRenderMeshForConversion GetMeshForCollision(INode* Node)
 		}
 		else
 		{
-
-			TimeValue Now = GetCOREInterface()->GetTime();
 			FTransform FTransformFromMatrix3(const Matrix3& MaxTransform, float UnitMultiplier); // todo: move to header and rename(F is for class!)
-			FTransform NodeWTM = FTransformFromMatrix3(Node->GetNodeTM(Now), Converter.UnitToCentimeter);
-			FTransform ColliderNodeWTM = FTransformFromMatrix3(CollisionNode->GetNodeTM(Now), Converter.UnitToCentimeter);
+			FTransform NodeWTM = FTransformFromMatrix3(Node->GetNodeTM(CurrentTime), Converter.UnitToCentimeter);
+			FTransform ColliderNodeWTM = FTransformFromMatrix3(CollisionNode->GetNodeTM(CurrentTime), Converter.UnitToCentimeter);
 
 			// if object-offset has been baked into the mesh data, we want collision mesh data in the mesh's node space
 			//   MeshVert_Node = RealPivot * max_vert_data
@@ -128,7 +158,7 @@ FRenderMeshForConversion GetMeshForCollision(INode* Node)
 			CollisionPivot = BakedTransform;
 		}
 	}
-	return GetMeshForNode(CollisionNode, CollisionPivot);
+	return GetMeshForNode(CurrentTime, CollisionNode, CollisionPivot);
 }
 
 
@@ -281,7 +311,7 @@ void SetNormalForAFace(FDatasmithMesh& DatasmithMesh, int32 Index, const FVector
 	DatasmithMesh.SetNormal(Index * 3 + 2, NormalVector.X, NormalVector.Y, NormalVector.Z);
 }
 
-bool FillDatasmithMeshFromBoundingBox(FDatasmithMesh& DatasmithMesh, INode* ExportNode, FTransform Pivot)
+bool FillDatasmithMeshFromBoundingBox(TimeValue CurrentTime, FDatasmithMesh& DatasmithMesh, INode* ExportNode, FTransform Pivot)
 {
 	if (!ExportNode)
 	{
@@ -291,7 +321,7 @@ bool FillDatasmithMeshFromBoundingBox(FDatasmithMesh& DatasmithMesh, INode* Expo
 
 	BOOL bNeedsDelete;
 
-	Mesh* MaxMesh = GetMeshFromRenderMesh(ExportNode, bNeedsDelete, GetCOREInterface()->GetTime());
+	Mesh* MaxMesh = GetMeshFromRenderMesh(CurrentTime, ExportNode, bNeedsDelete);
 
 	if (!MaxMesh)
 	{
@@ -390,14 +420,14 @@ bool FillDatasmithMeshFromBoundingBox(FDatasmithMesh& DatasmithMesh, INode* Expo
 }
 
 
-bool CreateDatasmithMeshFromMaxMesh(FDatasmithMesh& DatasmithMesh, INode* Node, const TCHAR* MeshName, const FRenderMeshForConversion& RenderMesh, bool bConsolidateMaterialIds, TSet<uint16>& SupportedChannels, TMap<int32, int32>& UVChannelsMap)
+bool CreateDatasmithMeshFromMaxMesh(FDatasmithMesh& DatasmithMesh, const MeshConversionParams& Params, TSet<uint16>& SupportedChannels, TMap<int32, int32>& UVChannelsMap)
 {
 	bool bResult = false;
-	if (RenderMesh.GetMesh()->getNumFaces())
+	if (Params.RenderMesh.GetMesh()->getNumFaces())
 	{
 		// Copy mesh to clean it before filling Datasmith mesh from it
 		Mesh CachedMesh;
-		CachedMesh.DeepCopy(RenderMesh.GetMesh(), TOPO_CHANNEL | GEOM_CHANNEL | TEXMAP_CHANNEL | VERTCOLOR_CHANNEL);
+		CachedMesh.DeepCopy(Params.RenderMesh.GetMesh(), TOPO_CHANNEL | GEOM_CHANNEL | TEXMAP_CHANNEL | VERTCOLOR_CHANNEL);
 
 		CachedMesh.DeleteIsoVerts();
 		CachedMesh.RemoveDegenerateFaces();
@@ -410,7 +440,7 @@ bool CreateDatasmithMeshFromMaxMesh(FDatasmithMesh& DatasmithMesh, INode* Node,
 		if (CachedMesh.getNumFaces() > 0)
 		{
 			// todo: pivot
-			FillDatasmithMeshFromMaxMesh(DatasmithMesh, CachedMesh, Node, bConsolidateMaterialIds, SupportedChannels, UVChannelsMap, RenderMesh.GetPivot());
+			FillDatasmithMeshFromMaxMesh(DatasmithMesh, CachedMesh, Params.Node, Params.bConsolidateMaterialIds, SupportedChannels, UVChannelsMap, Params.RenderMesh.GetPivot());
 
 			bResult = true; // Set to true, don't care what ExportToUObject does here - we need to move it to a thread anyway
 		}
@@ -419,81 +449,6 @@ bool CreateDatasmithMeshFromMaxMesh(FDatasmithMesh& DatasmithMesh, INode* Node,
 	return bResult;
 }
 
-// todo: paralelize calls to ExportToUObject 
-bool ConvertMaxMeshToDatasmith(ISceneTracker& Scene, TSharedPtr<IDatasmithMeshElement>& DatasmithMeshElement, INode* Node, const TCHAR* MeshName, const FRenderMeshForConversion& RenderMesh, bool bConsolidateMaterialIds, TSet<uint16>& SupportedChannels, const FRenderMeshForConversion& CollisionMesh)
-{
-	// Reset old mesh
-	if (DatasmithMeshElement)
-	{
-		// todo: potential mesh reuse - when DatasmithMeshElement allows to reset materials(as well as other params)
-		Scene.ReleaseMeshElement(DatasmithMeshElement);
-		SupportedChannels.Reset();
-	}
-
-	TOptional<FDatasmithMaxStaticMeshAttributes> DatasmithAttributes = FDatasmithMaxStaticMeshAttributes::ExtractStaticMeshAttributes(Node);
-
-	DatasmithMeshElement = FDatasmithSceneFactory::CreateMesh(MeshName);
-
-	if (DatasmithAttributes && DatasmithAttributes->GetExportMode() == EStaticMeshExportMode::BoundingBox)
-	{
-		FDatasmithMesh DatasmithMesh;
-		if (FillDatasmithMeshFromBoundingBox(DatasmithMesh, Node, FTransform::Identity))
-		{
-			Scene.AddMeshElement(DatasmithMeshElement, DatasmithMesh, nullptr);
-			return true;
-		}
-		else
-		{
-			LogWarning(FString(TEXT("Invalid object: ")) + Node->GetName());
-			return false;
-		}
-	}
-
-	FDatasmithMesh DatasmithMesh;
-	TMap<int32, int32> UVChannelsMap;
-	if (!CreateDatasmithMeshFromMaxMesh(DatasmithMesh, Node, MeshName, RenderMesh, bConsolidateMaterialIds, SupportedChannels, UVChannelsMap))
-	{
-		LogWarning(FString(TEXT("Invalid object: ")) + Node->GetName());
-		return false;
-	}
-
-	// Mapping between the 3ds max channel and the exported mesh channel
-	if (DatasmithAttributes)
-	{
-		constexpr int32 MaxToUnrealUVOffset = -1;
-		constexpr int32 DefaultValue = -1;
-		const int32 SelectedLightmapUVChannel = DatasmithAttributes->GetLightmapUVChannel();
-		const int32* ExportedSelectedChannel = UVChannelsMap.Find(SelectedLightmapUVChannel + MaxToUnrealUVOffset);
-
-		if (ExportedSelectedChannel)
-		{
-			DatasmithMeshElement->SetLightmapCoordinateIndex(*ExportedSelectedChannel);
-		}
-		else if (SelectedLightmapUVChannel != DefaultValue)
-		{
-			LogWarning(*FString::Printf(TEXT("%s won't use the channel %i for its lightmap because it's not supported by the mesh. A new channel will be generated.")
-				, static_cast<const TCHAR*>(Node->GetName())
-				, SelectedLightmapUVChannel));
-		}
-	}
-
-
-	FDatasmithMesh* DatasmithCollisionMeshPtr = nullptr;
-	FDatasmithMesh DatasmithCollisionMesh;
-	if (CollisionMesh.IsValid())
-	{
-		TSet<uint16> SupportedChannelsDummy; // ignore map channels for collision mesh
-		TMap<int32, int32> UVChannelsMapDummy;
-		if (CreateDatasmithMeshFromMaxMesh(DatasmithCollisionMesh, CollisionMesh.GetNode(), nullptr, CollisionMesh, true, SupportedChannelsDummy, UVChannelsMapDummy))
-		{
-			DatasmithCollisionMeshPtr = &DatasmithCollisionMesh;
-		}
-	}
-
-	Scene.AddMeshElement(DatasmithMeshElement, DatasmithMesh, DatasmithCollisionMeshPtr);
-	return true;
-}
-
 // todo: copied from DatasmithMaxMeshExporter.cpp
 
 Object* GetBaseObject(INode* Node, TimeValue Time)
@@ -516,7 +471,7 @@ Object* GetBaseObject(INode* Node, TimeValue Time)
 }
 #define VRAY_PROXY_DISPLAY_AS_MESH	4	// Value to set on a VRay Mesh Proxy to get the mesh
 
-int SetObjectParamValue(Object *Obj, const FString& ParamName, int DesiredValue)
+int SetObjectParamValue(TimeValue CurrentTime, Object *Obj, const FString& ParamName, int DesiredValue)
 {
 	bool bFoundDisplayValue = false;
 	int PrevDisplayValue = DesiredValue; // Display value to see mesh
@@ -530,10 +485,10 @@ int SetObjectParamValue(Object *Obj, const FString& ParamName, int DesiredValue)
 			ParamDef& ParamDefinition = ParamBlockDesc->paramdefs[ParamIndex];
 			if (FCString::Stricmp(ParamDefinition.int_name, *ParamName) == 0)
 			{
-				PrevDisplayValue = ParamBlock2->GetInt(ParamDefinition.ID, GetCOREInterface()->GetTime());
+				PrevDisplayValue = ParamBlock2->GetInt(ParamDefinition.ID, CurrentTime);
 				if (PrevDisplayValue != DesiredValue)
 				{
-					ParamBlock2->SetValue(ParamDefinition.ID, GetCOREInterface()->GetTime(), DesiredValue);
+					ParamBlock2->SetValue(ParamDefinition.ID, CurrentTime, DesiredValue);
 				}
 				bFoundDisplayValue = true;
 				break;
@@ -546,7 +501,7 @@ int SetObjectParamValue(Object *Obj, const FString& ParamName, int DesiredValue)
 }
 
 
-Mesh* GetMeshFromRenderMesh(INode* Node, BOOL& bNeedsDelete, TimeValue CurrentTime)
+Mesh* GetMeshFromRenderMesh(TimeValue CurrentTime, INode* Node, BOOL& bNeedsDelete)
 {
 	Object* Obj = GetBaseObject(Node, CurrentTime);
 	const Class_ID& ObjectClassID = Obj->ClassID();
@@ -557,12 +512,12 @@ Mesh* GetMeshFromRenderMesh(INode* Node, BOOL& bNeedsDelete, TimeValue CurrentTi
 	if (ObjectClassID == VRAYPROXY_CLASS_ID)
 	{
 		// Need the high resolution render mesh associated with the VRay Mesh Proxy for the export
-		PreviousMeshDisplayValue = SetObjectParamValue(Obj, VRayProxyParamName, VRAY_PROXY_DISPLAY_AS_MESH);
+		PreviousMeshDisplayValue = SetObjectParamValue(CurrentTime, Obj, VRayProxyParamName, VRAY_PROXY_DISPLAY_AS_MESH);
 	}
 	else if(ObjectClassID == BODYOBJECT_CLASS_ID)
 	{
 		// Need to make sure we are using the viewport mesh on BodyObject, otherwise the RenderMesh gives a tessellated low resolution mesh.
-		PreviousMeshDisplayValue = SetObjectParamValue(Obj, BodyObjectViewportMeshParamName, 1);
+		PreviousMeshDisplayValue = SetObjectParamValue(CurrentTime, Obj, BodyObjectViewportMeshParamName, 1);
 	}
 
 	GeomObject* GeomObj = static_cast<GeomObject*>(Obj);
@@ -577,11 +532,11 @@ Mesh* GetMeshFromRenderMesh(INode* Node, BOOL& bNeedsDelete, TimeValue CurrentTi
 	// Restore display state if different from mesh display
 	if (ObjectClassID == VRAYPROXY_CLASS_ID && PreviousMeshDisplayValue != VRAY_PROXY_DISPLAY_AS_MESH)
 	{
-		SetObjectParamValue(Obj, VRayProxyParamName, PreviousMeshDisplayValue);
+		SetObjectParamValue(CurrentTime, Obj, VRayProxyParamName, PreviousMeshDisplayValue);
 	}
 	else if(ObjectClassID == BODYOBJECT_CLASS_ID && PreviousMeshDisplayValue != 1)
 	{
-		SetObjectParamValue(Obj, BodyObjectViewportMeshParamName, PreviousMeshDisplayValue);
+		SetObjectParamValue(CurrentTime, Obj, BodyObjectViewportMeshParamName, PreviousMeshDisplayValue);
 	}
 
 	return RenderMesh;
@@ -589,6 +544,95 @@ Mesh* GetMeshFromRenderMesh(INode* Node, BOOL& bNeedsDelete, TimeValue CurrentTi
 
 
 }
+}
+
+namespace DatasmithMaxDirectLink
+{
+
+// todo: paralelize calls to ExportToUObject 
+bool ConvertMaxMeshToDatasmith(TimeValue CurrentTime, ISceneTracker& Scene, FMeshConverterSource& MeshSource, FMeshConverted& MeshConverted)
+{
+	TOptional<FDatasmithMaxStaticMeshAttributes> DatasmithAttributes = FDatasmithMaxStaticMeshAttributes::ExtractStaticMeshAttributes(MeshSource.Node);
+
+	MeshConverted.DatasmithMeshElement = FDatasmithSceneFactory::CreateMesh(*MeshSource.MeshName);
+
+	if (DatasmithAttributes && DatasmithAttributes->GetExportMode() == EStaticMeshExportMode::BoundingBox)
+	{
+		FDatasmithMesh DatasmithMesh;
+		if (GeomUtils::FillDatasmithMeshFromBoundingBox(CurrentTime, DatasmithMesh, MeshSource.Node, FTransform::Identity))
+		{
+			Scene.AddMeshElement(MeshConverted.DatasmithMeshElement, DatasmithMesh, nullptr);
+			return true;
+		}
+		else
+		{
+			LogWarning(FString(TEXT("Invalid object: ")) + MeshSource.Node->GetName());
+			return false;
+		}
+	}
+
+	FDatasmithMesh DatasmithMesh;
+	TMap<int32, int32> UVChannelsMap;
+
+	MeshConversionParams RenderMeshParams = {
+		MeshSource.RenderMesh.GetNode(),
+		*MeshSource.MeshName,
+		MeshSource.RenderMesh,
+		MeshSource.bConsolidateMaterialIds
+	};
+
+	if (!GeomUtils::CreateDatasmithMeshFromMaxMesh(DatasmithMesh, RenderMeshParams, MeshConverted.SupportedChannels, UVChannelsMap))
+	{
+		LogWarning(FString(TEXT("Invalid object: ")) + MeshSource.Node->GetName());
+		return false;
+	}
+
+	// Mapping between the 3ds max channel and the exported mesh channel
+	if (DatasmithAttributes)
+	{
+		constexpr int32 MaxToUnrealUVOffset = -1;
+		constexpr int32 DefaultValue = -1;
+		const int32 SelectedLightmapUVChannel = DatasmithAttributes->GetLightmapUVChannel();
+		const int32* ExportedSelectedChannel = UVChannelsMap.Find(SelectedLightmapUVChannel + MaxToUnrealUVOffset);
+
+		if (ExportedSelectedChannel)
+		{
+			MeshConverted.DatasmithMeshElement->SetLightmapCoordinateIndex(*ExportedSelectedChannel);
+		}
+		else if (SelectedLightmapUVChannel != DefaultValue)
+		{
+			LogWarning(*FString::Printf(TEXT("%s won't use the channel %i for its lightmap because it's not supported by the mesh. A new channel will be generated.")
+				, static_cast<const TCHAR*>(MeshSource.Node->GetName())
+				, SelectedLightmapUVChannel));
+		}
+	}
+
+	FDatasmithMesh* DatasmithCollisionMeshPtr = nullptr;
+	FDatasmithMesh DatasmithCollisionMesh;
+	if (MeshSource.CollisionMesh.IsValid())
+	{
+		TSet<uint16> SupportedChannelsDummy; // ignore map channels for collision mesh
+		TMap<int32, int32> UVChannelsMapDummy;
+
+		MeshConversionParams CollisionParams = {
+			MeshSource.CollisionMesh.GetNode(),
+			nullptr,
+			MeshSource.CollisionMesh,
+			true
+		};
+
+		if (GeomUtils::CreateDatasmithMeshFromMaxMesh(DatasmithCollisionMesh, CollisionParams, SupportedChannelsDummy, UVChannelsMapDummy))
+		{
+			DatasmithCollisionMeshPtr = &DatasmithCollisionMesh;
+		}
+	}
+
+	Scene.AddMeshElement(MeshConverted.DatasmithMeshElement, DatasmithMesh, DatasmithCollisionMeshPtr);
+	return true;
+}
+
+}
+
 
 #include "Windows/HideWindowsPlatformTypes.h"
 
diff --git a/Engine/Source/Programs/Enterprise/Datasmith/DatasmithMaxExporter/Private/DatasmithMaxDirectLinkNotifications.cpp b/Engine/Source/Programs/Enterprise/Datasmith/DatasmithMaxExporter/Private/DatasmithMaxDirectLinkNotifications.cpp
--- a/Engine/Source/Programs/Enterprise/Datasmith/DatasmithMaxExporter/Private/DatasmithMaxDirectLinkNotifications.cpp
+++ b/Engine/Source/Programs/Enterprise/Datasmith/DatasmithMaxExporter/Private/DatasmithMaxDirectLinkNotifications.cpp
@@ -46,6 +46,7 @@ public:
 		InvalidateGeometry(nodes);
 	}
 
+
 	virtual void MappingChanged(NodeKeyTab& nodes) override
 	{
 		LogNodeEvent(L"MappingChanged", nodes);
@@ -127,7 +128,7 @@ public:
 
 		for (int NodeIndex = 0; NodeIndex < nodes.Count(); ++NodeIndex)
 		{
-			SceneTracker.NodePropertiesChanged(nodes[NodeIndex]);
+			SceneTracker.NodeNameChanged(nodes[NodeIndex]);
 		}
 	}
 
@@ -537,6 +538,16 @@ void FNotifications::On3dsMaxNotification(void* param, NotifyInfo* info)
 
 	case NOTIFY_PLUGINS_PRE_SHUTDOWN: // This one crashes when calling LogInfo
 		return;
+	// case NOTIFY_NODE_LAYER_CHANGED:
+	// 	break;
+	case NOTIFY_NODE_UNLINKED:
+	case NOTIFY_NODE_CREATED:
+	case NOTIFY_NODE_LINKED:
+		{
+			INode* Node = reinterpret_cast<INode*>(info->callParam);
+			LogDebugNode(NotificationsHandler.ConvertNotificationCodeToString(info->intcode), Node);
+		}
+		break;
 	default:
 		LOG_DEBUG_HEAVY(FString(TEXT("Notify: ")) + NotificationsHandler.ConvertNotificationCodeToString(info->intcode));
 	};
diff --git a/Engine/Source/Programs/Enterprise/Datasmith/DatasmithMaxExporter/Private/DatasmithMaxExporterWithDirectLink.cpp b/Engine/Source/Programs/Enterprise/Datasmith/DatasmithMaxExporter/Private/DatasmithMaxExporterWithDirectLink.cpp
--- a/Engine/Source/Programs/Enterprise/Datasmith/DatasmithMaxExporter/Private/DatasmithMaxExporterWithDirectLink.cpp
+++ b/Engine/Source/Programs/Enterprise/Datasmith/DatasmithMaxExporter/Private/DatasmithMaxExporterWithDirectLink.cpp
@@ -59,6 +59,8 @@ MAX_INCLUDES_START
 	#include "ISceneEventManager.h"
 
 	#include "IFileResolutionManager.h" // for GetActualPath
+
+	#include "xref/iXrefObj.h"
 MAX_INCLUDES_END
 
 namespace DatasmithMaxDirectLink
@@ -100,9 +102,9 @@ public:
 		DatasmithSceneRef->SetProductVersion(*Version);
 	}
 
-	TSharedRef<IDatasmithScene> GetDatasmithScene()
+	TSharedPtr<IDatasmithScene> GetDatasmithScene()
 	{
-		return DatasmithSceneRef.ToSharedRef();
+		return DatasmithSceneRef;
 	}
 
 	FDatasmithSceneExporter& GetSceneExporter()
@@ -149,17 +151,95 @@ private:
 };
 
 // Every node which is resolved to the same object is considered an instance
-// This class holds all this nodes and the object they resolve to
-struct FInstances
+// This class holds all this nodes which resolve to the same object
+struct FInstances: FNoncopyable
 {
+	AnimHandle Handle; // Handle of anim that is instanced
+
 	Object* EvaluatedObj = nullptr;
 	Mtl* Material = nullptr; // Material assigned to Datasmith StaticMesh, used to check if a particular instance needs to verride it
 
 	TSet<class FNodeTracker*> NodeTrackers;
 
-	// Mesh conversion results
-	TSet<uint16> SupportedChannels;
-	TSharedPtr<IDatasmithMeshElement> DatasmithMeshElement;
+	// Datasmith mesh conversion results
+	FMeshConverted Converted;
+
+	bool HasMesh()
+	{
+		return Converted.DatasmithMeshElement.IsValid();
+	}
+
+	const TCHAR* GetStaticMeshPathName()
+	{
+		return Converted.DatasmithMeshElement->GetName();
+	}
+
+	void AssignMaterialToStaticMesh(Mtl* InMaterial)
+	{
+		Material = InMaterial;
+		AssignMeshMaterials(Converted.DatasmithMeshElement, Material, Converted.SupportedChannels);
+	}
+};
+
+// todo: rename Manager
+//   Groups all geometry nodes by their prototype object(geom they resolve to)
+class FInstancesManager
+{
+public:
+	void Reset()
+	{
+		InstancesForAnimHandle.Reset();
+	}
+
+	FInstances& AddNodeTracker(FNodeTracker& NodeTracker, FMeshNodeConverter& Converter, Object* Obj)
+	{
+		TUniquePtr<FInstances>& Instances = InstancesForAnimHandle.FindOrAdd(Converter.InstanceHandle);
+
+		if (!Instances)
+		{
+			Instances = MakeUnique<FInstances>();
+			Instances->Handle = Converter.InstanceHandle;
+			Instances->EvaluatedObj = Obj;
+		}
+
+		// need to invalidate mesh assignment to node that wasn't the first to add to instances(so if instances weren't invalidated - this node needs mesh anyway)
+		Instances->NodeTrackers.Add(&NodeTracker);
+		return *Instances;
+	}
+
+	FInstances* RemoveNodeTracker(FNodeTracker& NodeTracker)
+	{
+		if (FInstances* Instances = GetInstancesForNodeTracker(NodeTracker))
+		{
+			// need to invalidate mesh assignment to node that wasn't the first to add to instances(so if instances weren't invalidated - this node needs mesh anyway)
+			Instances->NodeTrackers.Remove(&NodeTracker);
+			return Instances;
+		}
+
+		return nullptr;
+	}
+
+	FInstances* GetInstancesForNodeTracker(FNodeTracker& NodeTracker)
+	{
+		if (!ensure(NodeTracker.GetConverter().ConverterType == FNodeConverter::MeshNode))
+		{
+			return nullptr;
+		}
+		if (TUniquePtr<FInstances>* InstancesPtr = InstancesForAnimHandle.Find(static_cast<FMeshNodeConverter&>(NodeTracker.GetConverter()).InstanceHandle))
+		{
+			return InstancesPtr->Get();
+		}
+		return nullptr;
+	}
+
+	void RemoveInstances(const FInstances& Instances)
+	{
+		ensure(Instances.NodeTrackers.IsEmpty()); // Supposed to remove only unused Instances
+		InstancesForAnimHandle.Remove(Instances.Handle);
+	}
+
+private:
+	TMap<AnimHandle, TUniquePtr<FInstances>> InstancesForAnimHandle; // set of instanced nodes for each AnimHandle
 };
 
 class FLayerTracker
@@ -350,7 +430,7 @@ public:
 	{
 		INode* Node = NodeTracker.Node;
 		INode* ParentNode = Node->GetParentNode();
-		DatasmithMaxExporterUtils::ExportMaxTagsForDatasmithActor( NodeTracker.DatasmithActorElement, Node, ParentNode, KnownMaxDesc, KnownMaxSuperClass );
+		DatasmithMaxExporterUtils::ExportMaxTagsForDatasmithActor( NodeTracker.GetConverted().DatasmithActorElement, Node, ParentNode, KnownMaxDesc, KnownMaxSuperClass );
 	}
 
 private:
@@ -551,6 +631,147 @@ public:
 	bool bLoaded = false;
 };
 
+class FIncludeXRefGuard
+{
+public:
+
+
+	FIncludeXRefGuard(bool bInIncludeXRefWhileParsing): bIncludeXRefWhileParsing(bInIncludeXRefWhileParsing)
+	{
+		if (bIncludeXRefWhileParsing)
+		{
+			bIncludeXRefsInHierarchyStored = GetCOREInterface()->GetIncludeXRefsInHierarchy();
+			GetCOREInterface()->SetIncludeXRefsInHierarchy(true);
+		}
+	}
+	~FIncludeXRefGuard()
+	{
+		if (bIncludeXRefWhileParsing)
+		{
+			GetCOREInterface()->SetIncludeXRefsInHierarchy(bIncludeXRefsInHierarchyStored);
+		}
+		
+	}
+
+	bool bIncludeXRefWhileParsing;
+	BOOL bIncludeXRefsInHierarchyStored;
+};
+
+class FInvalidatedNodeTrackers: FNoncopyable
+{
+public:
+	void Add(FNodeTracker& NodeTracker)
+	{
+		InvalidatedNodeTrackers.Add(&NodeTracker);
+	}
+
+	/**
+	 * @return if anything was deleted
+	 */
+	bool PurgeDeletedNodeTrackers(class FSceneTracker& Scene);
+
+	template <typename Func>
+	void EnumerateAll(Func Callable)
+	{
+		for (FNodeTracker* NodeTracker : InvalidatedNodeTrackers)
+		{
+			Callable(*NodeTracker);
+		}
+	}
+
+	int32 Num()
+	{
+		return InvalidatedNodeTrackers.Num();
+	}
+
+	void Append(const TSet<FNodeTracker*>& NodeTrackers)
+	{
+		InvalidatedNodeTrackers.Append(NodeTrackers);
+	}
+
+	// Called when update is finished and all changes are processed and recorded
+	void Finish()
+	{
+		InvalidatedNodeTrackers.Reset();
+	}
+
+	// Scene is reset so invalidation is reset too
+	void Reset()
+	{
+		InvalidatedNodeTrackers.Reset();
+	}
+
+	bool HasInvalidated()
+	{
+		return !InvalidatedNodeTrackers.IsEmpty();
+	}
+
+	void RemoveFromInvalidated(FNodeTracker& NodeTracker)
+	{
+		InvalidatedNodeTrackers.Remove(&NodeTracker);
+	}
+
+private:
+	TSet<FNodeTracker*> InvalidatedNodeTrackers;
+};
+
+class FNodeTrackersNames: FNoncopyable
+{
+public:
+	TMap<FString, TSet<FNodeTracker*>> NodesForName;  // Each name can be used by a set of nodes
+	void Reset()
+	{
+		NodesForName.Reset();
+	}
+
+	const FString& GetNodeName(FNodeTracker& NodeTracker)
+	{
+		return NodeTracker.Name;
+	}
+
+	void Update(FNodeTracker& NodeTracker)
+	{
+		FString Name = NodeTracker.Node->GetName();
+		if (Name != NodeTracker.Name)
+		{
+			NodesForName[NodeTracker.Name].Remove(&NodeTracker);
+
+			NodeTracker.Name = Name;
+			NodesForName.FindOrAdd(NodeTracker.Name).Add(&NodeTracker);
+		}
+	}
+
+	void Add(FNodeTracker& NodeTracker)
+	{
+		FString Name = NodeTracker.Node->GetName();
+
+		// todo: dupclicated with Update
+		NodeTracker.Name = Name;
+		NodesForName.FindOrAdd(NodeTracker.Name).Add(&NodeTracker);
+	}
+
+	void Remove(const FNodeTracker& NodeTracker)
+	{
+		if (TSet<FNodeTracker*>* NodeTrackersPtr = NodesForName.Find(NodeTracker.Name))
+		{
+			NodeTrackersPtr->Remove(&NodeTracker);
+		}
+	}
+
+	template <typename Func>
+	void EnumerateForName(const FString& Name, Func Callable)
+	{
+		if (TSet<FNodeTracker*>* NodeTrackersPtr = NodesForName.Find(Name))
+		{
+			for (FNodeTracker* NodeTracker: *NodeTrackersPtr)
+			{
+				Callable(*NodeTracker);
+			}
+		}
+	}
+};
+
+
 // Holds states of entities for syncronization and handles change events
 class FSceneTracker: public ISceneTracker
 {
@@ -561,13 +782,22 @@ public:
 		, NotificationsHandler(InNotificationsHandler)
 		, MaterialsCollectionTracker(*this) {}
 
+	virtual TSharedRef<IDatasmithScene> GetDatasmithSceneRef() override
+	{
+		return ExportedScene.GetDatasmithScene().ToSharedRef();
+	}
+
 	bool ParseScene()
 	{
+		FIncludeXRefGuard IncludeXRefGuard(bIncludeXRefWhileParsing);
 		INode* Node = GetCOREInterface()->GetRootNode();
 		bSceneParsed = ParseScene(Node);
 		return bSceneParsed;
 	}
 
+	bool bParseXRefScenes = true;
+	bool bIncludeXRefWhileParsing = false; 
+
 	// Parse scene or XRef scene(in this case attach to parent datasmith actor)
 	bool ParseScene(INode* SceneRootNode, FXRefScene XRefScene=FXRefScene())
 	{
@@ -582,33 +812,36 @@ public:
 			return false;
 		}
 
-		// Parse XRefScenes
-		for (int XRefChild = 0; XRefChild < SceneRootNode->GetXRefFileCount(); ++XRefChild)
+		if (!bIncludeXRefWhileParsing)
 		{
-			DWORD XRefFlags = SceneRootNode->GetXRefFlags(XRefChild);
+			// Parse XRefScenes
+			for (int XRefChild = 0; XRefChild < SceneRootNode->GetXRefFileCount(); ++XRefChild)
+			{
+				DWORD XRefFlags = SceneRootNode->GetXRefFlags(XRefChild);
 
-			SCENE_UPDATE_STAT_INC(ParseScene, XRefFileEncountered);
+				SCENE_UPDATE_STAT_INC(ParseScene, XRefFileEncountered);
 
-			// XRef is disabled - not shown in viewport/render. Not loaded.
-			if (XRefFlags & XREF_DISABLED)
-			{
-				SCENE_UPDATE_STAT_INC(ParseScene, XRefFileDisabled);
-				// todo: baseline - doesn't check this - exports even disabled and XREF_HIDDEN scenes
-				continue;
-			}
+				// XRef is disabled - not shown in viewport/render. Not loaded.
+				if (XRefFlags & XREF_DISABLED)
+				{
+					SCENE_UPDATE_STAT_INC(ParseScene, XRefFileDisabled);
+					// todo: baseline - doesn't check this - exports even disabled and XREF_HIDDEN scenes
+					continue;
+				}
 
-			FString Path = FDatasmithMaxSceneExporter::GetActualPath(SceneRootNode->GetXRefFile(XRefChild).GetFileName());
-			if (FPaths::FileExists(Path) == false)
-			{
-				SCENE_UPDATE_STAT_INC(ParseScene, XRefFileMissing);
-				FString Error = FString("XRefScene file \"") + FPaths::GetCleanFilename(*Path) + FString("\" cannot be found");
-				// todo: logging
-				// DatasmithMaxLogger::Get().AddMissingAssetError(*Error);
-			}
-			else
-			{
-				SCENE_UPDATE_STAT_INC(ParseScene, XRefFileToParse);
-				ParseScene(SceneRootNode->GetXRefTree(XRefChild), FXRefScene{SceneRootNode, XRefChild});
+				FString Path = FDatasmithMaxSceneExporter::GetActualPath(SceneRootNode->GetXRefFile(XRefChild).GetFileName());
+				if (FPaths::FileExists(Path) == false)
+				{
+					SCENE_UPDATE_STAT_INC(ParseScene, XRefFileMissing);
+					FString Error = FString("XRefScene file \"") + FPaths::GetCleanFilename(*Path) + FString("\" cannot be found");
+					// todo: logging
+					// DatasmithMaxLogger::Get().AddMissingAssetError(*Error);
+				}
+				else
+				{
+					SCENE_UPDATE_STAT_INC(ParseScene, XRefFileToParse);
+					ParseScene(SceneRootNode->GetXRefTree(XRefChild), FXRefScene{SceneRootNode, XRefChild});
+				}
 			}
 		}
 
@@ -617,6 +850,8 @@ public:
 		{
 			if (FNodeTracker* NodeTracker = ParseNode(SceneRootNode->GetChildNode(ChildIndex)))
 			{
+				// Record XRef this child node is at the root of
+				// It's needed restore hierarchy when converting
 				NodeTracker->SetXRefIndex(XRefScene);
 			}
 		}
@@ -629,43 +864,19 @@ public:
 
 		SCENE_UPDATE_STAT_INC(ParseNode, NodesEncountered);
 
-		BOOL bIsNodeHidden = Node->IsNodeHidden(TRUE);
-
-		if (bIsNodeHidden == false)
-		{
-			//todo: when a refernced file is not found XRefObj is not resolved and it's kept as XREFOBJ_CLASS_ID
-			// instead of resolved class that it references
-//			if (ObjState.obj->ClassID() == XREFOBJ_CLASS_ID)
-//			{
-//				//max2017 and newer versions allow developers to get the file, with previous versions you only can retrieve the active one
-//#ifdef MAX_RELEASE_R19
-//				FString Path = FDatasmithMaxSceneExporter::GetActualPath(((IXRefObject8*)ObjState.obj)->GetFile(FALSE).GetFileName());
-//#else
-//				FString Path = FDatasmithMaxSceneExporter::GetActualPath(((IXRefObject8*)ObjState.obj)->GetActiveFile().GetFileName());
-//#endif
-//				if (FPaths::FileExists(Path) == false)
-//				{
-//					FString Error = FString("XRefObj file \"") + FPaths::GetCleanFilename(*Path) + FString("\" cannot be found");
-//					DatasmithMaxLogger::Get().AddMissingAssetError(*Error);
-//					bIsNodeHidden = true;
-//				}
-//			}
-		}
-
-		FNodeKey NodeKey = NodeEventNamespace::GetKeyByNode(Node);
+		FNodeTracker* NodeTracker =  GetNodeTracker(Node);
 
-		if (FNodeTrackerHandle* NodeTrackerHandle = NodeTrackers.Find(NodeKey))
+		if (NodeTracker)
 		{
 			// Node being added might already be tracked(e.g. if it was deleted before but Update wasn't called to SceneTracker yet)
-			FNodeTracker* NodeTracker = NodeTrackerHandle->GetNodeTracker();
-			ensure(NodeTracker->bDeleted);
+			//ensure(NodeTracker->bDeleted);
 			NodeTracker->bDeleted = false;
 			InvalidateNode(*NodeTracker);
 			return NodeTracker;
 		}
 		else
 		{
-			FNodeTracker* NodeTracker = AddNode(NodeKey, Node);
+			NodeTracker = AddNode(NodeEventNamespace::GetKeyByNode(Node), Node);
 
 			// Parse children
 			int32 ChildNum = Node->NumberOfChildren();
@@ -677,22 +888,6 @@ public:
 		}
 	}
 
-	void Reset()
-	{
-		bSceneParsed = false;
-		NodeTrackers.Reset();
-		NodeTrackersNames.Reset();
-		CollisionNodes.Reset();
-		InvalidatedNodeTrackers.Reset();
-		InvalidatedInstances.Reset();
-		MaterialsCollectionTracker.Reset();
-		LayersForAnimHandle.Reset();
-		NodesPerLayer.Reset();
-		NodeDatasmithMetadata.Reset();
-
-		InstancesForAnimHandle.Reset();
-	}
-
 	// Check every layer and if it's modified invalidate nodes assigned to it
 	// 3ds Max doesn't have events for all Layer changes(e.g. Name seems to be just an UI thing and has no notifications) so
 	// we need to go through all layers every update to see what's changed
@@ -729,7 +924,7 @@ public:
 				{
 					for(FNodeTracker* NodeTracker:* NodeTrackersPtr)
 					{
-						InvalidateNode(*NodeTracker);
+						InvalidateNode(*NodeTracker, false);
 					}
 				}
 				LayerTracker->bIsInvalidated = false;
@@ -762,24 +957,25 @@ public:
 		bool bResult = false;
 
 		{
-			FProgressStageGuard ProgressStage(MainStage, TEXT("Update"), 10);
+			const int32 StageCount = 12;
+			FProgressStageGuard ProgressStage(MainStage, TEXT("Update"), StageCount);
 			PROGRESS_STAGE_RESULT_DEFERRED
 			{
-				if (!ExportedScene.DatasmithSceneRef)
+				if (TSharedPtr<IDatasmithScene> ScenePtr = ExportedScene.GetDatasmithScene())
 				{
-					return FString(TEXT("<empty>"));
+					IDatasmithScene& Scene = *ScenePtr;
+					return FString::Printf(TEXT("Actors: %d; Meshes: %d, Materials: %d"), 
+					                       Scene.GetActorsCount(),
+					                       Scene.GetMeshesCount(),
+					                       Scene.GetMaterialsCount(),
+					                       Scene.GetTexturesCount()
+					);
 				}
-
-				IDatasmithScene& Scene = *ExportedScene.GetDatasmithScene();
-				return FString::Printf(TEXT("Actors: %d; Meshes: %d, Materials: %d"), 
-					Scene.GetActorsCount(),
-					Scene.GetMeshesCount(),
-					Scene.GetMaterialsCount(),
-					Scene.GetTexturesCount()
-				);
+				return FString(TEXT("<no scene>"));
 			};
 
 			bResult = UpdateInternalSafe(ProgressStage.Stage);
+
 		}
 		bUpdateInProgress = false;
 
@@ -808,6 +1004,9 @@ public:
 
 	bool UpdateInternal(FUpdateProgress::FStage& MainStage)
 	{
+
+		CurrentSyncPoint.Time = GetCOREInterface()->GetTime();
+
 		bool bChangeEncountered = false;
 
 		Stats.Reset();
@@ -829,7 +1028,6 @@ public:
 		}
 
 		// Changes present only when there are modified layers(changes checked manually), nodes(notified by Max) or materials(notified by Max with all changes in dependencies)
-		bChangeEncountered |= !InvalidatedNodeTrackers.IsEmpty();
 		bChangeEncountered |= !MaterialsCollectionTracker.GetInvalidatedMaterials().IsEmpty();
 
 		{
@@ -839,39 +1037,28 @@ public:
 				return FormatStatsRemoveDeletedNodes();
 			};
 
-			TArray<FNodeTracker*> DeletedNodeTrackers;
-			for (FNodeTracker* NodeTracker : InvalidatedNodeTrackers)
-			{
-				if (NodeTracker->bDeleted)
-				{
-					DeletedNodeTrackers.Add(NodeTracker);
-				}
-			}
-
-			for (FNodeTracker* NodeTrackerPtr : DeletedNodeTrackers)
-			{
-				RemoveNodeTracker(*NodeTrackerPtr);
-				SCENE_UPDATE_STAT_INC(RemoveDeletedNodes, Nodes);
-			}
-
+			bChangeEncountered |= InvalidatedNodeTrackers.PurgeDeletedNodeTrackers(*this);
 		}
 
-		// todo: move to NameChanged and NodeAdded?
 		{
-			PROGRESS_STAGE("Update node names")
+			PROGRESS_STAGE("Check Time Slider Validity")
 			PROGRESS_STAGE_RESULT_DEFERRED
 			{
-				return FormatStatsUpdateNodeNames();
+				return FormatStatsCheckTimeSliderValidity();
 			};
 
-			for (FNodeTracker* NodeTracker : InvalidatedNodeTrackers)
+			FIncludeXRefGuard IncludeXRefGuard(bIncludeXRefWhileParsing);
+
+			INode* SceneRootNode = GetCOREInterface()->GetRootNode();
+
+			int32 ChildNum = SceneRootNode->NumberOfChildren();
+			for (int32 ChildIndex = 0; ChildIndex < ChildNum; ++ChildIndex)
 			{
-				FString Name = NodeTracker->Node->GetName();
-				if (Name != NodeTracker->Name)
+				INode* Node = SceneRootNode->GetChildNode(ChildIndex);
+
+				if (FNodeTracker* NodeTracker = GetNodeTracker(Node))
 				{
-					NodeTrackersNames[NodeTracker->Name].Remove(NodeTracker);
-					NodeTracker->Name = Name;
-					NodeTrackersNames.FindOrAdd(NodeTracker->Name).Add(NodeTracker);
+					InvalidateOutdatedNodeTracker(*NodeTracker);
 				}
 			}
 		}
@@ -881,12 +1068,18 @@ public:
 			PROGRESS_STAGE_COUNTER(InvalidatedNodeTrackers.Num());
 
 			TSet<FNodeTracker*> NodesWithChangedCollisionStatus; // Need to invalidate these nodes to make them renderable or to keep them from renderable depending on collision status
-			for (FNodeTracker* NodeTracker : InvalidatedNodeTrackers)
+
+			InvalidatedNodeTrackers.EnumerateAll([&](FNodeTracker& NodeTracker)
+				{
+					ProgressCounter.Next();
+					UpdateCollisionStatus(NodeTracker, NodesWithChangedCollisionStatus);
+				});
+
+			// Rebuild all nodes that has changed them being colliders
+			for (FNodeTracker* NodeTracker : NodesWithChangedCollisionStatus)
 			{
-				ProgressCounter.Next();
-				UpdateCollisionStatus(NodeTracker, NodesWithChangedCollisionStatus);
+				InvalidateNode(*NodeTracker, false);
 			}
-			InvalidatedNodeTrackers.Append(NodesWithChangedCollisionStatus);
 
 			SCENE_UPDATE_STAT_SET(RefreshCollisions, ChangedNodes, NodesWithChangedCollisionStatus.Num());
 
@@ -904,11 +1097,11 @@ public:
 			};
 
 			PROGRESS_STAGE_COUNTER(InvalidatedNodeTrackers.Num());
-			for (FNodeTracker* NodeTracker : InvalidatedNodeTrackers)
-			{
-				ProgressCounter.Next();
-				UpdateNode(*NodeTracker);
-			}
+			InvalidatedNodeTrackers.EnumerateAll([&](FNodeTracker& NodeTracker)
+				{
+					ProgressCounter.Next();
+					UpdateNode(NodeTracker);
+				});
 		}
 
 		{
@@ -922,12 +1115,34 @@ public:
 			{
 				ProgressCounter.Next();
 				UpdateInstances(*Instances);
+
+				// Need to re-convert and reattach all instances of an updated node
 				InvalidatedNodeTrackers.Append(Instances->NodeTrackers);
 			}
 
 			InvalidatedInstances.Reset();
 		}
 
+		{
+			PROGRESS_STAGE("Convert nodes to datasmith")
+			PROGRESS_STAGE_RESULT_DEFERRED
+			{
+				return FormatStatsConvertNodesToDatasmith();
+			};
+
+			PROGRESS_STAGE_COUNTER(InvalidatedNodeTrackers.Num());
+			InvalidatedNodeTrackers.EnumerateAll([&](FNodeTracker& NodeTracker)
+				{
+					ProgressCounter.Next();
+
+					if (NodeTracker.HasConverter())
+					{
+						SCENE_UPDATE_STAT_INC(ConvertNodes, Converted)
+						NodeTracker.GetConverter().ConvertToDatasmith(*this, NodeTracker);
+					}
+				});
+		}
+
 		{
 			PROGRESS_STAGE("Reparent Datasmith Actors");
 			PROGRESS_STAGE_RESULT_DEFERRED
@@ -935,12 +1150,39 @@ public:
 				return FormatStatsReparentDatasmithActors();
 			};
 			
-			for (FNodeTracker* NodeTracker : InvalidatedNodeTrackers)
-			{
-				AttachNodeToDatasmithScene(*NodeTracker);
-			}
+			InvalidatedNodeTrackers.EnumerateAll([&](FNodeTracker& NodeTracker)
+				{
+					AttachNodeToDatasmithScene(NodeTracker);
+				});
 			;
-			InvalidatedNodeTrackers.Reset();
+		}
+
+		{
+			PROGRESS_STAGE("Mark nodes validated")
+			// Finally mark all nodetrackers as valid(we didn't interrupt update as it went to this point)
+			// And calculate subtree validity for each subtree of each updated node
+
+			// Maximize each invalidated node's subtree validity interval before recalculating it from updated nodes
+			InvalidatedNodeTrackers.EnumerateAll([](FNodeTracker& NodeTracker)
+				{
+					NodeTracker.SubtreeValidity.Invalidate();
+					NodeTracker.SubtreeValidity.ResetValidityInterval();
+				});
+
+			// Recalculate subtree validity 
+			InvalidatedNodeTrackers.EnumerateAll([&](FNodeTracker& NodeTracker)
+				{
+					PromoteValidity(NodeTracker, NodeTracker.Validity);
+				});
+
+			InvalidatedNodeTrackers.EnumerateAll([](FNodeTracker& NodeTracker)
+				{
+					NodeTracker.SubtreeValidity.SetValid();
+					NodeTracker.SetValid();
+				});
+
+			bChangeEncountered |= InvalidatedNodeTrackers.HasInvalidated(); // Right before resetting invalidated nodes, record that anything was invalidated
+			InvalidatedNodeTrackers.Finish();
 		}
 
 		// Each tracked(i.e. assigned to a visible node) Max material can result in multiple actual materials
@@ -984,7 +1226,7 @@ public:
 			{
 				ProgressCounter.Next();
 
-				MaterialsCollectionTracker.ConvertMaterial(ActualMaterial, ExportedScene.GetDatasmithScene(), ExportedScene.GetSceneExporter().GetAssetsOutputPath(), ActualTexmapsToUpdate);
+				MaterialsCollectionTracker.ConvertMaterial(ActualMaterial, ExportedScene.GetDatasmithScene().ToSharedRef(), ExportedScene.GetSceneExporter().GetAssetsOutputPath(), ActualTexmapsToUpdate);
 			}
 		}
 
@@ -1002,7 +1244,7 @@ public:
 				SCENE_UPDATE_STAT_INC(UpdateTextures, Total);
 
 				TArray<TSharedPtr<IDatasmithTextureElement>> TextureElements;
-				FDatasmithMaxMatExport::GetXMLTexture(ExportedScene.GetDatasmithScene(), Texture, ExportedScene.GetSceneExporter().GetAssetsOutputPath(), &TextureElements);
+				FDatasmithMaxMatExport::GetXMLTexture(ExportedScene.GetDatasmithScene().ToSharedRef(), Texture, ExportedScene.GetSceneExporter().GetAssetsOutputPath(), &TextureElements);
 				MaterialsCollectionTracker.UsedTextureToDatasmithElement.Add(Texture, TextureElements);
 			}
 		}
@@ -1039,7 +1281,7 @@ public:
 
 	FString FormatStatsRemoveDeletedNodes()
 	{
-		return FString::Printf(TEXT("Nodes: deleted %d of %d invalidated of total %d"), SCENE_UPDATE_STAT_GET(RemoveDeletedNodes, Nodes), InvalidatedNodeTrackers.Num(), NodeTrackers.Num());
+		return FString::Printf(TEXT("Nodes: deleted %d"), SCENE_UPDATE_STAT_GET(RemoveDeletedNodes, Nodes));
 	}
 
 	FString FormatStatsUpdateNodeNames()
@@ -1052,11 +1294,26 @@ public:
 		return FString::Printf(TEXT("Nodes: added %d to invalidated %d"), SCENE_UPDATE_STAT_GET(RefreshCollisions, ChangedNodes), InvalidatedNodeTrackers.Num());
 	}
 
+	FString FormatStatsCheckTimeSliderValidity()
+	{
+		return FString::Printf(TEXT("Check TimeSlider: checked %d, invalidated %d, skipped  - already invalidated %d, subtree valid %d"),
+			SCENE_UPDATE_STAT_GET(CheckTimeSlider, TotalChecks),
+			SCENE_UPDATE_STAT_GET(CheckTimeSlider, Invalidated),
+			SCENE_UPDATE_STAT_GET(CheckTimeSlider, SkippedAsAlreadyInvalidated),
+			SCENE_UPDATE_STAT_GET(CheckTimeSlider, SkippedAsSubtreeValid)
+			);
+	}
+
 	FString FormatStatsProcessInvalidatedNodes()
 	{
 		return FString::Printf(TEXT("Nodes: %d updated, %d skipped unselected, %d skipped hidden"), SCENE_UPDATE_STAT_GET(UpdateNode, NodesUpdated), SCENE_UPDATE_STAT_GET(UpdateNode, SkippedAsUnselected), SCENE_UPDATE_STAT_GET(UpdateNode, SkippedAsHiddenNode));
 	}
 
+	FString FormatStatsConvertNodesToDatasmith()
+	{
+		return FString::Printf(TEXT("Nodes: %d converted"), SCENE_UPDATE_STAT_GET(ConvertNodes, Converted));
+	}
+
 	FString FormatStatsProcessInvalidatedInstances()
 	{
 		return FString::Printf(TEXT("Instances: %ld updated"), InvalidatedInstances.Num());
@@ -1086,32 +1343,32 @@ public:
 	{
 		FDatasmithConverter Converter;
 		// Use the same name for the unique level sequence as the scene name
-		TSharedRef<IDatasmithLevelSequenceElement> LevelSequence = FDatasmithSceneFactory::CreateLevelSequence(ExportedScene.GetDatasmithScene()->GetName());
+		TSharedRef<IDatasmithLevelSequenceElement> LevelSequence = FDatasmithSceneFactory::CreateLevelSequence(GetDatasmithScene().GetName());
 		LevelSequence->SetFrameRate(GetFrameRate());
 
 		for (TPair<FNodeKey, FNodeTrackerHandle> NodeKeyAndNodeTracker: NodeTrackers)
 		{
 			FNodeTracker* NodeTracker = NodeKeyAndNodeTracker.Value.GetNodeTracker();
 
-			if (NodeTracker->DatasmithActorElement)
+			if (NodeTracker->HasConverted())
 			{
 
-				if (Lights.Contains(NodeTracker))
+				if (NodeTracker->GetConverterType() == FNodeConverter::LightNode)
 				{
-					const TSharedPtr<IDatasmithLightActorElement> LightElement = StaticCastSharedPtr< IDatasmithLightActorElement >(NodeTracker->DatasmithActorElement);
+					const TSharedPtr<IDatasmithLightActorElement> LightElement = StaticCastSharedPtr< IDatasmithLightActorElement >(NodeTracker->GetConverted().DatasmithActorElement);
 					const FMaxLightCoordinateConversionParams LightParams(NodeTracker->Node,
 						LightElement->IsA(EDatasmithElementType::AreaLight) ? StaticCastSharedPtr<IDatasmithAreaLightElement>(LightElement)->GetLightShape() : EDatasmithLightShape::None);
-					FDatasmithMaxSceneExporter::ExportAnimation(LevelSequence, NodeTracker->Node, NodeTracker->DatasmithActorElement->GetName(), Converter.UnitToCentimeter, LightParams);
+					FDatasmithMaxSceneExporter::ExportAnimation(LevelSequence, NodeTracker->Node, NodeTracker->GetConverted().DatasmithActorElement->GetName(), Converter.UnitToCentimeter, LightParams);
 				}
 				else
 				{
-					FDatasmithMaxSceneExporter::ExportAnimation(LevelSequence, NodeTracker->Node, NodeTracker->DatasmithActorElement->GetName(), Converter.UnitToCentimeter);
+					FDatasmithMaxSceneExporter::ExportAnimation(LevelSequence, NodeTracker->Node, NodeTracker->GetConverted().DatasmithActorElement->GetName(), Converter.UnitToCentimeter);
 				}
 			}
 		}
 		if (LevelSequence->GetAnimationsCount() > 0)
 		{
-			ExportedScene.GetDatasmithScene()->AddLevelSequence(LevelSequence);
+			GetDatasmithScene().AddLevelSequence(LevelSequence);
 		}
 	}
 
@@ -1121,8 +1378,8 @@ public:
 		LogDebugNode(TEXT("AddNode"), Node);
 		FNodeTracker* NodeTracker = NodeTrackers.Emplace(NodeKey, FNodeTrackerHandle(NodeKey, Node)).GetNodeTracker();
 
-		NodeTrackersNames.FindOrAdd(NodeTracker->Name).Add(NodeTracker);
-		InvalidatedNodeTrackers.Add(NodeTracker);
+		NodeTrackersNames.Add(*NodeTracker);
+		InvalidatedNodeTrackers.Add(*NodeTracker);
 
 		return NodeTracker;
 	}
@@ -1137,117 +1394,226 @@ public:
 		ExportedScene.DatasmithSceneRef->RemoveTexture(DatasmithTextureElement);		
 	}
 
-	void InvalidateNode(FNodeTracker& NodeTracker)
+	FNodeTracker* GetNodeTracker(FNodeKey NodeKey)
+	{
+		if (FNodeTrackerHandle* NodeTrackerHandle = NodeTrackers.Find(NodeKey))
+		{
+			return NodeTrackerHandle->GetNodeTracker();
+		}
+		return nullptr;
+	}
+
+	FNodeTracker* GetNodeTracker(INode* ParentNode)
 	{
-		ensure(!bUpdateInProgress);
+		return GetNodeTracker(NodeEventNamespace::GetKeyByNode(ParentNode));
+	}
 
-		if (NodeTracker.bDeleted)
+	// Promote validity up the ancestor chain
+	// Each Subtree Validity interval represents intersection of all invervals of descendants
+	// Used to quickly determine if some whole subtree needs to be updated when Time Slider changes
+	void PromoteValidity(FNodeTracker& NodeTracker, const FValidity& Validity)
+	{
+		if (Validity.Overlaps(NodeTracker.SubtreeValidity))
 		{
-			// Change events sometimes received for nodes that are already deleted
-			// skipping processing of node subtree because INode pointer may already be invalid
-			// Test case: create container, add node to it. Close it, open it, close again, then sync
-			// Change event will be received for NodeKey that returns NULL from NodeEventNamespace::GetNodeByKey
+			// Subtree validity is already fully within new validity - don't need to compute intersection and promote it further
 			return;
 		}
 
-		NodeTracker.Invalidate();
-		InvalidatedNodeTrackers.Add(&NodeTracker);
+		NodeTracker.SubtreeValidity.NarrowValidityToInterval(Validity);
 
-		// Invalidate whole sub-hierarchy of nodes that were previously recorded as children(this might have changed)
-		for (FNodeTracker* Child : NodeTracker.Children)
+		// Promote recalculated SubtreeValidity to parent
+		if (FNodeTracker* ParentNodeTracker = GetNodeTracker(NodeTracker.Node->GetParentNode()))
 		{
-			InvalidateNode(*Child);
+			PromoteValidity(*ParentNodeTracker, NodeTracker.SubtreeValidity);
 		}
+	}
 
-		// Invalidate whole sub-hierarchy of nodes are now children.
-		// E.g. a node could have been hidden so its children were attached to grandparent(parent of hidden node)
-		// Need to invalidate those to reattach
-		int32 ChildNum = NodeTracker.Node->NumberOfChildren();
-		for (int32 ChildIndex = 0; ChildIndex < ChildNum; ++ChildIndex)
+	// Node invalidated:
+	//   - on initial parsing
+	//   - change event received for node itself
+	//   - node up hierachy invalidated - previous or current
+	//   - node's validity interval is invalid for current time
+	// Possible improvements:
+	//   - do we really need to invalidate hierachy immediately? Might be long task(even though it's just setting flags) 
+	//      - probably do this on Update
+	//          'MarkedForUpdate' => recursively Invalidate(i.e. rebuild)
+	//          separate invalidation to 'Changed'(received event, interval, etc) and 'NeedRebuild'?
+	//            'NeedRebuild' are caused by 'Changed', not just 'invalidated' all.
+	//         - e.g. on initial parsing - only root nodes are 'changed'
+	//         - mark 'changed' as 'need rebuild' and descend to children, if a child has no 'need rebuild' set it, else skip its hierarchy
+
+	void InvalidateNode(FNodeTracker& NodeTracker, bool bCheckCalledInProgress = true)
+	{
+		// We don't expect node chances while Update inprogress(unless Invalidate called explicitly)
+		if (bCheckCalledInProgress)
+		{
+			ensure(!bUpdateInProgress);
+		}
+
+		if (NodeTracker.bDeleted)
+		{
+			// Change events sometimes received for nodes that are already deleted
+			// skipping processing of node subtree because INode pointer may already be invalid
+			// Test case: create container, add node to it. Close it, open it, close again, then sync
+			// Change event will be received for NodeKey that returns NULL from NodeEventNamespace::GetNodeByKey
+			if (bIncludeXRefWhileParsing)
+			{
+				check(NodeEventNamespace::GetNodeByKey(NodeTracker.NodeKey));
+				NodeTracker.bDeleted = false;
+			}
+			else
+			{
+				return;
+			}
+		}
+
+		if (NodeTracker.IsInvalidated())
 		{
-			InvalidateNode(NodeEventNamespace::GetKeyByNode(NodeTracker.Node->GetChildNode(ChildIndex)));
+			// Don't do work twice - already invalidated node would have its subhierarchy invalidated
+			//   in case subhierarchy was changed for already invalidated node - the changed nodes would invalidate too(responding to reparent event)
+			return;
 		}
 
-		NodeTracker.Children.Reset();
+		NodeTracker.Invalidate();
+		InvalidatedNodeTrackers.Add(NodeTracker);
+
+		// Invalidate whole sub-hierarchy of nodes are now children.
+		// E.g. a node could have been hidden so its children were attached to grandparent(parent of hidden node)
+		// Need to invalidate those to reattach
+		FIncludeXRefGuard IncludeXRefGuard(bIncludeXRefWhileParsing);
+		int32 ChildNum = NodeTracker.Node->NumberOfChildren();
+		for (int32 ChildIndex = 0; ChildIndex < ChildNum; ++ChildIndex)
+		{
+			InvalidateNode(NodeEventNamespace::GetKeyByNode(NodeTracker.Node->GetChildNode(ChildIndex)), bCheckCalledInProgress);
+		}
 	}
 
+	void InvalidateOutdatedNodeTracker(FNodeTracker& NodeTracker)
+	{
+		SCENE_UPDATE_STAT_INC(CheckTimeSlider, TotalChecks);
+		// Skip check in case whole subtree valid for current time OR node was already invalidated(and so its whole subtree)
+		if (NodeTracker.IsInvalidated())
+		{
+			SCENE_UPDATE_STAT_INC(CheckTimeSlider, SkippedAsAlreadyInvalidated);
+			return;
+		}
+
+		if (NodeTracker.IsSubtreeValidForSyncPoint(CurrentSyncPoint))
+		{
+			SCENE_UPDATE_STAT_INC(CheckTimeSlider, SkippedAsSubtreeValid);
+			return;
+		}
+
+		// If node is invalid - reevaliate whole subtree
+		// todo: it it possible to optimize reevaluation of whole subtree?
+		//   certain types of invalidation need not to propagate down to descendants(e.g. geometry change)
+		//   but some need to, like transform change
+		if (!NodeTracker.IsValidForSyncPoint(CurrentSyncPoint))
+		{
+			SCENE_UPDATE_STAT_INC(CheckTimeSlider, Invalidated);
+			InvalidateNode(NodeTracker, false);
+		}
+		else
+		{
+			FIncludeXRefGuard IncludeXRefGuard(bIncludeXRefWhileParsing);
+			int32 ChildNum = NodeTracker.Node->NumberOfChildren();
+			for (int32 ChildIndex = 0; ChildIndex < ChildNum; ++ChildIndex)
+			{
+				if (FNodeTracker* ChildNodeTracker = GetNodeTracker(NodeTracker.Node->GetChildNode(ChildIndex)))
+				{
+					InvalidateOutdatedNodeTracker(*ChildNodeTracker);
+				}
+			}
+		}
+	}
+	
 	// todo: make fine invalidates - full only something like geometry change, but finer for transform, name change and more
-	FNodeTracker* InvalidateNode(FNodeKey NodeKey)
+	FNodeTracker* InvalidateNode(FNodeKey NodeKey, bool bCheckCalledInProgress = true)
 	{
-		if (FNodeTrackerHandle* NodeTrackerHandle = NodeTrackers.Find(NodeKey))
+		FIncludeXRefGuard IncludeXRefGuard(bIncludeXRefWhileParsing);
+
+		LogDebugNode(TEXT("InvalidateNode"), NodeEventNamespace::GetNodeByKey(NodeKey));
+		// We don't expect node chances while Update inprogress(unless Invalidate called explicitly)
+		if (bCheckCalledInProgress)
+		{
+			ensure(!bUpdateInProgress);
+		}
+
+		if (FNodeTracker* NodeTracker =  GetNodeTracker(NodeKey))
 		{
-			FNodeTracker* NodeTracker = NodeTrackerHandle->GetNodeTracker();
 			if (NodeEventNamespace::GetNodeByKey(NodeKey))
 			{
-				InvalidateNode(*NodeTracker);
+				InvalidateNode(*NodeTracker, bCheckCalledInProgress);
 				return NodeTracker;
 			}
 			else
 			{
 				// Sometimes node update received without node Delete event
 				// Test case: create container, add node to it. Close it, open it, close again, then sync
-				InvalidatedNodeTrackers.Add(NodeTracker);
+				InvalidatedNodeTrackers.Add(*NodeTracker);
 				NodeTracker->bDeleted = true;
 			}
 		}
+		else
+		{
+			NodeAdded(NodeEventNamespace::GetNodeByKey(NodeKey));
+		}
 		return nullptr;
 	}
 
-	bool IsNodeInvalidated(const FNodeTrackerHandle& NodeTracker)
-	{
-		return NodeTracker.GetNodeTracker()->IsInvalidated();
-	}
-
-
 	void ClearNodeFromDatasmithScene(FNodeTracker& NodeTracker)
 	{
+		ReleaseNodeTrackerFromDatasmithMetadata(NodeTracker);
+
 		// remove from hierarchy
-		if (NodeTracker.DatasmithActorElement)
+		if (NodeTracker.HasConverted())
 		{
-			TSharedPtr<IDatasmithMeshActorElement> DatasmithMeshActor = NodeTracker.DatasmithMeshActor;
+			FNodeConverted& Converted = NodeTracker.GetConverted();
+			TSharedPtr<IDatasmithMeshActorElement> DatasmithMeshActor = Converted.DatasmithMeshActor;
 
 			// Remove mesh actor before removing its parent Actor in case there a separate MeshActor
-			if (NodeTracker.DatasmithMeshActor)
+			if (Converted.DatasmithMeshActor)
 			{
 				// if (NodeTracker.DatasmithMeshActor != NodeTracker.DatasmithActorElement)
 				{
-					NodeTracker.DatasmithActorElement->RemoveChild(NodeTracker.DatasmithMeshActor);
+					Converted.DatasmithActorElement->RemoveChild(Converted.DatasmithMeshActor);
 				}
-				NodeTracker.DatasmithMeshActor.Reset();
+				Converted.DatasmithMeshActor.Reset();
 				// todo: consider pool of MeshActors
 			}
 
-			if (TSharedPtr<IDatasmithActorElement> ParentActor = NodeTracker.DatasmithActorElement->GetParentActor())
+			if (TSharedPtr<IDatasmithActorElement> ParentActor = Converted.DatasmithActorElement->GetParentActor())
 			{
-				ParentActor->RemoveChild(NodeTracker.DatasmithActorElement);
+				ParentActor->RemoveChild(Converted.DatasmithActorElement);
 			}
 			else
 			{
 				// Detach all children(so they won't be reattached automatically to root when actor is detached from parent below)
 				// Children reattachment will happen later in Update
-				int32 ChildCount = NodeTracker.DatasmithActorElement->GetChildrenCount();
+				int32 ChildCount = Converted.DatasmithActorElement->GetChildrenCount();
 				// Remove last child each time to optimize array elements relocation
 				for(int32 ChildIndex = ChildCount-1; ChildIndex >= 0; --ChildIndex)
 				{
-					NodeTracker.DatasmithActorElement->RemoveChild(NodeTracker.DatasmithActorElement->GetChild(ChildIndex));
+					Converted.DatasmithActorElement->RemoveChild(Converted.DatasmithActorElement->GetChild(ChildIndex));
 				}
-				ExportedScene.DatasmithSceneRef->RemoveActor(NodeTracker.DatasmithActorElement, EDatasmithActorRemovalRule::RemoveChildren);
+				ExportedScene.DatasmithSceneRef->RemoveActor(Converted.DatasmithActorElement, EDatasmithActorRemovalRule::RemoveChildren);
 			}
-			NodeTracker.DatasmithActorElement.Reset();
-			NodeTracker.Children.Reset();
+			Converted.DatasmithActorElement.Reset();
+
+			NodeTracker.ReleaseConverted();
 		}
+		
 	}
 
 	// Called when mesh element is not needed anymore and should be removed from the scene
-	virtual void ReleaseMeshElement(TSharedPtr<IDatasmithMeshElement> Mesh) override
+	virtual void ReleaseMeshElement(FMeshConverted& Converted) override
 	{
-		ExportedScene.GetDatasmithScene()->RemoveMesh(Mesh);
+		GetDatasmithScene().RemoveMesh(Converted.DatasmithMeshElement);
+		Converted.ReleaseMeshConverted();
 	}
 
-	void RemoveFromConverted(FNodeTracker& NodeTracker)
+	void ReleaseNodeTrackerFromLayer(FNodeTracker& NodeTracker)
 	{
-		// todo: record previous converter Node type to speed up cleanup. Or just add 'unconverted' flag to speed up this for nodes that weren't converted yet
-
 		if (NodeTracker.Layer)
 		{
 			if (TSet<FNodeTracker*>* NodeTrackerPtr = NodesPerLayer.Find(NodeTracker.Layer))
@@ -1256,86 +1622,58 @@ public:
 			}
 			NodeTracker.Layer = nullptr;
 		}
+	}
 
+	void ReleaseNodeTrackerFromDatasmithMetadata(FNodeTracker& NodeTracker)
+	{
 		TSharedPtr<IDatasmithMetaDataElement> DatasmithMetadata;
 		if (NodeDatasmithMetadata.RemoveAndCopyValue(&NodeTracker, DatasmithMetadata))
 		{
-			ExportedScene.GetDatasmithScene()->RemoveMetaData(DatasmithMetadata);
+			GetDatasmithScene().RemoveMetaData(DatasmithMetadata);
 		}
+	}
 
-		Helpers.Remove(&NodeTracker);
-		Cameras.Remove(&NodeTracker);
-		Lights.Remove(&NodeTracker);
-
-		{
-			// remove static meshes used by the RailClone
-			TUniquePtr<FRailClonesConverted> RailClonesConverted;
-			if (RailClones.RemoveAndCopyValue(&NodeTracker, RailClonesConverted))
-			{
-				UnregisterNodeForMaterial(NodeTracker);
+	// Release node from any connection to other tracked objects
+	// When node is removed or to clear converter that is configured for invalidated node in Update
+	// BUT excluding node name and collision status - these two are updated before UpdateNode is called
+	void RemoveFromTracked(FNodeTracker& NodeTracker)
+	{
+		// todo: record previous converter Node type to speed up cleanup. Or just add 'unconverted' flag to speed up this for nodes that weren't converted yet
 
-				for (TSharedPtr<IDatasmithMeshElement> Mesh : RailClonesConverted->Meshes)
-				{
-					ReleaseMeshElement(Mesh);
-				}
-			}
-		}
+		ReleaseNodeTrackerFromLayer(NodeTracker);
 
-		if (NodeTracker.IsInstance())
+		if (NodeTracker.HasConverter())
 		{
-			if (TUniquePtr<FInstances>* InstancesPtr = InstancesForAnimHandle.Find(NodeTracker.InstanceHandle))
-			{
-				FInstances& Instances = **InstancesPtr;
-				UnregisterNodeForMaterial(NodeTracker);
-
-				Instances.NodeTrackers.Remove(&NodeTracker);
-
-				if (!Instances.NodeTrackers.Num())
-				{
-					ReleaseMeshElement(Instances.DatasmithMeshElement);
-					InstancesForAnimHandle.Remove(NodeTracker.InstanceHandle);
-					InvalidatedInstances.Remove(&Instances);
-				}
-				else
-				{
-					InvalidateInstances(Instances); // Invalidate instances that had a node removed - need to rebuild for various reasons(mesh might have been built from removed node, material assignment needds rebuild))
-				}
-			}
+			NodeTracker.GetConverter().RemoveFromTracked(*this, NodeTracker);
+			NodeTracker.ReleaseConverter();
 		}
-
-		ClearNodeFromDatasmithScene(NodeTracker);
 	}
 
-	void UpdateCollisionStatus(FNodeTracker* NodeTracker, TSet<FNodeTracker*>& NodesWithChangedCollisionStatus)
+	void UpdateCollisionStatus(FNodeTracker& NodeTracker, TSet<FNodeTracker*>& NodesWithChangedCollisionStatus)
 	{
 		// Check if collision assigned for node changed
 		{
-			TOptional<FDatasmithMaxStaticMeshAttributes> DatasmithAttributes = FDatasmithMaxStaticMeshAttributes::ExtractStaticMeshAttributes(NodeTracker->Node);
+			TOptional<FDatasmithMaxStaticMeshAttributes> DatasmithAttributes = FDatasmithMaxStaticMeshAttributes::ExtractStaticMeshAttributes(NodeTracker.Node);
 
 			bool bOutFromDatasmithAttributes;
-			INode* CollisionNode = FDatasmithMaxMeshExporter::GetCollisionNode(NodeTracker->Node, DatasmithAttributes ? &DatasmithAttributes.GetValue() : nullptr, bOutFromDatasmithAttributes);
+			INode* CollisionNode = GeomUtils::GetCollisionNode(*this, NodeTracker.Node, DatasmithAttributes ? &DatasmithAttributes.GetValue() : nullptr, bOutFromDatasmithAttributes);
 
-			FNodeTracker* CollisionNodeTracker = nullptr;
-			FNodeKey CollisionNodeKey = NodeEventNamespace::GetKeyByNode(CollisionNode);
-			if (FNodeTrackerHandle* NodeTrackerHandle = NodeTrackers.Find(CollisionNodeKey)) // This node should be tracked
-			{
-				CollisionNodeTracker = NodeTrackerHandle->GetNodeTracker();
-			}
+			FNodeTracker* CollisionNodeTracker = GetNodeTracker(CollisionNode); // This node should be tracked
 
-			if (NodeTracker->Collision != CollisionNodeTracker)
+			if (NodeTracker.Collision != CollisionNodeTracker)
 			{
 				// Update usage counters for collision nodes
 
 				// Remove previous
-				if (TSet<FNodeTracker*>* CollisionUsersPtr = CollisionNodes.Find(NodeTracker->Collision))
+				if (TSet<FNodeTracker*>* CollisionUsersPtr = CollisionNodes.Find(NodeTracker.Collision))
 				{
 					TSet<FNodeTracker*>& CollisionUsers = *CollisionUsersPtr;
-					CollisionUsers.Remove(NodeTracker);
+					CollisionUsers.Remove(&NodeTracker);
 
 					if (CollisionUsers.IsEmpty())
 					{
-						CollisionNodes.Remove(NodeTracker->Collision);
-						NodesWithChangedCollisionStatus.Add(NodeTracker->Collision);
+						CollisionNodes.Remove(NodeTracker.Collision);
+						NodesWithChangedCollisionStatus.Add(NodeTracker.Collision);
 					}
 				}
 
@@ -1345,50 +1683,47 @@ public:
 					if (TSet<FNodeTracker*>* CollisionUsersPtr = CollisionNodes.Find(CollisionNodeTracker))
 					{
 						TSet<FNodeTracker*>& CollisionUsers = *CollisionUsersPtr;
-						CollisionUsers.Add(NodeTracker);
+						CollisionUsers.Add(&NodeTracker);
 					}
 					else
 					{
 						TSet<FNodeTracker*>& CollisionUsers = CollisionNodes.Add(CollisionNodeTracker);
-						CollisionUsers.Add(NodeTracker);
+						CollisionUsers.Add(&NodeTracker);
 						NodesWithChangedCollisionStatus.Add(CollisionNodeTracker);
 					}
 				}
-				NodeTracker->Collision = CollisionNodeTracker;
+				NodeTracker.Collision = CollisionNodeTracker;
 			}
 		}
 
 		// Check if node changed its being assigned as collision
 		{
-			if (FDatasmithMaxSceneParser::HasCollisionName(NodeTracker->Node))
+			if (FDatasmithMaxSceneParser::HasCollisionName(NodeTracker.Node))
 			{
-				CollisionNodes.Add(NodeTracker); // Always view node with 'collision' name as a collision node(i.e. no render)
+				CollisionNodes.Add(&NodeTracker); // Always view node with 'collision' name as a collision node(i.e. no render)
 
 				//Check named collision assignment(e.g. 'UCP_<other nothe name>')
 				// Split collision prefix and find node that might use this node as collision mesh
-				FString NodeName = NodeTracker->Node->GetName();
+				FString NodeName = NodeTrackersNames.GetNodeName(NodeTracker);
 				FString LeftString, RightString;
 				NodeName.Split(TEXT("_"), &LeftString, &RightString);
 
-				if (TSet<FNodeTracker*>* CollisionUserNodeTrackersPtr = NodeTrackersNames.Find(RightString))
+				NodeTrackersNames.EnumerateForName(RightString, [&](FNodeTracker& CollisionUserNodeTracker)
 				{
-					for (FNodeTracker* CollisionUserNodeTracker: *CollisionUserNodeTrackersPtr)
+					if (CollisionUserNodeTracker.Collision != &NodeTracker)
 					{
-						if (CollisionUserNodeTracker->Collision != NodeTracker)
-						{
-							NodesWithChangedCollisionStatus.Add(CollisionUserNodeTracker); // Invalidate each node that has collision changed
-						}
+						NodesWithChangedCollisionStatus.Add(&CollisionUserNodeTracker); // Invalidate each node that has collision changed
 					}
-				}
+				});
 			}
 			else
 			{
 				// Remove from registered collision nodes if there's not other users(i.e. using Datasmith attributes reference)
-				if (TSet<FNodeTracker*>* CollisionUsersPtr = CollisionNodes.Find(NodeTracker))
+				if (TSet<FNodeTracker*>* CollisionUsersPtr = CollisionNodes.Find(&NodeTracker))
 				{
 					if (CollisionUsersPtr->IsEmpty())
 					{
-						CollisionNodes.Remove(NodeTracker);
+						CollisionNodes.Remove(&NodeTracker);
 					}
 				}
 			}
@@ -1397,14 +1732,14 @@ public:
 
 	void RemoveNodeTracker(FNodeTracker& NodeTracker)
 	{
-		InvalidatedNodeTrackers.Remove(&NodeTracker);
+		SCENE_UPDATE_STAT_INC(RemoveDeletedNodes, Nodes);
 
-		RemoveFromConverted(NodeTracker);
+		InvalidatedNodeTrackers.RemoveFromInvalidated(NodeTracker);
 
-		if (TSet<FNodeTracker*>* NodeTrackersPtr = NodeTrackersNames.Find(NodeTracker.Name))
-		{
-			NodeTrackersPtr->Remove(&NodeTracker);
-		}
+		ClearNodeFromDatasmithScene(NodeTracker);
+		RemoveFromTracked(NodeTracker);
+
+		NodeTrackersNames.Remove(NodeTracker);
 
 		if (TSet<FNodeTracker*>* CollisionUsersPtr = CollisionNodes.Find(NodeTracker.Collision))
 		{
@@ -1424,7 +1759,11 @@ public:
 	{
 		SCENE_UPDATE_STAT_INC(UpdateNode, NodesUpdated);
 		// Forget anything that this node was before update: place in datasmith hierarchy, datasmith objects, instances connection. Updating may change anything
-		RemoveFromConverted(NodeTracker);
+		ClearNodeFromDatasmithScene(NodeTracker);
+		RemoveFromTracked(NodeTracker); // todo: might keep it tracked by complicating conversions later(e.g. to avoid extra work if object stays the same)
+
+		NodeTracker.ResetValidityInterval(); // Make infinite validity to narrow it during update
+
 		ConvertNodeObject(NodeTracker);
 	}
 
@@ -1460,7 +1799,7 @@ public:
 			return;
 		}
 
-		ObjectState ObjState = NodeTracker.Node->EvalWorldState(0);
+		ObjectState ObjState = NodeTracker.Node->EvalWorldState(CurrentSyncPoint.Time);
 		Object* Obj = ObjState.obj;
 
 		if (!Obj)
@@ -1468,17 +1807,28 @@ public:
 			return;
 		}
 
-		switch (Obj->SuperClassID())
+		SClass_ID SuperClassID = Obj->SuperClassID();
+		switch (SuperClassID)
 		{
 		case HELPER_CLASS_ID:
-			ConvertHelper(NodeTracker, Obj);
+			SCENE_UPDATE_STAT_INC(UpdateNode, HelpersEncontered);
+			NodeTracker.CreateConverter<FHelperNodeConverter>();
 			break;
 		case CAMERA_CLASS_ID:
-			ConvertCamera(NodeTracker, Obj);
+			SCENE_UPDATE_STAT_INC(UpdateNode, CamerasEncontered);
+			NodeTracker.CreateConverter<FCameraNodeConverter>();
 			break;
 		case LIGHT_CLASS_ID:
 		{
-			ConvertLight(NodeTracker, Obj);
+			SCENE_UPDATE_STAT_INC(UpdateNode, LightsEncontered);
+
+			if (EMaxLightClass::Unknown == FDatasmithMaxSceneParser::GetLightClass(NodeTracker.Node))
+			{
+				SCENE_UPDATE_STAT_INC(UpdateNode, LightsSkippedAsUnknown);
+				break;
+			}
+
+			NodeTracker.CreateConverter<FLightNodeConverter>();
 			break;
 		}
 		case SHAPE_CLASS_ID:
@@ -1488,16 +1838,16 @@ public:
 			Class_ID ClassID = ObjState.obj->ClassID();
 			if (ClassID.PartA() == TARGET_CLASS_ID) // Convert camera target as regular actor
 			{
-				ConvertHelper(NodeTracker, Obj);
+				NodeTracker.CreateConverter<FHelperNodeConverter>();
 			}
 			else if (ClassID == RAILCLONE_CLASS_ID)
 			{
-				ConvertRailClone(*this, NodeTracker, Obj);
+				NodeTracker.CreateConverter<FRailCloneNodeConverter>();
 				break;
 			}
 			else if (ClassID == ITOOFOREST_CLASS_ID)
 			{
-				ConvertForest(*this, NodeTracker, Obj);
+				NodeTracker.CreateConverter<FForestNodeConverter>();
 				break;
 			}
 			else
@@ -1508,14 +1858,41 @@ public:
 				}
 				else
 				{
-					ConvertGeomObj(NodeTracker, Obj);
+					if (Obj->IsRenderable()) // Shape's Enable In Render flag(note - different from Node's Renderable flag)
+					{
+						SCENE_UPDATE_STAT_INC(UpdateNode, GeomObjConverted);
+
+						NodeTracker.CreateConverter<FMeshNodeConverter>();
+					}
+					else
+					{
+						SCENE_UPDATE_STAT_INC(UpdateNode, GeomObjSkippedAsNonRenderable);
+					}
 				}
 			}
 			break;
 		}
-		// todo: other object types besides geometry
+		case SYSTEM_CLASS_ID:
+		{
+			//When a referenced file is not found XRefObj is not resolved then it's kept as XREFOBJ_CLASS_ID instead of resolved class that it references
+			if (Obj->ClassID() == XREFOBJ_CLASS_ID)
+			{
+				FString Path = FDatasmithMaxSceneExporter::GetActualPath(static_cast<IXRefObject8*>(Obj)->GetFile(FALSE).GetFileName());
+				if (!FPaths::FileExists(Path))
+				{
+					LogWarning(FString("XRefObj file \"") + FPaths::GetCleanFilename(*Path) + FString("\" cannot be found"));
+				}
+			}
+
+			break;
+		}
 		default:;
 		}
+
+		if (NodeTracker.HasConverter())
+		{
+			NodeTracker.GetConverter().Parse(*this, NodeTracker);
+		}
 	}
 
 	void InvalidateInstances(FInstances& Instances)
@@ -1525,54 +1902,61 @@ public:
 
 	void UpdateInstances(FInstances& Instances)
 	{
-		 if (!Instances.NodeTrackers.IsEmpty())
-		 {
-			 bool bGeometryUpdated = false; // Use first node to extract information from evaluated object(e.g. GetRenderMesh needs it)
+		if (Instances.NodeTrackers.IsEmpty())
+		{
+			// Invalidated instances without actual instances left(all removed)
 
-			 bool bMaterialsAssignToStaticMesh = true; // assign materials to static mesh for the first instance(others will use override on mesh actors)
-			 for(FNodeTracker* NodeTrackerPtr: Instances.NodeTrackers)
-			 {
-				FNodeTracker& NodeTracker = *NodeTrackerPtr;
-				ClearNodeFromDatasmithScene(NodeTracker);
+			// NOTE: release mesh element and release usage of Instances pointer
+			//   BEFORE deallocating Instances instance from the map below
+			ReleaseMeshElement(Instances.Converted);
 
-				if (!bGeometryUpdated)
-				{
-					// todo: use single EnumProc instance to enumerate all nodes during update to:
-					//    - have single call to BeginEnumeration and EndEnumeration
-					//    - track all Begin'd nodes to End them together after all is updated(to prevent duplicated Begin's of referenced objects that might be shared by different ndoes)
-					NodesPreparer.PrepareNode(NodeTracker.Node);
-					UpdateInstancesGeometry(Instances, NodeTracker);
+			InstancesManager.RemoveInstances(Instances);
+			return;
+		}
 
-					bGeometryUpdated = true;
-				}
+		for (FNodeTracker* NodeTrackerPtr : Instances.NodeTrackers)
+		{
+			FNodeTracker& NodeTracker = *NodeTrackerPtr;
+			ClearNodeFromDatasmithScene(NodeTracker);
+			if (ensure(NodeTracker.GetConverter().ConverterType == FNodeConverter::MeshNode))
+			{
+				static_cast<FMeshNodeConverter&>(NodeTracker.GetConverter()).bMaterialsAssignedToStaticMesh = false;
+			}
+		}
 
-				UpdateGeometryNode(NodeTracker, Instances, bMaterialsAssignToStaticMesh);
-				bMaterialsAssignToStaticMesh = false;
+		// Export static mesh using first lucky node
+		// todo: possible optimization to reuse previous node. Somehow. 
+		for (FNodeTracker* NodeTrackerPtr : Instances.NodeTrackers)
+		{
+			FNodeTracker& NodeTracker = *NodeTrackerPtr;
 
-			 	// Mark node as updated as soon as it is - in order for next nodes to be able to use its DatasmithActor
-				NodeTracker.bInvalidated = false;
-			 }
-		 }
+			// todo: use single EnumProc instance to enumerate all nodes during update to:
+			//    - have single call to BeginEnumeration and EndEnumeration
+			//    - track all Begin'd nodes to End them together after all is updated(to prevent duplicated Begin's of referenced objects that might be shared by different noвes)
+			NodesPreparer.PrepareNode(NodeTracker.Node);
+			UpdateInstancesGeometry(Instances, NodeTracker);
+
+			// assign materials to static mesh for the first instance(others will use override on mesh actors)
+			static_cast<FMeshNodeConverter&>(NodeTracker.GetConverter()).bMaterialsAssignedToStaticMesh = true;
+			if (Mtl* Material = UpdateGeometryNodeMaterial(*this, Instances, NodeTracker))
+			{
+				Instances.AssignMaterialToStaticMesh(Material);
+			}
+			break;
+		}
 	}
 
 	void UpdateNodeMetadata(FNodeTracker& NodeTracker)
 	{
-		TSharedPtr<IDatasmithMetaDataElement> MetadataElement = FDatasmithMaxSceneExporter::ParseUserProperties(NodeTracker.Node, NodeTracker.DatasmithActorElement.ToSharedRef(), ExportedScene.GetDatasmithScene());
+		TSharedPtr<IDatasmithMetaDataElement> MetadataElement = FDatasmithMaxSceneExporter::ParseUserProperties(NodeTracker.Node, NodeTracker.GetConverted().DatasmithActorElement.ToSharedRef(), ExportedScene.GetDatasmithScene().ToSharedRef());
 		NodeDatasmithMetadata.Add(&NodeTracker, MetadataElement);
 	}
 
-	// Find first ansestor node which has DatasmithActor created for it
+	// Get parent node, transparently resolving XRefScene binding
 	FNodeTracker* GetParentNodeTracker(FNodeTracker& NodeTracker)
 	{
-		INode* XRefParent = NodeTracker.GetXRefParent();
-		FNodeKey ParentNodeKey = NodeEventNamespace::GetKeyByNode(XRefParent ? XRefParent : NodeTracker.Node->GetParentNode());
-		FNodeTrackerHandle* ParentNodeTrackerHandle = NodeTrackers.Find(ParentNodeKey);
-		if (!ParentNodeTrackerHandle)
-		{
-			return nullptr; // Max node not known, assume it doesn't exist
-		}
-
-		return ParentNodeTrackerHandle->GetNodeTracker();
+		INode* XRefParent = NodeTracker.GetXRefParent(); // Node may be at the root of an XRefScene, in this case get the scene node that this XRef if bound to(by XRef UI 'bind' or being a Container node)
+		return GetNodeTracker(XRefParent ? XRefParent : NodeTracker.Node->GetParentNode());
 	}
 
 	// Not all nodes result in creation of DatasmithActor for them(e.g. skipped as invisible), find first ascestor that has it
@@ -1581,7 +1965,7 @@ public:
 		FNodeTracker* NodeTracker = &InNodeTracker;
 		while (FNodeTracker* ParentNodeTracker = GetParentNodeTracker(*NodeTracker))
 		{
-			if (ParentNodeTracker->DatasmithActorElement)
+			if (ParentNodeTracker->HasConverted())
 			{
 				return ParentNodeTracker;
 			}
@@ -1592,7 +1976,7 @@ public:
 
 	bool AttachNodeToDatasmithScene(FNodeTracker& NodeTracker)
 	{
-		if (!NodeTracker.DatasmithActorElement)
+		if (!NodeTracker.HasConverted())
 		{
 			SCENE_UPDATE_STAT_INC(ReparentActors, SkippedWithoutDatasmithActor);
 
@@ -1602,15 +1986,14 @@ public:
 
 		if (FNodeTracker* ParentNodeTracker = GetAncestorNodeTrackerWithDatasmithActor(NodeTracker))
 		{
-			ParentNodeTracker->DatasmithActorElement->AddChild(NodeTracker.DatasmithActorElement, EDatasmithActorAttachmentRule::KeepWorldTransform);
-			ParentNodeTracker->Children.Add(&NodeTracker);
+			ParentNodeTracker->GetConverted().DatasmithActorElement->AddChild(NodeTracker.GetConverted().DatasmithActorElement, EDatasmithActorAttachmentRule::KeepWorldTransform);
 		}
 		else
 		{
 			SCENE_UPDATE_STAT_INC(ReparentActors, AttachedToRoot);
 			// If there's no ancestor node with DatasmithActor assume node it at root
 			// (node's parent might be node that was skipped - e.g. it was hidden in Max or not selected when exporting only selected objects)
-			ExportedScene.GetDatasmithScene()->AddActor(NodeTracker.DatasmithActorElement);
+			GetDatasmithScene().AddActor(NodeTracker.GetConverted().DatasmithActorElement);
 		}
 		return true;
 	}
@@ -1622,16 +2005,21 @@ public:
 
 		const FMaxLightCoordinateConversionParams LightParams = FMaxLightCoordinateConversionParams(NodeTracker.Node);
 		// todo: do we really need to call GetObjectTM if there's no WSM attached? Maybe just call GetObjTMAfterWSM always?
+		Interval ValidityInterval;
+		ValidityInterval.SetInfinite();
 		if (NodeTracker.Node->GetWSMDerivedObject() != nullptr)
 		{
-			FDatasmithMaxSceneExporter::MaxToUnrealCoordinates(NodeTracker.Node->GetObjTMAfterWSM(GetCOREInterface()->GetTime()), Translation, Rotation, Scale, Converter.UnitToCentimeter, LightParams);
+			FDatasmithMaxSceneExporter::MaxToUnrealCoordinates(NodeTracker.Node->GetObjTMAfterWSM(CurrentSyncPoint.Time, &ValidityInterval), Translation, Rotation, Scale, Converter.UnitToCentimeter, LightParams);
 		}
 		else
 		{
-			FDatasmithMaxSceneExporter::MaxToUnrealCoordinates(NodeTracker.Node->GetObjectTM(GetCOREInterface()->GetTime()), Translation, Rotation, Scale, Converter.UnitToCentimeter, LightParams);
+			FDatasmithMaxSceneExporter::MaxToUnrealCoordinates(NodeTracker.Node->GetObjectTM(CurrentSyncPoint.Time, &ValidityInterval), Translation, Rotation, Scale, Converter.UnitToCentimeter, LightParams);
 		}
+		LogDebug(FString::Printf(TEXT("Validity: (%ld, %ld)"), ValidityInterval.Start(), ValidityInterval.End()));
 		Rotation.Normalize();
 		ObjectTransform = FTransform(Rotation, Translation, Scale);
+
+		NodeTracker.NarrowValidityToInterval(ValidityInterval);
 	}
 
 	void RegisterNodeForMaterial(FNodeTracker& NodeTracker, Mtl* Material)
@@ -1641,7 +2029,7 @@ public:
 		MaterialsAssignedToNodes.FindOrAdd(MaterialTracker).Add(&NodeTracker);
 	}
 
-	void UnregisterNodeForMaterial(FNodeTracker& NodeTracker)
+	virtual void UnregisterNodeForMaterial(FNodeTracker& NodeTracker) override 
 	{
 		for (FMaterialTracker* MaterialTracker: NodeTracker.MaterialTrackers)
 		{
@@ -1649,13 +2037,65 @@ public:
 			if (MaterialsAssignedToNodes[MaterialTracker].IsEmpty())
 			{
 				MaterialsCollectionTracker.ReleaseMaterial(*MaterialTracker);
+				MaterialsAssignedToNodes.Remove(MaterialTracker);
 			}
 		}
 		NodeTracker.MaterialTrackers.Reset();
 	}
 
-	void UpdateGeometryNode(FNodeTracker& NodeTracker, FInstances& Instances, bool bMaterialsAssignToStaticMesh)
+	static Mtl* UpdateGeometryNodeMaterial(FSceneTracker& SceneTracker, FInstances& Instances, FNodeTracker& NodeTracker)
+	{
+		if (Instances.HasMesh())
+		{
+			if (Mtl* Material = NodeTracker.Node->GetMtl())
+			{
+				bool bMaterialRegistered = false;
+				for(FMaterialTracker* MaterialTracker: NodeTracker.MaterialTrackers)
+				{
+					bMaterialRegistered = bMaterialRegistered || (MaterialTracker->Material == Material);
+				}
+
+				if (!bMaterialRegistered) // Different material
+				{
+					// Release old material
+					SceneTracker.UnregisterNodeForMaterial(NodeTracker);
+					// Record new connection
+					SceneTracker.RegisterNodeForMaterial(NodeTracker, Material);
+				}
+				return Material;
+			}
+
+			// Release old material when node has no material now
+			SceneTracker.UnregisterNodeForMaterial(NodeTracker);
+		}
+		return nullptr;
+	}
+
+	virtual void AddGeometryNodeInstance(FNodeTracker& NodeTracker, FMeshNodeConverter& MeshConverter, Object* Obj) override
 	{
+		InvalidateInstances(InstancesManager.AddNodeTracker(NodeTracker, MeshConverter, Obj));
+		
+	}
+
+	virtual void RemoveGeometryNodeInstance(FNodeTracker& NodeTracker) override
+	{
+		if (FInstances* Instances = InstancesManager.RemoveNodeTracker(NodeTracker))
+		{
+			// Invalidate instances that had a node removed
+			//  - need to rebuild for various reasons(mesh might have been built from removed node, material assignment needds rebuild), remove empty
+			InvalidateInstances(*Instances);
+		}
+	}
+
+	virtual void ConvertGeometryNodeToDatasmith(FNodeTracker& NodeTracker, FMeshNodeConverter& MeshConverter) override
+	{
+		FInstances* InstancesPtr = InstancesManager.GetInstancesForNodeTracker(NodeTracker);
+		if (!InstancesPtr)
+		{
+			return;
+		}
+		FInstances& Instances = *InstancesPtr;
+
 		FDatasmithConverter Converter;
 
 		FTransform ObjectTransform;
@@ -1664,18 +2104,16 @@ public:
 		FTransform Pivot = FDatasmithMaxSceneExporter::GetPivotTransform(NodeTracker.Node, Converter.UnitToCentimeter);
 		FTransform NodeTransform = Pivot.Inverse() * ObjectTransform; // Remove pivot from the node actor transform
 
-		const TSharedPtr<IDatasmithMeshElement>& DatasmithMeshElement = Instances.DatasmithMeshElement;
-		bool bHasMesh = DatasmithMeshElement.IsValid();
 		bool bNeedPivotComponent = !Pivot.Equals(FTransform::Identity);
 
 		TSharedPtr<IDatasmithActorElement> DatasmithActorElement;
 		TSharedPtr<IDatasmithMeshActorElement> DatasmithMeshActor;
 
 		FString UniqueName = FString::FromInt(NodeTracker.Node->GetHandle());
-		FString Label = FString((const TCHAR*)NodeTracker.Node->GetName());
+		FString Label = NodeTrackersNames.GetNodeName(NodeTracker);
 
 		// Create and setup mesh actor if there's a mesh
-		if (bHasMesh)
+		if (Instances.HasMesh())
 		{
 			FString MeshActorName = UniqueName;
 			if (bNeedPivotComponent)
@@ -1683,7 +2121,7 @@ public:
 				MeshActorName += TEXT("_Pivot");
 			}
 
-			FString MeshActorLabel = FString((const TCHAR*)NodeTracker.Node->GetName());
+			FString MeshActorLabel = NodeTrackersNames.GetNodeName(NodeTracker);
 			DatasmithMeshActor = FDatasmithSceneFactory::CreateMeshActor(*MeshActorName);
 			DatasmithMeshActor->SetLabel(*Label);
 
@@ -1693,11 +2131,12 @@ public:
 				DatasmithMeshActor->AddTag(TEXT("Datasmith.Attributes.Geometry: BoundingBox"));
 			}
 
-			DatasmithMeshActor->SetStaticMeshPathName(DatasmithMeshElement->GetName());
+			DatasmithMeshActor->SetStaticMeshPathName(Instances.GetStaticMeshPathName());
 		}
 
-		if (bNeedPivotComponent || !bHasMesh)
+		if (bNeedPivotComponent || !Instances.HasMesh())
 		{
+			// Create a separate actor in case pivot is different or there's not mesh(so no mesh actor)
 			DatasmithActorElement = FDatasmithSceneFactory::CreateActor(*UniqueName);
 			DatasmithActorElement->SetLabel(*Label);
 		}
@@ -1706,11 +2145,13 @@ public:
 			DatasmithActorElement = DatasmithMeshActor;
 		}
 
+		// Set transforms
 		DatasmithActorElement->SetTranslation(NodeTransform.GetTranslation());
 		DatasmithActorElement->SetScale(NodeTransform.GetScale3D());
 		DatasmithActorElement->SetRotation(NodeTransform.GetRotation());
 
-		if (bNeedPivotComponent && bHasMesh)
+		// Setup mesh actor transform with pivot's
+		if (bNeedPivotComponent && Instances.HasMesh()) 
 		{
 			DatasmithMeshActor->SetTranslation(Pivot.GetTranslation());
 			DatasmithMeshActor->SetRotation(Pivot.GetRotation());
@@ -1720,65 +2161,41 @@ public:
 			DatasmithActorElement->AddChild(DatasmithMeshActor, EDatasmithActorAttachmentRule::KeepRelativeTransform);
 		}
 
-		NodeTracker.DatasmithActorElement = DatasmithActorElement;
-		NodeTracker.DatasmithMeshActor = DatasmithMeshActor;
+		FNodeConverted& Converted = NodeTracker.CreateConverted();
+		Converted.DatasmithActorElement = DatasmithActorElement;
+		Converted.DatasmithMeshActor = DatasmithMeshActor;
 
 		UpdateNodeMetadata(NodeTracker);
 		TagsConverter.ConvertNodeTags(NodeTracker);
 		if (NodeTracker.Layer)
 		{
-			NodeTracker.DatasmithActorElement->SetLayer(*NodeTracker.Layer->Name);
+			Converted.DatasmithActorElement->SetLayer(*NodeTracker.Layer->Name);
 		}
 
 		// Apply material
-		if (DatasmithMeshElement)
+		if (Mtl* Material = UpdateGeometryNodeMaterial(*this, Instances, NodeTracker))
 		{
-			if (Mtl* Material = NodeTracker.Node->GetMtl())
+			// Assign materials
+			if (!MeshConverter.bMaterialsAssignedToStaticMesh)
 			{
-				bool bMaterialRegistered = false;
-				for(FMaterialTracker* MaterialTracker: NodeTracker.MaterialTrackers)
-				{
-					bMaterialRegistered = bMaterialRegistered || (MaterialTracker->Material == Material);
-				}
-
-				if (!bMaterialRegistered)
-				{
-					// Release old material
-					UnregisterNodeForMaterial(NodeTracker);
-					RegisterNodeForMaterial(NodeTracker, Material);
-				}
-
-				// Assign materials
-				if (bMaterialsAssignToStaticMesh)
+				if (Instances.Material != Material)
 				{
-					AssignMeshMaterials(Instances.DatasmithMeshElement, Material, Instances.SupportedChannels);
-					Instances.Material = Material;
+					TSharedRef<IDatasmithMeshActorElement> DatasmithMeshActorRef = NodeTracker.GetConverted().DatasmithMeshActor.ToSharedRef();
+					FDatasmithMaxSceneExporter::ParseMaterialForMeshActor(Material, DatasmithMeshActorRef, Instances.Converted.SupportedChannels, FVector3f(NodeTracker.GetConverted().DatasmithMeshActor->GetTranslation()));
 				}
-				else // Assign material overrides to meshactor
-				{
-					if (Instances.Material != Material)
-					{
-						TSharedRef<IDatasmithMeshActorElement> DatasmithMeshActorRef = NodeTracker.DatasmithMeshActor.ToSharedRef();
-						FDatasmithMaxSceneExporter::ParseMaterialForMeshActor(Material, DatasmithMeshActorRef, Instances.SupportedChannels, FVector3f(NodeTracker.DatasmithMeshActor->GetTranslation()));
-					}
-				}
-			}
-			else
-			{
-				// Release old material
-				UnregisterNodeForMaterial(NodeTracker);
-				NodeTracker.DatasmithMeshActor->ResetMaterialOverrides();
 			}
-
-			// todo: test mesh becoming empty/invalid/not created - what happens?
-			// todo: test multimaterial changes
-			// todo: check other material permutations
 		}
+
+	}
+
+	IDatasmithScene& GetDatasmithScene()
+	{
+		return *ExportedScene.GetDatasmithScene();
 	}
 
 	virtual void AddMeshElement(TSharedPtr<IDatasmithMeshElement>& DatasmithMeshElement, FDatasmithMesh& DatasmithMesh, FDatasmithMesh* CollisionMesh) override
 	{
-		ExportedScene.GetDatasmithScene()->AddMesh(DatasmithMeshElement);
+		GetDatasmithScene().AddMesh(DatasmithMeshElement);
 
 		// todo: parallelize this
 		FDatasmithMeshExporter DatasmithMeshExporter;
@@ -1788,19 +2205,20 @@ public:
 		}
 	}
 
-	bool UpdateInstancesGeometry(FInstances& Instances, FNodeTracker& NodeTracker)
+	void UpdateInstancesGeometry(FInstances& Instances, FNodeTracker& NodeTracker)
 	{
 		INode* Node = NodeTracker.Node;
 		Object* Obj = Instances.EvaluatedObj;
 
-		TSharedPtr<IDatasmithMeshElement>& DatasmithMeshElement = Instances.DatasmithMeshElement;
-		TSet<uint16>& SupportedChannels = Instances.SupportedChannels;
 		FString MeshName = FString::FromInt(Node->GetHandle());
 
-		FRenderMeshForConversion RenderMesh = GetMeshForGeomObject(Node, Obj);
-		FRenderMeshForConversion CollisionMesh = GetMeshForCollision(Node);
+		FMeshConverterSource MeshSource = {
+			Node, MeshName,
+			GeomUtils::GetMeshForGeomObject(CurrentSyncPoint.Time, Node, Obj), false,
+			GeomUtils::GetMeshForCollision(CurrentSyncPoint.Time, *this, Node),
+		};
 
-		if (RenderMesh.GetMesh())
+		if (MeshSource.RenderMesh.GetMesh())
 		{
 			bool bHasInstanceWithMultimat = false;
 			for (FNodeTracker* InstanceNodeTracker : Instances.NodeTrackers)
@@ -1813,26 +2231,42 @@ public:
 					}
 				}
 			}
-			if (ConvertMaxMeshToDatasmith(*this, DatasmithMeshElement, Node, *MeshName, RenderMesh, !bHasInstanceWithMultimat, SupportedChannels, CollisionMesh)) // export might produce anything(e.g. if mesh is empty)
+
+			MeshSource.bConsolidateMaterialIds = !bHasInstanceWithMultimat;
+
+			FMeshes::AddMesh(*this, MeshSource, Instances.Converted, [&](bool bHasConverted, FMeshConverted& MeshConverted)
 			{
-				DatasmithMeshElement->SetLabel(Node->GetName());
-				return true;
-			}
+				if (bHasConverted)
+				{
+					MeshConverted.DatasmithMeshElement->SetLabel(*NodeTrackersNames.GetNodeName(NodeTracker));
+				}
+				else
+				{
+					MeshConverted.DatasmithMeshElement.Reset();
+				}
+				Instances.Converted = MeshConverted;
+			});
+		}
+		else
+		{
+			// Whe RenderMesh can be null?
+			// seems like the only way is to static_cast<GeomObject*>(Obj) return null in GetMeshFromRenderMesh?
+			// Where Obj is return value of GetBaseObject(Node, CurrentTime)
+			// Or ObjectState.obj is null
+			ensure(false);
+			ReleaseMeshElement(Instances.Converted);
 		}
-
-		DatasmithMeshElement.Reset();
-		return false;
 	}
 
 	virtual void SetupActor(FNodeTracker& NodeTracker) override
 	{
-		NodeTracker.DatasmithActorElement->SetLabel(NodeTracker.Node->GetName());
+		NodeTracker.GetConverted().DatasmithActorElement->SetLabel(*NodeTrackersNames.GetNodeName(NodeTracker));
 
 		UpdateNodeMetadata(NodeTracker);
 		TagsConverter.ConvertNodeTags(NodeTracker);
 		if (NodeTracker.Layer)
 		{
-			NodeTracker.DatasmithActorElement->SetLayer(*NodeTracker.Layer->Name);
+			NodeTracker.GetConverted().DatasmithActorElement->SetLayer(*NodeTracker.Layer->Name);
 		}
 
 		FDatasmithConverter Converter;
@@ -1840,232 +2274,85 @@ public:
 		GetNodeObjectTransform(NodeTracker, Converter, ObjectTransform);
 
 		FTransform NodeTransform = ObjectTransform;
-		TSharedRef<IDatasmithActorElement> DatasmithActorElement = NodeTracker.DatasmithActorElement.ToSharedRef();
+		TSharedRef<IDatasmithActorElement> DatasmithActorElement = NodeTracker.GetConverted().DatasmithActorElement.ToSharedRef();
 		DatasmithActorElement->SetTranslation(NodeTransform.GetTranslation());
 		DatasmithActorElement->SetScale(NodeTransform.GetScale3D());
 		DatasmithActorElement->SetRotation(NodeTransform.GetRotation());
 	}
 
-	bool ConvertHelper(FNodeTracker& NodeTracker, Object* Obj)
-	{
-		SCENE_UPDATE_STAT_INC(UpdateNode, HelpersEncontered);
-		Helpers.Add(&NodeTracker);
-
-		if(!NodeTracker.DatasmithActorElement)
-		{
-			// note: this is how baseline exporter derives names
-			FString UniqueName = FString::FromInt(NodeTracker.Node->GetHandle());
-			NodeTracker.DatasmithActorElement = FDatasmithSceneFactory::CreateActor((const TCHAR*)*UniqueName);
-		}
-		SetupActor(NodeTracker);
-
-		NodeTracker.bInvalidated = false;
-
-		return true;
-	}
-
-	bool ConvertCamera(FNodeTracker& NodeTracker, Object* Obj)
+	class FMeshes
 	{
-		SCENE_UPDATE_STAT_INC(UpdateNode, CamerasEncontered);
-		Cameras.Add(&NodeTracker);
-
-		if(!NodeTracker.DatasmithActorElement)
-		{
-			// note: this is how baseline exporter derives names
-			FString UniqueName = FString::FromInt(NodeTracker.Node->GetHandle());
-			NodeTracker.DatasmithActorElement = FDatasmithSceneFactory::CreateCameraActor((const TCHAR*)*UniqueName);
-		}
-
-		FDatasmithMaxCameraExporter::ExportCamera(*NodeTracker.Node, StaticCastSharedPtr<IDatasmithCameraActorElement>(NodeTracker.DatasmithActorElement).ToSharedRef());
-
-		SetupActor(NodeTracker);
-
-		// Max camera view direction is Z-, Unreal's X+
-		// Max camera Up is Y+,  Unrela's Z+
-		FQuat Rotation = NodeTracker.DatasmithActorElement->GetRotation();
-		Rotation *= FQuat(0.0, 0.707107, 0.0, 0.707107);
-		Rotation *= FQuat(0.707107, 0.0, 0.0, 0.707107);
-		NodeTracker.DatasmithActorElement->SetRotation(Rotation);
-
-		NodeTracker.bInvalidated = false;
-
-		return true;
-	}
-
-	bool ConvertLight(FNodeTracker& NodeTracker, Object* Obj)
-	{
-		SCENE_UPDATE_STAT_INC(UpdateNode, LightsEncontered);
-
-		if (EMaxLightClass::Unknown == FDatasmithMaxSceneParser::GetLightClass(NodeTracker.Node))
-		{
-			SCENE_UPDATE_STAT_INC(UpdateNode, LightsSkippedAsUnknown);
-			return false;
-		}
-
-		Lights.Add(&NodeTracker);
-
-		TSharedPtr<IDatasmithLightActorElement> LightElement;
-		if(!NodeTracker.DatasmithActorElement)
-		{
-			// note: this is how baseline exporter derives names
-			FString UniqueName = FString::FromInt(NodeTracker.Node->GetHandle());
-
-			LightElement = FDatasmithMaxSceneExporter::CreateLightElementForNode(NodeTracker.Node, *UniqueName);
-
-			if (!LightElement)
-			{
-				if (FDatasmithMaxSceneParser::GetLightClass(NodeTracker.Node) == EMaxLightClass::SkyEquivalent)
-				{
-					ExportedScene.DatasmithSceneRef->SetUsePhysicalSky(true);
-				}
-				else
-				{
-					LogWarning(TEXT("Unsupported light: ") + DatasmithMaxLogger::Get().GetLightDescription(NodeTracker.Node));
-				}
-				return false;
-			}
-			else
-			{
-				if ( !FDatasmithMaxSceneExporter::ParseLight(NodeTracker.Node, LightElement.ToSharedRef(), ExportedScene.DatasmithSceneRef.ToSharedRef()) )
-				{
-					return false;
-				}
-			}
-
-			NodeTracker.DatasmithActorElement = LightElement;
-		}
-		SetupActor(NodeTracker);
-
-		//Cylinder shaped lights don't have the same default orientations, so we recalculate their transform and add the shape information.
-		if (LightElement->IsA(EDatasmithElementType::AreaLight)
-			&& StaticCastSharedPtr<IDatasmithAreaLightElement>(LightElement)->GetLightShape() == EDatasmithLightShape::Cylinder)
+	public:
+		static void AddMesh(FSceneTracker& Scene, FMeshConverterSource& MeshSource, FMeshConverted& MeshConverted, TFunction<void(bool, FMeshConverted&)> CompletionCallback)
 		{
-			FVector Translation, Scale;
-			FQuat Rotation;
-
-			const float UnitMultiplier = (float)GetSystemUnitScale(UNITS_CENTIMETERS);
-			const FMaxLightCoordinateConversionParams LightParams = FMaxLightCoordinateConversionParams(NodeTracker.Node, EDatasmithLightShape::Cylinder);
-			if (NodeTracker.Node->GetWSMDerivedObject() != nullptr)
-			{
-				FDatasmithMaxSceneExporter::MaxToUnrealCoordinates(NodeTracker.Node->GetObjTMAfterWSM(GetCOREInterface()->GetTime()), Translation, Rotation, Scale, UnitMultiplier, LightParams);
-			}
-			else
-			{
-				FDatasmithMaxSceneExporter::MaxToUnrealCoordinates(NodeTracker.Node->GetObjectTM(GetCOREInterface()->GetTime()), Translation, Rotation, Scale, UnitMultiplier, LightParams);
-			}
+			// Reset old mesh
+			// todo: potential mesh reuse - when DatasmithMeshElement allows to reset materials(as well as other params)
+			Scene.ReleaseMeshElement(MeshConverted);
 
-			Rotation.Normalize();
-			LightElement->SetTranslation(Translation);
-			LightElement->SetScale(Scale);
-			LightElement->SetRotation(Rotation);
+			bool bConverted = ConvertMaxMeshToDatasmith(Scene.CurrentSyncPoint.Time, Scene, MeshSource, MeshConverted);
+			CompletionCallback(bConverted, MeshConverted);
 		}
+	};
 
-
-		NodeTracker.bInvalidated = false;
-
-		return true;
-	}
-
-	virtual void SetupDatasmithHISMForNode(FNodeTracker& NodeTracker, INode* GeometryNode, const FRenderMeshForConversion& RenderMesh, Mtl* Material, int32 MeshIndex, const TArray<Matrix3>& Transforms) override
+	virtual void SetupDatasmithHISMForNode(FNodeTracker& NodeTracker, FMeshConverterSource& MeshSource, Mtl* Material, int32 MeshIndex, const TArray<Matrix3>& Transforms) override
 	{
 		FString MeshName = FString::FromInt(NodeTracker.Node->GetHandle()) + TEXT("_") + FString::FromInt(MeshIndex);
 
 		// note: when export Mesh goes to other place due to parallellizing it's result would be unknown here so MeshIndex handling will change(i.e. increment for any mesh)
 
-		TSharedPtr<IDatasmithMeshElement> DatasmithMeshElement;
-		TSet<uint16> SupportedChannels;
+		MeshSource.MeshName = MeshName; // todo: !!!
 
-		if (ConvertMaxMeshToDatasmith(*this, DatasmithMeshElement, GeometryNode, *MeshName, RenderMesh, false, SupportedChannels))
+		FMeshConverted MeshConvertedDummy; // todo: possible reuse of previous converted mesh
+		FMeshes::AddMesh(*this, MeshSource, MeshConvertedDummy, [&](bool bHasConverted, FMeshConverted& MeshConverted)
 		{
-			TUniquePtr<FRailClonesConverted>& RailClonesConverted = RailClones.FindOrAdd(&NodeTracker);
-			if (!RailClonesConverted)
+			if (bHasConverted)
 			{
-				RailClonesConverted = MakeUnique<FRailClonesConverted>();
-			}
-			RailClonesConverted->Meshes.Add(DatasmithMeshElement);
-
-			RegisterNodeForMaterial(NodeTracker, Material);
-			AssignMeshMaterials(DatasmithMeshElement, Material, SupportedChannels);
-
-			FString MeshLabel = NodeTracker.Node->GetName() + (TEXT("_") + FString::FromInt(MeshIndex));
-			DatasmithMeshElement->SetLabel(*MeshLabel);
-
-			FDatasmithConverter Converter;
-
-			// todo: override material
-			TSharedPtr< IDatasmithActorElement > InversedHISMActor;
-			// todo: ExportHierarchicalInstanceStaticMeshActor CustomMeshNode only used for Material - can be simplified, Material anyway is dealt with outside too
-			TSharedRef<IDatasmithActorElement> HismActorElement = FDatasmithMaxSceneExporter::ExportHierarchicalInstanceStaticMeshActor(
-				ExportedScene.GetDatasmithScene(), NodeTracker.Node, GeometryNode, *MeshLabel, SupportedChannels,
-				Material, &Transforms, *MeshName, Converter.UnitToCentimeter, EStaticMeshExportMode::Default, InversedHISMActor);
-			NodeTracker.DatasmithActorElement->AddChild(HismActorElement, EDatasmithActorAttachmentRule::KeepWorldTransform);
-			if (InversedHISMActor)
-			{
-				NodeTracker.DatasmithActorElement->AddChild(InversedHISMActor, EDatasmithActorAttachmentRule::KeepWorldTransform);
-			}
-			MeshIndex++;
-		}
-	}
-
-	bool ConvertGeomObj(FNodeTracker& NodeTracker, Object* Obj)
-	{
-		// todo: reuse mesh element(make sure to reset all)
-
-		bool bResult = false;
-		if (!Obj->IsRenderable()) // Shape's Enable In Render flag(note - different from Node's Renderable flag)
-		{
-			SCENE_UPDATE_STAT_INC(UpdateNode, GeomObjSkippedAsNonRenderable);
-			return bResult;
-		}
-
-		SCENE_UPDATE_STAT_INC(UpdateNode, GeomObjConverted);
-		// AnimHandle is unique and never reused for new objects
-		// todo: reset instances and nodes when one node of an instance changes??? Check how it should be done actually, dependencies - nodes, object, invalidation place(Update, Event) etc...
-		AnimHandle Handle = Animatable::GetHandleByAnim(Obj);
+				FHismNodeConverter& NodeConverter = static_cast<FHismNodeConverter&>(NodeTracker.GetConverter());
 
-		NodeTracker.InstanceHandle = Handle;
+				NodeConverter.Meshes.Add(MeshConverted);
 
-		TUniquePtr<FInstances>& Instances = InstancesForAnimHandle.FindOrAdd(Handle);
+				RegisterNodeForMaterial(NodeTracker, Material);
+				AssignMeshMaterials(MeshConverted.DatasmithMeshElement, Material, MeshConverted.SupportedChannels);
 
-		if (!Instances)
-		{
-			Instances = MakeUnique<FInstances>();
-			Instances->EvaluatedObj = Obj;
-		}
+				FString MeshLabel = NodeTrackersNames.GetNodeName(NodeTracker) + (TEXT("_") + FString::FromInt(MeshIndex));
+				MeshConverted.DatasmithMeshElement->SetLabel(*MeshLabel);
 
-		// need to invalidate mesh assignment to node that wasn't the first to add to instances(so if instances weren't invalidated - this node needs mesh anyway)
-		Instances->NodeTrackers.Add(&NodeTracker);
-		InvalidateInstances(*Instances);
+				FDatasmithConverter Converter;
 
-		return bResult;
+				// todo: override material
+				TSharedPtr< IDatasmithActorElement > InversedHISMActor;
+				// todo: ExportHierarchicalInstanceStaticMeshActor CustomMeshNode only used for Material - can be simplified, Material anyway is dealt with outside too
+				TSharedRef<IDatasmithActorElement> HismActorElement = FDatasmithMaxSceneExporter::ExportHierarchicalInstanceStaticMeshActor(NodeTracker.Node, MeshSource.Node, *MeshLabel, MeshConverted.SupportedChannels,
+					Material, &Transforms, *MeshName, Converter.UnitToCentimeter, EStaticMeshExportMode::Default, InversedHISMActor);
+				NodeTracker.GetConverted().DatasmithActorElement->AddChild(HismActorElement, EDatasmithActorAttachmentRule::KeepWorldTransform);
+				if (InversedHISMActor)
+				{
+					NodeTracker.GetConverted().DatasmithActorElement->AddChild(InversedHISMActor, EDatasmithActorAttachmentRule::KeepWorldTransform);
+				}
+				MeshIndex++;
+			}
+		});
 	}
 
 	void ConvertNamedCollisionNode(FNodeTracker& NodeTracker)
 	{
 		// Split collision prefix and find node that might use this node as collision mesh
-		FString NodeName = NodeTracker.Node->GetName();
+		FString NodeName = NodeTrackersNames.GetNodeName(NodeTracker);
 		FString LeftString, RightString;
 		NodeName.Split(TEXT("_"), &LeftString, &RightString);
 
-		INode* CollisionUserNode = GetCOREInterface()->GetINodeByName(*RightString);
-		if (!CollisionUserNode)
+		FNodeTracker* CollisionUserNodeTracker = GetNodeTrackerByNodeName(*RightString);
+		
+		if (!CollisionUserNodeTracker)
 		{
 			return;
 		}
 
-		// If some node is using this collision node then invalidate that node's instances
-		FNodeKey CollisionUserNodeKey = NodeEventNamespace::GetKeyByNode(CollisionUserNode);
-		if (FNodeTrackerHandle* NodeTrackerHandle = NodeTrackers.Find(CollisionUserNodeKey)) // This node should be tracked
+		if (CollisionUserNodeTracker->GetConverterType() == FNodeConverter::MeshNode)
 		{
-			FNodeTracker& CollisionUserNodeTracker = *NodeTrackerHandle->GetNodeTracker();
-
-			if (CollisionUserNodeTracker.IsInstance())
+			if (FInstances* Instances = InstancesManager.GetInstancesForNodeTracker(*CollisionUserNodeTracker))
 			{
-				if (TUniquePtr<FInstances>* InstancesPtr = InstancesForAnimHandle.Find(CollisionUserNodeTracker.InstanceHandle))
-				{
-					FInstances& Instances = **InstancesPtr;
-					InvalidateInstances(Instances);
-				}
+				InvalidateInstances(*Instances);
 			}
 		}
 	}
@@ -2074,6 +2361,7 @@ public:
 
 	virtual void NodeAdded(INode* Node) override
 	{
+		LogDebugNode(TEXT("NodeAdded"), Node);
 		// Node sometimes is null. 'Added' NodeEvent might come after node was actually deleted(immediately after creation)
 		// e.g.[mxs]: b = box(); delete b
 		// NodeEvents are delayed(not executed in the same stack frame as command that causes them) so they come later.
@@ -2097,6 +2385,7 @@ public:
 			return;
 		}
 
+		// Search where this XRef Tree is attached to the scene
 		int32 XRefIndex = -1; // Node that has this xref scene attached to(e.g. to place in hierarchy and to transform)
 		INode* SceneRootNode = GetCOREInterface()->GetRootNode();
 		for (int XRefChild = 0; XRefChild < SceneRootNode->GetXRefFileCount(); ++XRefChild)
@@ -2107,23 +2396,24 @@ public:
 			}
 		}
 
+		NotificationsHandler->AddNode(Node);
+
 		FNodeKey NodeKey = NodeEventNamespace::GetKeyByNode(Node);
 		InvalidateNode(NodeKey);
 
-		ParseScene(Node, FXRefScene{SceneRootNode, XRefIndex}); // Parse xref hierarchy - it won't add itself! Or will it?
+		if (!bIncludeXRefWhileParsing)
+		{
+			ParseScene(Node, FXRefScene{SceneRootNode, XRefIndex}); // Parse xref hierarchy - it won't add itself! Or will it?
+		}
 	}
 
 	virtual void NodeDeleted(INode* Node) override
 	{
 		LogDebugNode(TEXT("NodeDeleted"), Node);
-		// todo: check for null
-		FNodeKey NodeKey = NodeEventNamespace::GetKeyByNode(Node);
 
-		if (FNodeTrackerHandle* NodeTrackerHandle = NodeTrackers.Find(NodeKey))
+		if (FNodeTracker* NodeTracker = GetNodeTracker(Node))
 		{
-			// todo: schedule for delete on Update?
-			FNodeTracker* NodeTracker = NodeTrackerHandle->GetNodeTracker();
-			InvalidatedNodeTrackers.Add(NodeTracker);
+			InvalidatedNodeTrackers.Add(*NodeTracker);
 			NodeTracker->bDeleted = true;
 		}
 	}
@@ -2134,18 +2424,6 @@ public:
 
 		// todo: grouping makes this crash. Need to handle event before?
 		InvalidateNode(NodeKey);
-
-		// ControllerOtherEvent sent only for top actors in hierarchy when it's moved
-		INode* Node = NodeEventNamespace::GetNodeByKey(NodeKey);
-		if (Node)
-		{
-			int32 ChildNum = Node->NumberOfChildren();
-			for (int32 ChildIndex = 0; ChildIndex < ChildNum; ++ChildIndex)
-			{
-				// todo: pass INode to NodeTransformChanged to remove redundant get
-				NodeTransformChanged(NodeEventNamespace::GetKeyByNode(Node->GetChildNode(ChildIndex)));
-			}
-		}
 	}
 
 	virtual void NodeMaterialAssignmentChanged(FNodeKey NodeKey) override
@@ -2159,7 +2437,7 @@ public:
 		//identify material tree and update all materials
 		//todo: possible to handle this more precisely(only refresh changed materials) - see FMaterialObserver
 
-		if (FNodeTrackerHandle* NodeTracker = NodeTrackers.Find(NodeKey))
+		if (FNodeTracker* NodeTracker = GetNodeTracker(NodeKey))
 		{
 			// todo: investigate why NodeEventNamespace::GetNodeByKey may stil return NULL
 			// testcase - add XRef Material - this will immediately have this
@@ -2194,6 +2472,15 @@ public:
 		InvalidateNode(NodeKey);
 	}
 
+	virtual void NodeNameChanged(FNodeKey NodeKey) override
+	{
+		if (FNodeTracker* NodeTracker =  GetNodeTracker(NodeKey))
+		{
+			NodeTrackersNames.Update(*NodeTracker);
+			InvalidateNode(*NodeTracker);
+		}
+	}
+
 	virtual void NodePropertiesChanged(FNodeKey NodeKey) override
 	{
 		// todo: invalidate visibility only - note to handle this not enought add/remove
@@ -2212,48 +2499,73 @@ public:
 		return Stats;
 	}
 
+	virtual FNodeTracker* GetNodeTrackerByNodeName(const TCHAR* Name) override
+	{
+		FNodeTracker* Result = nullptr;
+		NodeTrackersNames.EnumerateForName(Name, [&](FNodeTracker& NodeTracker)
+		{
+			// todo: currently support only one name/node(as Max api GetINodeByName)
+			//   so collision with prefix only takes one node. But this can change.
+			Result = &NodeTracker;
+		});
+
+		return Result;
+	}
+	
+	void Reset()
+	{
+		bSceneParsed = false;
+
+		NodeTrackers.Reset();
+		NodeTrackersNames.Reset();
+		InstancesManager.Reset();
+		CollisionNodes.Reset();
+
+		LayersForAnimHandle.Reset();
+		NodesPerLayer.Reset();
+
+		MaterialsCollectionTracker.Reset();
+		MaterialsAssignedToNodes.Reset();
+
+		NodeDatasmithMetadata.Reset();
+
+		InvalidatedNodeTrackers.Reset();
+		InvalidatedInstances.Reset();
+	}
+
 	///////////////////////////////////////////////
 
 	const FExportOptions& Options;
 	FDatasmith3dsMaxScene& ExportedScene;
+
 	FNotifications* NotificationsHandler;
 
-	bool bSceneParsed = false;
 	bool bUpdateInProgress = false;
 
+	// Scene tracked/converted entities and connections beetween them
+	bool bSceneParsed = false;
 	TMap<FNodeKey, FNodeTrackerHandle> NodeTrackers; // All scene nodes
-	TMap<FString, TSet<FNodeTracker*>> NodeTrackersNames; // Nodes grouped by name
-	TSet<FNodeTracker*> InvalidatedNodeTrackers; // Nodes that need to be rebuilt
-	TMap<FNodeTracker*, TSharedPtr<IDatasmithMetaDataElement>> NodeDatasmithMetadata; // All scene nodes
-
+	FNodeTrackersNames NodeTrackersNames; // Nodes grouped by name, for faster access
+	FInstancesManager InstancesManager; // Groups geometry nodes by their shared mesh
 	TMap<FNodeTracker*, TSet<FNodeTracker*>> CollisionNodes; // Nodes used as collision meshes for other nodes, counted by each user
 
-	FMaterialsCollectionTracker MaterialsCollectionTracker;
-
-	TMap<FMaterialTracker*, TSet<FNodeTracker*>> MaterialsAssignedToNodes;
-
-	TMap<AnimHandle, TUniquePtr<FInstances>> InstancesForAnimHandle; // set of instanced nodes for each AnimHandle
-	TSet<FNodeTracker*> Helpers;
-	TSet<FNodeTracker*> Lights;
-	TSet<FNodeTracker*> Cameras;
-
 	TMap<AnimHandle, TUniquePtr<FLayerTracker>> LayersForAnimHandle;
 	TMap<FLayerTracker*, TSet<FNodeTracker*>> NodesPerLayer;
 
-	FNodesPreparer NodesPreparer;
-
-	struct FRailClonesConverted
-	{
-		TArray<TSharedPtr<IDatasmithMeshElement>> Meshes; // Meshes created for this railclones object
-	};
+	FMaterialsCollectionTracker MaterialsCollectionTracker;
+	TMap<FMaterialTracker*, TSet<FNodeTracker*>> MaterialsAssignedToNodes;
 
-	TMap<FNodeTracker*, TUniquePtr<FRailClonesConverted>> RailClones;
+	TMap<FNodeTracker*, TSharedPtr<IDatasmithMetaDataElement>> NodeDatasmithMetadata; // All scene nodes
 
+	// Nodes/instances that need to be rebuilt
+	FInvalidatedNodeTrackers InvalidatedNodeTrackers;
 	TSet<FInstances*> InvalidatedInstances;
 
+	// Utility
+	FSceneUpdateStats Stats;
 	FTagsConverter TagsConverter; // Converts max node information to Datasmith tags
+	FNodesPreparer NodesPreparer;
 
-	FSceneUpdateStats Stats;
 };
 
 class FExporter: public IExporter
@@ -2312,7 +2624,7 @@ public:
 		SetName(*SceneName);
 
 		DirectLinkImpl.Reset(new FDatasmithDirectLink);
-		DirectLinkImpl->InitializeForScene(ExportedScene.GetDatasmithScene());
+		DirectLinkImpl->InitializeForScene(ExportedScene.GetDatasmithScene().ToSharedRef());
 	}
 
 	virtual void UpdateDirectLinkScene() override
@@ -2325,7 +2637,7 @@ public:
 		}
 
 		LogDebug(TEXT("UpdateDirectLinkScene"));
-		DirectLinkImpl->UpdateScene(ExportedScene.GetDatasmithScene());
+		DirectLinkImpl->UpdateScene(ExportedScene.GetDatasmithScene().ToSharedRef());
 		StartSceneChangeTracking(); // Always track scene changes if it's synced with DirectLink
 	}
 
@@ -2533,6 +2845,25 @@ IPersistentExportOptions& GetPersistentExportOptions()
 	return PersistentExportOptions;
 }
 
+bool FInvalidatedNodeTrackers::PurgeDeletedNodeTrackers(FSceneTracker& Scene)
+{
+	TArray<FNodeTracker*> DeletedNodeTrackers;
+	for (FNodeTracker* NodeTracker : InvalidatedNodeTrackers)
+	{
+		if (NodeTracker->bDeleted)
+		{
+			DeletedNodeTrackers.Add(NodeTracker);
+		}
+	}
+
+	for (FNodeTracker* NodeTrackerPtr : DeletedNodeTrackers)
+	{
+		Scene.RemoveNodeTracker(*NodeTrackerPtr);
+	}
+
+	return !DeletedNodeTrackers.IsEmpty(); // If the only change is deleted node than we need to record it(deleted will be removed from InvalidatedNodeTrackers)
+}
+
 void FExporter::Shutdown()
 {
 	Exporter.Reset();
@@ -2564,7 +2895,7 @@ bool Export(const TCHAR* Name, const TCHAR* OutputPath, bool bQuiet)
 		PROGRESS_STAGE("Save Datasmith Scene");
 
 		IDatasmithScene& Scene = *ExportedScene.GetDatasmithScene();
-		ExportedScene.GetSceneExporter().Export(ExportedScene.GetDatasmithScene(), false);
+		ExportedScene.GetSceneExporter().Export(ExportedScene.GetDatasmithScene().ToSharedRef(), false);
 
 		PROGRESS_STAGE_RESULT(FString::Printf(TEXT("Actors: %d; Meshes: %d, Materials: %d"), 
 			Scene.GetActorsCount(),
diff --git a/Engine/Source/Programs/Enterprise/Datasmith/DatasmithMaxExporter/Private/DatasmithMaxSceneExporter.cpp b/Engine/Source/Programs/Enterprise/Datasmith/DatasmithMaxExporter/Private/DatasmithMaxSceneExporter.cpp
--- a/Engine/Source/Programs/Enterprise/Datasmith/DatasmithMaxExporter/Private/DatasmithMaxSceneExporter.cpp
+++ b/Engine/Source/Programs/Enterprise/Datasmith/DatasmithMaxExporter/Private/DatasmithMaxSceneExporter.cpp
@@ -336,7 +336,7 @@ void FDatasmithMaxSceneExporter::ExportMeshActor(TSharedRef< IDatasmithScene > D
 	DatasmithScene->AddActor( NodeActor );
 }
 
-TSharedRef< IDatasmithActorElement > FDatasmithMaxSceneExporter::ExportHierarchicalInstanceStaticMeshActor(TSharedRef< IDatasmithScene > DatasmithScene, INode* Node, INode* CustomMeshNode, const TCHAR* Label, TSet<uint16>& SupportedChannels, Mtl* StaticMeshMtl, const TArray<Matrix3>* Instances,
+TSharedRef< IDatasmithActorElement > FDatasmithMaxSceneExporter::ExportHierarchicalInstanceStaticMeshActor(INode* Node, INode* CustomMeshNode, const TCHAR* Label, TSet<uint16>& SupportedChannels, Mtl* StaticMeshMtl, const TArray<Matrix3>* Instances,
 	const TCHAR* MeshName, float UnitMultiplier, const EStaticMeshExportMode& ExportMode, TSharedPtr< IDatasmithActorElement >& OutInversedHISMActor)
 {
 	check(Node && Instances && MeshName);
diff --git a/Engine/Source/Programs/Enterprise/Datasmith/DatasmithMaxExporter/Public/DatasmithMaxConverters.h b/Engine/Source/Programs/Enterprise/Datasmith/DatasmithMaxExporter/Public/DatasmithMaxConverters.h
new file mode 100644
--- /dev/null
+++ b/Engine/Source/Programs/Enterprise/Datasmith/DatasmithMaxExporter/Public/DatasmithMaxConverters.h
@@ -0,0 +1,146 @@
+// Copyright Epic Games, Inc. All Rights Reserved.
+
+#pragma once
+
+#include "CoreMinimal.h"
+
+#include "DatasmithMaxExporterDefines.h"
+
+#include "Windows/AllowWindowsPlatformTypes.h"
+MAX_INCLUDES_START
+	#include "Animatable.h"
+MAX_INCLUDES_END
+#include "Windows/HideWindowsPlatformTypes.h"
+
+class IDatasmithActorElement;
+class IDatasmithMeshActorElement;
+class IDatasmithMeshElement;
+
+namespace DatasmithMaxDirectLink
+{
+
+class ISceneTracker;
+class FNodeTracker;
+
+// Identifies data for Max node converted to Datasmith
+class FNodeConverted: FNoncopyable
+{
+public:
+	TSharedPtr<IDatasmithActorElement> DatasmithActorElement;
+	TSharedPtr<IDatasmithMeshActorElement> DatasmithMeshActor;
+};
+
+class FMeshConverted
+{
+public:
+	TSharedPtr<IDatasmithMeshElement> DatasmithMeshElement;
+	TSet<uint16> SupportedChannels;
+	void ReleaseMeshConverted()
+	{
+		DatasmithMeshElement.Reset();
+		SupportedChannels.Reset();
+	}
+};
+
+class FNodeConverter: FNoncopyable
+{
+public:
+	enum EType
+	{
+		Unknown,
+		MeshNode,
+		HismNode,
+		LightNode,
+		CameraNode,
+		HelperNode,
+	};
+
+	FNodeConverter(EType InConverterType=EType::Unknown): ConverterType(InConverterType)
+	{
+	}
+
+	virtual ~FNodeConverter(){}
+
+	virtual void Parse(ISceneTracker& SceneTracker, FNodeTracker& NodeTracker) = 0;
+	virtual void ConvertToDatasmith(ISceneTracker& SceneTracker, FNodeTracker& NodeTracker) = 0;
+	virtual void RemoveFromTracked(ISceneTracker& SceneTracker, FNodeTracker& NodeTracker) = 0;
+
+	const EType ConverterType;
+};
+
+class FMeshNodeConverter: public FNodeConverter
+{
+public:
+	FMeshNodeConverter(): FNodeConverter(EType::MeshNode)
+	{
+	}
+	virtual void Parse(ISceneTracker& SceneTracker, FNodeTracker& NodeTracker) override;
+	virtual void ConvertToDatasmith(ISceneTracker& SceneTracker, FNodeTracker& NodeTracker) override;
+	virtual void RemoveFromTracked(ISceneTracker& SceneTracker, FNodeTracker& NodeTracker) override;
+
+	AnimHandle InstanceHandle = 0; // todo: rename - this is handle for object this node is instance of
+	bool bMaterialsAssignedToStaticMesh; // This node will used materials assigned to static mesh, otherwise actor's override
+};
+
+class FHismNodeConverter: public FNodeConverter
+{
+public:
+	FHismNodeConverter(): FNodeConverter(EType::HismNode)
+	{
+	}
+	virtual void RemoveFromTracked(ISceneTracker& SceneTracker, FNodeTracker& NodeTracker) override;
+	
+	TArray<FMeshConverted> Meshes;
+};
+
+class FRailCloneNodeConverter: public FHismNodeConverter
+{
+public:
+	virtual void Parse(ISceneTracker& SceneTracker, FNodeTracker& NodeTracker) override;
+	virtual void ConvertToDatasmith(ISceneTracker& SceneTracker, FNodeTracker& NodeTracker) override;
+};
+
+class FForestNodeConverter: public FHismNodeConverter
+{
+public:
+	virtual void Parse(ISceneTracker& SceneTracker, FNodeTracker& NodeTracker) override;
+	virtual void ConvertToDatasmith(ISceneTracker& SceneTracker, FNodeTracker& NodeTracker) override;
+};
+
+class FHelperNodeConverter: public FNodeConverter
+{
+public:
+	FHelperNodeConverter(): FNodeConverter(EType::HelperNode)
+	{
+	}
+
+	virtual void Parse(ISceneTracker& SceneTracker, FNodeTracker& NodeTracker) override;
+	virtual void ConvertToDatasmith(ISceneTracker& SceneTracker, FNodeTracker& NodeTracker) override;
+	virtual void RemoveFromTracked(ISceneTracker& SceneTracker, FNodeTracker& NodeTracker) override;
+};
+
+class FLightNodeConverter: public FNodeConverter
+{
+public:
+	FLightNodeConverter(): FNodeConverter(EType::LightNode)
+	{
+	}
+
+	virtual void Parse(ISceneTracker& SceneTracker, FNodeTracker& NodeTracker) override;
+	virtual void ConvertToDatasmith(ISceneTracker& SceneTracker, FNodeTracker& NodeTracker) override;
+	virtual void RemoveFromTracked(ISceneTracker& SceneTracker, FNodeTracker& NodeTracker) override;
+};
+
+class FCameraNodeConverter: public FNodeConverter
+{
+public:
+	FCameraNodeConverter(): FNodeConverter(EType::CameraNode)
+	{
+	}
+	virtual void Parse(ISceneTracker& SceneTracker, FNodeTracker& NodeTracker) override;
+	virtual void ConvertToDatasmith(ISceneTracker& SceneTracker, FNodeTracker& NodeTracker) override;
+	virtual void RemoveFromTracked(ISceneTracker& SceneTracker, FNodeTracker& NodeTracker) override;
+};
+
+}
+
diff --git a/Engine/Source/Programs/Enterprise/Datasmith/DatasmithMaxExporter/Public/DatasmithMaxDirectLink.h b/Engine/Source/Programs/Enterprise/Datasmith/DatasmithMaxExporter/Public/DatasmithMaxDirectLink.h
--- a/Engine/Source/Programs/Enterprise/Datasmith/DatasmithMaxExporter/Public/DatasmithMaxDirectLink.h
+++ b/Engine/Source/Programs/Enterprise/Datasmith/DatasmithMaxExporter/Public/DatasmithMaxDirectLink.h
@@ -4,6 +4,9 @@
 
 #include "DatasmithMaxExporterDefines.h"
 
+#include "DatasmithMaxConverters.h"
+#include "DatasmithMaxGeomUtils.h"
+
 #include "IDatasmithSceneElements.h"
 #include "DatasmithMesh.h"
 
@@ -22,6 +25,7 @@ MAX_INCLUDES_START
 MAX_INCLUDES_END
 
 
+
 // Log all messages
 // #define LOG_DEBUG_HEAVY_ENABLE 1
 // #define LOG_DEBUG_ENABLE 1
@@ -29,6 +33,8 @@ MAX_INCLUDES_END
 
 void AssignMeshMaterials(TSharedPtr<IDatasmithMeshElement>&MeshElement, Mtl * Material, const TSet<uint16>&SupportedChannels);
 
+class FDatasmithMaxStaticMeshAttributes;
+
 namespace DatasmithMaxDirectLink
 {
 
@@ -102,6 +108,7 @@ IPersistentExportOptions& GetPersistentExportOptions();
 typedef NodeEventNamespace::NodeKey FNodeKey;
 typedef MtlBase* FMaterialKey;
 
+
 struct FXRefScene
 {
 	INode* Tree = nullptr;
@@ -110,26 +117,123 @@ struct FXRefScene
 	operator bool(){ return XRefFileIndex != -1; }
 };
 
+// Incapsulating time slider value to distinguish usage of TimeValue specifically for point on timeslider where we syncing
+struct FSyncPoint
+{
+	TimeValue Time;
+};
+
+// Incapsulating time interval value to distinguish usage of Interval specifically for validity of synced entity 
+struct FValidity
+{
+	explicit FValidity()
+	{
+		Invalidate(); // Invalid by default(to essentially force update(validate explicitly)
+	}
+
+	bool IsInvalidated() const
+	{
+		bool bResult = ValidityInterval.Empty() != 0; // Invalidated equals interval is empty
+		check(bIsInvalidated == bResult);
+		return bResult;
+	}
+
+	void SetValid()
+	{
+		bIsInvalidated = false;
+	}
+
+	void Invalidate()
+	{
+		bIsInvalidated = true;
+		ValidityInterval.SetEmpty();
+	}
+
+	bool IsValidForSyncPoint(const FSyncPoint& State) const
+	{
+		return ValidityInterval.InInterval(State.Time) != 0;
+	}
+
+	// Maximize validity interval before updating invalidated 
+	void ResetValidityInterval()
+	{
+		check(bIsInvalidated);
+		ValidityInterval.SetInfinite();
+	}
 
-// Identifies Max node to track its changes
+	// Intersect
+	void NarrowValidityToInterval(const Interval& InValidityInterval)
+	{
+		ValidityInterval &= InValidityInterval;
+	}
+
+	void NarrowValidityToInterval(const FValidity& Validity)
+	{
+		NarrowValidityToInterval(Validity.ValidityInterval);
+	}
+
+	// this Validity contains another validity interval fully
+	// Used to determine if this validity doesn't need to be updated
+	// @return True is this validity doesn't need to be updated
+	bool Overlaps(const FValidity& Validity) const
+	{
+		return ValidityInterval.IsSubset(Validity.ValidityInterval);
+	}
+
+private:
+	Interval ValidityInterval;
+	// todo: this flag indicates, that 'invalidate' called specifically(no force rebuild
+	//  this is done in order to distinguish from Empty Interval that might happen somehow when updated
+	//  probably this is not needed - need to think about this and just rely on Empty for invalidated state 
+	bool bIsInvalidated; 
+};
+
+// NodeTracker - associated with INode, plus validity
+// Converter - holds convertion data for SPECIFIC node object type, e.g. geometry, railclone, light, dummy
+// Converted - holds Datasmith converted data for node
+
+// Identifies Max node to track its changes, including tracked dependencies
 class FNodeTracker: FNoncopyable
 {
 public:
-	explicit FNodeTracker(FNodeKey InNodeKey, INode* InNode) : NodeKey(InNodeKey), Node(InNode), Name(Node->GetName()) {}
+	explicit FNodeTracker(FNodeKey InNodeKey, INode* InNode) : NodeKey(InNodeKey), Node(InNode) {}
 
 	void Invalidate()
 	{
-		bInvalidated = true;
+		Validity.Invalidate();
 	}
 
+	void SetValid()
+	{
+		Validity.SetValid();
+	}
+
+	void ResetValidityInterval()
+	{
+		Validity.ResetValidityInterval();
+	}
+
+	bool IsValidForSyncPoint(const FSyncPoint& State) const
+	{
+		return Validity.IsValidForSyncPoint(State);
+	}
+
+	bool IsSubtreeValidForSyncPoint(const FSyncPoint& State) const
+	{
+		return SubtreeValidity.IsValidForSyncPoint(State);
+	}
+
+	// Forcefully invalidated
 	bool IsInvalidated() const
 	{
-		return bInvalidated;
+		return Validity.IsInvalidated();
 	}
 
-	bool IsInstance()
+	// todo: rename. Intersect validity with currently set one
+	//   e.g. validity interval of a geometry should be intersection of validity of transform, mesh, etc
+	void NarrowValidityToInterval(Interval ValidityInterval) 
 	{
-		return InstanceHandle != 0;
+		return Validity.NarrowValidityToInterval(ValidityInterval);
 	}
 
 	void SetXRefIndex(FXRefScene InXRefScene)
@@ -139,77 +243,108 @@ public:
 
 	INode* GetXRefParent()
 	{
-		return XRefScene ? XRefScene.Tree->GetXRefParent(XRefScene.XRefFileIndex): nullptr;
+		INode* Result = XRefScene ? XRefScene.Tree->GetXRefParent(XRefScene.XRefFileIndex): nullptr;
+		return Result;
 	}
 
-	FNodeKey NodeKey;
-	INode* const Node;
-	FNodeTracker* Collision = nullptr;
-	FLayerTracker* Layer = nullptr;
-	AnimHandle InstanceHandle = 0; // todo: rename - this is handle for object this node is instance of
-	FString Name; 
-	bool bInvalidated = true;
-	bool bDeleted = false;
-
-	TSharedPtr<IDatasmithActorElement> DatasmithActorElement;
-
-	TSet<class FMaterialTracker*> MaterialTrackers;
-
-	TSharedPtr<IDatasmithMeshActorElement> DatasmithMeshActor;
+	FNodeConverted& CreateConverted()
+	{
+		check(!Converted);
+		Converted.Reset(new FNodeConverted);
+		return *Converted;
+	}
 
-	FXRefScene XRefScene;
+	FNodeConverted& GetConverted()
+	{
+		check(Converted);
+		return *Converted;
+	}
 
-	TArray<FNodeTracker*> Children;
+	bool HasConverted()
+	{
+		return Converted.IsValid();
+	}
 
-};
+	void ReleaseConverted()
+	{
+		Converted.Reset();
+	}
 
-class FRenderMeshForConversion: FNoncopyable
-{
-public:
-	explicit FRenderMeshForConversion(){}
-	explicit FRenderMeshForConversion(INode* InNode, Mesh* InMaxMesh, bool bInNeedsDelete, FTransform InPivot=FTransform::Identity)
-		: Node(InNode)
-		, MaxMesh(InMaxMesh)
-		, bNeedsDelete(bInNeedsDelete)
-		, Pivot(InPivot)
-	{}
-	~FRenderMeshForConversion()
-	{
-		if (bNeedsDelete)
-		{
-			MaxMesh->DeleteThis();
-		}
+	template<typename ConverterType>
+	ConverterType& CreateConverter()
+	{
+		check(!Converter);
+		ConverterType* Result = new ConverterType;
+		Converter.Reset(Result);
+		return *Result;
 	}
 
-	bool IsValid() const
+	FNodeConverter& GetConverter()
 	{
-		return MaxMesh != nullptr;
+		check(Converter);
+		return *Converter;
 	}
 
-	INode* GetNode() const
+	bool HasConverter()
 	{
-		return Node;
+		return Converter.IsValid();
 	}
 
-	Mesh* GetMesh() const
+	FNodeConverter::EType GetConverterType()
 	{
-		return MaxMesh;
+		if (Converter.IsValid())
+		{
+			return Converter->ConverterType;
+		}
+		return FNodeConverter::Unknown;
 	}
 
-	const FTransform& GetPivot() const
+	void ReleaseConverter()
 	{
-		return Pivot;
+		Converter.Reset();
 	}
 
+	// Node entity identification data
+	FNodeKey NodeKey;
+	INode* const Node;
+	// Keep root node and xref index when this node is a direct child of an XRef scene
+	// This is needed to retrieve parent node(e.g. when updated)
+	// Keeping parent node itself doesn't work - it can change and the only way to get it when it's changed is to call INode::GetXRefParent
+	FXRefScene XRefScene;
+	FString Name;
+
+	// Node validity
+	bool bDeleted = false;
+	FValidity Validity;
+	FValidity SubtreeValidity;
 
-private:
-	INode* Node = nullptr;
-	Mesh* MaxMesh = nullptr;
-	bool bNeedsDelete = false;
-	FTransform Pivot;
+	// Other related tracked entities
+	FNodeTracker* Collision = nullptr;
+	FLayerTracker* Layer = nullptr;
+	TSet<class FMaterialTracker*> MaterialTrackers; 
+
+	// Node conversion state
+	TUniquePtr<FNodeConverted> Converted = nullptr; // Datasmith element that this Node is converted to
+	TUniquePtr<FNodeConverter> Converter = nullptr; // Converter for specific Max object type
+};
+
+class FMeshConverterSource
+{
+public:
+	// Node this mesh instantiates. When this is a 'regular' node it's  just instantiates mesh for Params
+	// but it's possible that this Node wants mesh to be bounding-box(when DatasmithAttributes spicifies it)
+	INode* Node;
+
+	FString MeshName; // Suggested Mesh name, resulting mesh name should be this???
 
+	GeomUtils::FRenderMeshForConversion RenderMesh; // Extracted render mesh
+	bool bConsolidateMaterialIds; // Whether to join all materials id into single material slot for render mesh (used when a geometry doesn't have a multimaterial assigned)
+
+	GeomUtils::FRenderMeshForConversion CollisionMesh; 
 };
 
+
+
 struct FSceneUpdateStats
 {
 	void Reset()
@@ -245,6 +380,11 @@ struct FSceneUpdateStats
 	int32 UpdateMaterialsSkippedAsAlreadyConverted =0;
 	int32 UpdateMaterialsConverted = 0;
 	int32 UpdateTexturesTotal = 0;
+	int32 CheckTimeSliderTotalChecks = 0;
+	int32 CheckTimeSliderSkippedAsAlreadyInvalidated = 0;
+	int32 CheckTimeSliderSkippedAsSubtreeValid = 0;
+	int32 CheckTimeSliderInvalidated = 0;
+	int32 ConvertNodesConverted = 0;
 };
 
 #define SCENE_UPDATE_STAT_INC(Category, Name) {Stats.##Category##Name++;}
@@ -264,6 +404,7 @@ public:
 	virtual void NodeDeleted(INode* Node) = 0;
 	virtual void NodeGeometryChanged(FNodeKey NodeKey) = 0;
 	virtual void NodeHideChanged(FNodeKey NodeKey) = 0;
+	virtual void NodeNameChanged(FNodeKey NodeKey) = 0;
 	virtual void NodePropertiesChanged(FNodeKey NodeKey) = 0;
 	virtual void NodeLinkChanged(FNodeKey NodeKey) = 0;
 	virtual void NodeTransformChanged(FNodeKey NodeKey) = 0;
@@ -272,29 +413,42 @@ public:
 
 	// Scene modification
 	virtual void AddMeshElement(TSharedPtr<IDatasmithMeshElement>& Mesh, FDatasmithMesh& DatasmithMesh, FDatasmithMesh* CollisionMesh) = 0;
-	virtual void ReleaseMeshElement(TSharedPtr<IDatasmithMeshElement> Mesh) = 0;
+	virtual void ReleaseMeshElement(FMeshConverted& Converted) = 0;
 	virtual void SetupActor(FNodeTracker& NodeTracker) = 0;
-	virtual void SetupDatasmithHISMForNode(FNodeTracker& NodeTracker, INode* GeometryNode, const FRenderMeshForConversion& RenderMesh, Mtl* Material, int32 MeshIndex, const TArray<Matrix3>& Transforms) = 0;
+	virtual void SetupDatasmithHISMForNode(FNodeTracker& NodeTracker, FMeshConverterSource& MeshSource, Mtl* Material, int32 MeshIndex, const TArray<Matrix3>& Transforms) = 0;
 	virtual void RemoveMaterial(const TSharedPtr<IDatasmithBaseMaterialElement>& DatasmithMaterial) = 0;
 	virtual void RemoveTexture(const TSharedPtr<IDatasmithTextureElement>&) = 0;
 	virtual void NodeXRefMerged(INode* Node) = 0;
-	virtual FSceneUpdateStats& GetStats() = 0;
-};
 
-//---- Geometry utility function
-bool ConvertRailClone(ISceneTracker& SceneTracker, FNodeTracker& NodeTracker, Object* Obj);
-bool ConvertForest(ISceneTracker& Scene, FNodeTracker& NodeTracker, Object* Obj);
 
-Mesh* GetMeshFromRenderMesh(INode* Node, BOOL& bNeedsDelete, TimeValue CurrentTime);
-FRenderMeshForConversion GetMeshForGeomObject(INode* Node, Object* Obj); // Extract mesh using already evaluated object
-FRenderMeshForConversion GetMeshForNode(INode* Node, FTransform Pivot); // Extract mesh evaluating node object
-FRenderMeshForConversion GetMeshForCollision(INode* Node);
 
-void FillDatasmithMeshFromMaxMesh(FDatasmithMesh& DatasmithMesh, Mesh& MaxMesh, INode* ExportedNode, bool bForceSingleMat, TSet<uint16>& SupportedChannels, TMap<int32, int32>& UVChannelsMap, FTransform Pivot);
+	// Sync/Update
+	FSyncPoint CurrentSyncPoint;
+
+	virtual void AddGeometryNodeInstance(FNodeTracker& NodeTracker, FMeshNodeConverter& MeshConverter, Object* Obj) = 0;
+	virtual void RemoveGeometryNodeInstance(FNodeTracker& NodeTracker) = 0;
+	virtual void ConvertGeometryNodeToDatasmith(FNodeTracker& NodeTracker, FMeshNodeConverter& MeshConverter) = 0;
+
+	virtual void UnregisterNodeForMaterial(FNodeTracker& NodeTracker) = 0;
 
+	virtual TSharedRef<IDatasmithScene> GetDatasmithSceneRef() = 0;
+
+	// Utility
+	virtual FNodeTracker* GetNodeTrackerByNodeName(const TCHAR* Name) = 0;
+	virtual FSceneUpdateStats& GetStats() = 0;
+};
+
+// Input data for mesh conversion
+struct MeshConversionParams
+{
+	INode* Node; // Node, this geom object created from
+	const TCHAR* MeshName;
+	const GeomUtils::FRenderMeshForConversion& RenderMesh; // Extracted render mesh
+	bool bConsolidateMaterialIds; // Whether to join all materials id into single material slot (used when a geometry doesn't have a multimaterial assigned)
+};
 
 // Creates Mesh element and converts max mesh into it
-bool ConvertMaxMeshToDatasmith(ISceneTracker& Scene, TSharedPtr<IDatasmithMeshElement>& DatasmithMeshElement, INode* Node, const TCHAR* MeshName, const FRenderMeshForConversion& RenderMesh, bool bConsolidateMaterialIds, TSet<uint16>& SupportedChannels, const FRenderMeshForConversion& CollisionMesh = FRenderMeshForConversion());
+bool ConvertMaxMeshToDatasmith(TimeValue CurrentTime, ISceneTracker& Scene, FMeshConverterSource& MeshSource, FMeshConverted& MeshConverted);
 
 bool OpenDirectLinkUI();
 const TCHAR* GetDirectlinkCacheDirectory();
diff --git a/Engine/Source/Programs/Enterprise/Datasmith/DatasmithMaxExporter/Public/DatasmithMaxGeomUtils.h b/Engine/Source/Programs/Enterprise/Datasmith/DatasmithMaxExporter/Public/DatasmithMaxGeomUtils.h
new file mode 100644
--- /dev/null
+++ b/Engine/Source/Programs/Enterprise/Datasmith/DatasmithMaxExporter/Public/DatasmithMaxGeomUtils.h
@@ -0,0 +1,85 @@
+// Copyright Epic Games, Inc. All Rights Reserved.
+
+#pragma once
+
+#include "DatasmithMaxExporterDefines.h"
+
+#include "Windows/AllowWindowsPlatformTypes.h"
+MAX_INCLUDES_START
+	#include "max.h"
+MAX_INCLUDES_END
+#include "Windows/HideWindowsPlatformTypes.h"
+
+
+class FDatasmithMaxStaticMeshAttributes;
+class FDatasmithMesh;
+
+namespace DatasmithMaxDirectLink
+{
+
+class ISceneTracker;
+class FNodeTracker;
+
+namespace GeomUtils
+{
+
+class FRenderMeshForConversion: FNoncopyable
+{
+public:
+	explicit FRenderMeshForConversion(){}
+	explicit FRenderMeshForConversion(INode* InNode, Mesh* InMaxMesh, bool bInNeedsDelete, FTransform InPivot=FTransform::Identity)
+		: Node(InNode)
+		, MaxMesh(InMaxMesh)
+		, bNeedsDelete(bInNeedsDelete)
+		, Pivot(InPivot)
+	{}
+	~FRenderMeshForConversion()
+	{
+		if (bNeedsDelete)
+		{
+			MaxMesh->DeleteThis();
+		}
+	}
+
+	bool IsValid() const
+	{
+		return MaxMesh != nullptr;
+	}
+
+	INode* GetNode() const
+	{
+		return Node;
+	}
+
+	Mesh* GetMesh() const
+	{
+		return MaxMesh;
+	}
+
+	const FTransform& GetPivot() const
+	{
+		return Pivot;
+	}
+
+
+private:
+	INode* Node = nullptr;
+	Mesh* MaxMesh = nullptr;
+	bool bNeedsDelete = false;
+	FTransform Pivot;
+};
+
+bool ConvertRailClone(DatasmithMaxDirectLink::ISceneTracker& SceneTracker, DatasmithMaxDirectLink::FNodeTracker& NodeTracker);
+bool ConvertForest(DatasmithMaxDirectLink::ISceneTracker& Scene, DatasmithMaxDirectLink::FNodeTracker& NodeTracker);
+
+Mesh* GetMeshFromRenderMesh(TimeValue CurrentTime, INode* Node, BOOL& bNeedsDelete);
+FRenderMeshForConversion GetMeshForGeomObject(TimeValue CurrentTime, INode* Node, Object* Obj); // Extract mesh using already evaluated object
+FRenderMeshForConversion GetMeshForNode(TimeValue CurrentTime, INode* Node, FTransform Pivot); // Extract mesh evaluating node object
+
+FRenderMeshForConversion GetMeshForCollision(TimeValue CurrentTime, DatasmithMaxDirectLink::ISceneTracker& SceneTracker, INode* Node);
+INode* GetCollisionNode(DatasmithMaxDirectLink::ISceneTracker& SceneTracker, INode* OriginalNode, const FDatasmithMaxStaticMeshAttributes* DatasmithAttributes, bool& bOutFromDatasmithAttribute);
+
+void FillDatasmithMeshFromMaxMesh(TimeValue CurrentTime, FDatasmithMesh& DatasmithMesh, Mesh& MaxMesh, INode* ExportedNode, bool bForceSingleMat, TSet<uint16>& SupportedChannels, TMap<int32, int32>& UVChannelsMap, FTransform Pivot);
+
+}
+}
diff --git a/Engine/Source/Programs/Enterprise/Datasmith/DatasmithMaxExporter/Public/DatasmithMaxSceneExporter.h b/Engine/Source/Programs/Enterprise/Datasmith/DatasmithMaxExporter/Public/DatasmithMaxSceneExporter.h
--- a/Engine/Source/Programs/Enterprise/Datasmith/DatasmithMaxExporter/Public/DatasmithMaxSceneExporter.h
+++ b/Engine/Source/Programs/Enterprise/Datasmith/DatasmithMaxExporter/Public/DatasmithMaxSceneExporter.h
@@ -55,7 +55,7 @@ public:
 	static bool ExportActor(TSharedRef< IDatasmithScene > DatasmithScene, INode* Node, const TCHAR* Name, float UnitMultiplier);
 	static void ExportMeshActor(TSharedRef< IDatasmithScene > DatasmithScene,  TSet<uint16>& SupportedChannels, INode* Node, const TCHAR* MeshName,
 		float UnitMultiplier, bool bPivotIsBakedInGeometry, Mtl* StaticMeshMtl, const EStaticMeshExportMode& ExportMode);
-	static TSharedRef< IDatasmithActorElement > ExportHierarchicalInstanceStaticMeshActor(TSharedRef< IDatasmithScene > DatasmithScene, INode* Node, INode* CustomMeshNode, const TCHAR* Label, TSet<uint16>& SupportedChannels, Mtl* StaticMeshMtl, const TArray<Matrix3>* Instances,
+	static TSharedRef< IDatasmithActorElement > ExportHierarchicalInstanceStaticMeshActor(INode* Node, INode* CustomMeshNode, const TCHAR* Label, TSet<uint16>& SupportedChannels, Mtl* StaticMeshMtl, const TArray<Matrix3>* Instances,
 		const TCHAR* MeshName, float UnitMultiplier, const EStaticMeshExportMode& ExportMode, TSharedPtr< IDatasmithActorElement >& OutInversedHISMActor );
 	static bool WriteXMLLightActor(TSharedRef< IDatasmithScene > DatasmithScene, INode* Parent, INode* Node, const TCHAR* Name, float UnitMultiplier);
 	static bool ExportCameraActor(TSharedRef< IDatasmithScene > DatasmithScene, INode* Parent, INodeTab Instances, int InstanceIndex, const TCHAR* Name, float UnitMultiplier);
