DatasmithRhinoExporter - Fix crash when deleting a layer with DirectLink active.

diff --git a/Engine/Source/Programs/Enterprise/Datasmith/DatasmithRhinoExporter/Private/DirectLink/DatasmithRhinoChangeListener.cs b/Engine/Source/Programs/Enterprise/Datasmith/DatasmithRhinoExporter/Private/DirectLink/DatasmithRhinoChangeListener.cs
--- a/Engine/Source/Programs/Enterprise/Datasmith/DatasmithRhinoExporter/Private/DirectLink/DatasmithRhinoChangeListener.cs
+++ b/Engine/Source/Programs/Enterprise/Datasmith/DatasmithRhinoExporter/Private/DirectLink/DatasmithRhinoChangeListener.cs
@@ -171,7 +171,7 @@ namespace DatasmithRhino.DirectLink
 					TryCatchExecute(() => ExportContext.ModifyActor(RhinoEventArgs.NewState, bReparent));
 					break;
 				case LayerTableEventType.Deleted:
-					TryCatchExecute(() => ExportContext.DeleteActor(RhinoEventArgs.OldState));
+					TryCatchExecute(() => ExportContext.DeleteActor(RhinoEventArgs.NewState));
 					break;
 				case LayerTableEventType.Sorted:
 				case LayerTableEventType.Current:
