DatasmithRhinoExporter - Export mesh vertex colors

diff --git a/Engine/Source/Programs/Enterprise/Datasmith/DatasmithRhinoExporter/Private/ElementExporters/DatasmithRhinoMeshExporter.cs b/Engine/Source/Programs/Enterprise/Datasmith/DatasmithRhinoExporter/Private/ElementExporters/DatasmithRhinoMeshExporter.cs
--- a/Engine/Source/Programs/Enterprise/Datasmith/DatasmithRhinoExporter/Private/ElementExporters/DatasmithRhinoMeshExporter.cs
+++ b/Engine/Source/Programs/Enterprise/Datasmith/DatasmithRhinoExporter/Private/ElementExporters/DatasmithRhinoMeshExporter.cs
@@ -3,6 +3,7 @@
 using DatasmithRhino.ExportContext;
 using Rhino.Geometry;
 using System.Collections.Generic;
+using System.Drawing;
 
 namespace DatasmithRhino.ElementExporters
 {
@@ -126,6 +127,7 @@ namespace DatasmithRhino.ElementExporters
 				}
 
 				bool bUseFaceNormals = RhinoMesh.Normals.Count != RhinoMesh.Vertices.Count && RhinoMesh.FaceNormals.Count == RhinoMesh.Faces.Count;
+				bool bHasVertexColor = RhinoMesh.VertexColors.Count == RhinoMesh.Vertices.Count;
 
 				//Add triangles and normals to the mesh.
 				for (int FaceIndex = 0, FaceQuadOffset = 0; FaceIndex < RhinoMesh.Faces.Count; ++FaceIndex)
@@ -145,6 +147,11 @@ namespace DatasmithRhino.ElementExporters
 						DatasmithMesh.SetFaceUV(DatasmithFaceIndex + 1, 0, VertexIndexOffset + Face.A, VertexIndexOffset + Face.C, VertexIndexOffset + Face.D);
 					}
 
+					if (bHasVertexColor)
+					{
+						AddVertexColorToMesh(DatasmithMesh, RhinoMesh, Face, DatasmithFaceIndex, VertexIndexOffset);
+					}
+
 					if (bUseFaceNormals)
 					{
 						Vector3f Normal = RhinoMesh.FaceNormals[FaceIndex];
@@ -222,6 +229,26 @@ namespace DatasmithRhino.ElementExporters
 			}
 		}
 
+		private static void AddVertexColorToMesh(FDatasmithFacadeMesh DatasmithMesh, Mesh RhinoMesh, MeshFace Face, int FaceIndex, int VertexIndexOffset)
+		{
+			int VertexInstanceIndex = (FaceIndex * 3) + VertexIndexOffset;
+
+			Color ColorA = RhinoMesh.VertexColors[Face.A];
+			Color ColorB = RhinoMesh.VertexColors[Face.B];
+			Color ColorC = RhinoMesh.VertexColors[Face.C];
+			DatasmithMesh.SetVertexColor(VertexInstanceIndex + 0, ColorA.R, ColorA.G, ColorA.B, ColorA.A);
+			DatasmithMesh.SetVertexColor(VertexInstanceIndex + 1, ColorB.R, ColorB.G, ColorB.B, ColorB.A);
+			DatasmithMesh.SetVertexColor(VertexInstanceIndex + 2, ColorC.R, ColorC.G, ColorC.B, ColorC.A);
+
+			if (Face.IsQuad)
+			{
+				Color ColorD = RhinoMesh.VertexColors[Face.D];
+				DatasmithMesh.SetVertexColor(VertexInstanceIndex + 3, ColorA.R, ColorA.G, ColorA.B, ColorA.A);
+				DatasmithMesh.SetVertexColor(VertexInstanceIndex + 4, ColorC.R, ColorC.G, ColorC.B, ColorC.A);
+				DatasmithMesh.SetVertexColor(VertexInstanceIndex + 5, ColorD.R, ColorD.G, ColorD.B, ColorD.A);
+			}
+		}
+
 		private static void AddNormalsToMesh(FDatasmithFacadeMesh Mesh, int FaceIndex, Vector3f Normal)
 		{
 			AddNormalsToMesh(Mesh, FaceIndex, Normal, Normal, Normal);
