DatasmithRhinoExporter - Fix crash after when updating a worksession document when DirectLink is active.

diff --git a/Engine/Source/Programs/Enterprise/Datasmith/DatasmithRhinoExporter/Private/ExportContext/DatasmithInfoBase.cs b/Engine/Source/Programs/Enterprise/Datasmith/DatasmithRhinoExporter/Private/ExportContext/DatasmithInfoBase.cs
--- a/Engine/Source/Programs/Enterprise/Datasmith/DatasmithRhinoExporter/Private/ExportContext/DatasmithInfoBase.cs
+++ b/Engine/Source/Programs/Enterprise/Datasmith/DatasmithRhinoExporter/Private/ExportContext/DatasmithInfoBase.cs
@@ -64,12 +64,22 @@ namespace DatasmithRhino.ExportContext
 			SetDirectLinkStatus(DirectLinkSynchronizationStatus.PendingHidding);
 		}
 
+		public bool HasHiddenStatus()
+		{
+			return DirectLinkStatus == DirectLinkSynchronizationStatus.PendingHidding || DirectLinkStatus == DirectLinkSynchronizationStatus.Hidden;
+		}
+
 		[MethodImpl(MethodImplOptions.AggressiveInlining)]
 		public void ApplyDeletedStatus()
 		{
 			SetDirectLinkStatus(DirectLinkSynchronizationStatus.PendingDeletion);
 		}
 
+		public bool HasDeletedStatus()
+		{
+			return DirectLinkStatus == DirectLinkSynchronizationStatus.PendingDeletion || DirectLinkStatus == DirectLinkSynchronizationStatus.Deleted;
+		}
+
 		private void SetDirectLinkStatus(DirectLinkSynchronizationStatus Status)
 		{
 			// We ignore the "modified" status if the current status is "created", "deleted" or "hidden".
diff --git a/Engine/Source/Programs/Enterprise/Datasmith/DatasmithRhinoExporter/Private/ExportContext/DatasmithRhinoExportContext.cs b/Engine/Source/Programs/Enterprise/Datasmith/DatasmithRhinoExporter/Private/ExportContext/DatasmithRhinoExportContext.cs
--- a/Engine/Source/Programs/Enterprise/Datasmith/DatasmithRhinoExporter/Private/ExportContext/DatasmithRhinoExportContext.cs
+++ b/Engine/Source/Programs/Enterprise/Datasmith/DatasmithRhinoExporter/Private/ExportContext/DatasmithRhinoExportContext.cs
@@ -308,7 +308,7 @@ namespace DatasmithRhino.ExportContext
 			{
 				// The actor already exists in the cache, it means that no direct link synchronization was done since its deletion.
 				// In that case, we can simply undelete it.
-				UndeleteActor(CachedActorInfo);
+				UndeleteActor(CachedActorInfo, InModelComponent);
 				return;
 			}
 
@@ -338,29 +338,48 @@ namespace DatasmithRhino.ExportContext
 			bIsDirty = true;
 		}
 
-		private void UndeleteActor(DatasmithActorInfo ActorInfo)
+		private void UndeleteActor(DatasmithActorInfo InDeletedActorInfo, ModelComponent InUndeletedModelComponent)
 		{
-			foreach (DatasmithActorInfo ActorInfoValue in ActorInfo.GetEnumerator())
-			{
-				System.Diagnostics.Debug.Assert(ActorInfoValue.DirectLinkStatus == DirectLinkSynchronizationStatus.PendingDeletion);
-				if (ActorInfoValue.RestorePreviousDirectLinkStatus())
+			Func<DatasmithActorInfo, Guid, bool> UndeleteActorImpl = (DatasmithActorInfo ActorInfo, Guid RhinoObjectId) =>
 				{
-					if (ObjectIdToMeshInfoDictionary.TryGetValue(ActorInfoValue.RhinoObjectId, out DatasmithMeshInfo MeshInfo))
+					System.Diagnostics.Debug.Assert(ActorInfo.DirectLinkStatus == DirectLinkSynchronizationStatus.PendingDeletion);
+					if (ActorInfo.RestorePreviousDirectLinkStatus())
 					{
-						MeshInfo.RestorePreviousDirectLinkStatus();
+						if (ObjectIdToMeshInfoDictionary.TryGetValue(RhinoObjectId, out DatasmithMeshInfo MeshInfo))
+						{
+							MeshInfo.RestorePreviousDirectLinkStatus();
+						}
 					}
-				}
-				else
+					else
+					{
+						System.Diagnostics.Debug.Fail("The previous direct link state was lost on an undeleted actor");
+						return false;
+					}
+
+					return true;
+				};
+
+			if (UndeleteActorImpl(InDeletedActorInfo, InUndeletedModelComponent.Id))
+			{
+				// Only undelete the children if the current actor is a block instance.
+				// For the other cases Rhino will explicitly undelete the concerned actors.
+				if (InDeletedActorInfo.DefinitionNode != null)
 				{
-					System.Diagnostics.Debug.Fail("The previous direct link state was lost on an undeleted actor");
-					break;
+					foreach (DatasmithActorInfo ChildActorInfoValue in InDeletedActorInfo.GetDescendantEnumerator())
+					{
+						if (!UndeleteActorImpl(ChildActorInfoValue, ChildActorInfoValue.RhinoObjectId))
+						{
+							break;
+						}
+					}
 				}
 			}
 
-			if (ActorInfo.RhinoCommonObject is ModelComponent RhinoModelComponent)
+			if (InUndeletedModelComponent != null)
 			{
-				//Undelete is always used during an undo, and undo often use "undelete" to restore an actor state (instead of calling modify).
-				ModifyActor(RhinoModelComponent, /*bReparent=*/false);
+				// Undelete is always used during an undo, even to restore an actor state (instead of calling modify).
+				// bReparent must be true because actors may be reparented on a worksession update.
+				ModifyActor(InUndeletedModelComponent, /*bCheckForReparenting=*/true);
 			}
 
 			bIsDirty = true;
@@ -404,16 +423,17 @@ namespace DatasmithRhino.ExportContext
 		/// Generic Actor modification event used when the RhinoObject may have multiple property modifications, both the associated DatasmithActor and DatasmithMesh will be synced.
 		/// </summary>
 		/// <param name="InModelComponent"></param>
-		/// <param name="bReparent"></param>
-		public void ModifyActor(ModelComponent InModelComponent, bool bReparent)
+		/// <param name="bCheckForReparenting"></param>
+		public void ModifyActor(ModelComponent InModelComponent, bool bCheckForReparenting)
 		{
 			RhinoObject InRhinoObject = InModelComponent as RhinoObject;
 			if (ObjectIdToHierarchyActorNodeDictionary.TryGetValue(InModelComponent.Id, out DatasmithActorInfo ActorInfo))
 			{
 				Layer InLayer = InModelComponent as Layer;
-				DatasmithActorInfo ActorParentInfo = bReparent 
+				DatasmithActorInfo ActorParentInfo = bCheckForReparenting
 					? GetActorParentInfo(InModelComponent)
 					: ActorInfo.Parent;
+				bool bReparent = ActorParentInfo != ActorInfo.Parent;
 
 				DatasmithActorInfo DiffActorInfo = null;
 				if (InRhinoObject != null)
@@ -447,7 +467,7 @@ namespace DatasmithRhino.ExportContext
 						ActorInfo.Reparent(ActorParentInfo);
 					}
 
-					if (bReparent || bMaterialChanged)
+					if (bCheckForReparenting || bMaterialChanged)
 					{
 						UpdateChildActorsMaterialIndex(ActorInfo);
 					}
@@ -483,9 +503,15 @@ namespace DatasmithRhino.ExportContext
 
 		private void UpdateHiddenFlagsRecursively(DatasmithActorInfo ActorInfo, HashSet<RhinoObject> OutNewObjects)
 		{
+			if (ActorInfo.HasDeletedStatus())
+			{
+				// There is no need (and it is unsafe) to check for visibility on objects that will be deleted.
+				return;
+			}
+
 			// if the visibility changed, update the status of the descendants.
 			bool bIsVisible = ActorInfo.bIsVisible;
-			bool bHasHiddenFlag = (ActorInfo.DirectLinkStatus & (DirectLinkSynchronizationStatus.PendingHidding | DirectLinkSynchronizationStatus.Hidden)) != DirectLinkSynchronizationStatus.None;
+			bool bHasHiddenFlag = ActorInfo.HasHiddenStatus();
 			if (bIsVisible == bHasHiddenFlag)
 			{
 				if (!bIsVisible)
@@ -554,6 +580,12 @@ namespace DatasmithRhino.ExportContext
 			{
 				foreach (DatasmithActorInfo ActorInfoValue in ActorInfo.GetEnumerator())
 				{
+					if (ActorInfoValue.HasDeletedStatus())
+					{
+						// Don't try to delete the same object twice, it is not safe to access it after it's been flagged for deletion.
+						continue;
+					}
+
 					ActorInfoValue.ApplyDeletedStatus();
 
 					// If this actor is not a block instance, we can delete its associated mesh.
@@ -571,6 +603,11 @@ namespace DatasmithRhino.ExportContext
 		{
 			foreach (DatasmithActorInfo ChildActorInfo in ActorInfo.Children)
 			{
+				if (ChildActorInfo.HasDeletedStatus())
+				{
+					continue;
+				}
+
 				ChildActorInfo.ApplyModifiedStatus();
 				bool bMaterialAffectedByParent = ChildActorInfo.RhinoCommonObject is RhinoObject ChildRhinoObject && ChildRhinoObject.Attributes.MaterialSource != ObjectMaterialSource.MaterialFromObject;
 
@@ -665,7 +702,7 @@ namespace DatasmithRhino.ExportContext
 
 		private DatasmithActorInfo TryParseLayer(Layer CurrentLayer, DatasmithActorInfo ParentNode)
 		{
-			if ((ExportOptions.bSkipHidden && !CurrentLayer.IsVisible) || CurrentLayer.IsDeleted)
+			if (CurrentLayer.IsDeleted || (ExportOptions.bSkipHidden && !CurrentLayer.IsVisible))
 			{
 				return null;
 			}
