DatasmithRhinoExporter - Calculate pivot point at the botton-center instead of the centroid.

diff --git a/Engine/Source/Programs/Enterprise/Datasmith/DatasmithRhinoExporter/Private/Utils/DatasmithRhinoUtilities.cs b/Engine/Source/Programs/Enterprise/Datasmith/DatasmithRhinoExporter/Private/Utils/DatasmithRhinoUtilities.cs
--- a/Engine/Source/Programs/Enterprise/Datasmith/DatasmithRhinoExporter/Private/Utils/DatasmithRhinoUtilities.cs
+++ b/Engine/Source/Programs/Enterprise/Datasmith/DatasmithRhinoExporter/Private/Utils/DatasmithRhinoUtilities.cs
@@ -252,10 +252,10 @@ namespace DatasmithRhino.Utils
 		}
 
 		/// <summary>
-		/// Center the given meshes on the pivot determined from the union of their bounding boxes. Returns the pivot point.
+		/// Center the given meshes on the pivot determined from the union of their bounding boxes. Returns the bottom center point.
 		/// </summary>
 		/// <param name="RhinoMeshes"></param>
-		/// <returns>The pivot point on which the Mesh was centered</returns>
+		/// <returns>The bottom-center point on which the Mesh was centered</returns>
 		public static Vector3d CenterMeshesOnPivot(IEnumerable<Mesh> RhinoMeshes)
 		{
 			BoundingBox MeshesBoundingBox = BoundingBox.Empty;
@@ -274,7 +274,7 @@ namespace DatasmithRhino.Utils
 				}
 			}
 
-			Vector3d PivotPoint = new Vector3d(MeshesBoundingBox.Center.X, MeshesBoundingBox.Center.Y, MeshesBoundingBox.Center.Z);
+			Vector3d PivotPoint = new Vector3d(MeshesBoundingBox.Center.X, MeshesBoundingBox.Center.Y, MeshesBoundingBox.Min.Z);
 			
 			foreach (Mesh CurrentMesh in RhinoMeshes)
 			{
