[3ds Max] DirectLink: Fix 3DsMax - Cameras import option makes Datasmith imports NOT only camera

- object with non-degenerate pivot but no ather instances doesn't create redundand actor

diff --git a/Engine/Source/Programs/Enterprise/Datasmith/DatasmithMaxExporter/Private/DatasmithMaxExporterWithDirectLink.cpp b/Engine/Source/Programs/Enterprise/Datasmith/DatasmithMaxExporter/Private/DatasmithMaxExporterWithDirectLink.cpp
--- a/Engine/Source/Programs/Enterprise/Datasmith/DatasmithMaxExporter/Private/DatasmithMaxExporterWithDirectLink.cpp
+++ b/Engine/Source/Programs/Enterprise/Datasmith/DatasmithMaxExporter/Private/DatasmithMaxExporterWithDirectLink.cpp
@@ -2102,9 +2102,9 @@ public:
 		GetNodeObjectTransform(NodeTracker, Converter, ObjectTransform);
 
 		FTransform Pivot = FDatasmithMaxSceneExporter::GetPivotTransform(NodeTracker.Node, Converter.UnitToCentimeter);
-		FTransform NodeTransform = Pivot.Inverse() * ObjectTransform; // Remove pivot from the node actor transform
 
-		bool bNeedPivotComponent = !Pivot.Equals(FTransform::Identity);
+		// Create separate actor only when there are multiple instances, 
+		bool bNeedPivotComponent = !Pivot.Equals(FTransform::Identity) && (Instances.NodeTrackers.Num() > 1) && Instances.HasMesh();  
 
 		TSharedPtr<IDatasmithActorElement> DatasmithActorElement;
 		TSharedPtr<IDatasmithMeshActorElement> DatasmithMeshActor;
@@ -2134,9 +2134,9 @@ public:
 			DatasmithMeshActor->SetStaticMeshPathName(Instances.GetStaticMeshPathName());
 		}
 
+		// Create a dummy actor in case pivot is non-degenerate or there's no mesh(so no mesh actor)
 		if (bNeedPivotComponent || !Instances.HasMesh())
 		{
-			// Create a separate actor in case pivot is different or there's not mesh(so no mesh actor)
 			DatasmithActorElement = FDatasmithSceneFactory::CreateActor(*UniqueName);
 			DatasmithActorElement->SetLabel(*Label);
 		}
@@ -2146,13 +2146,16 @@ public:
 		}
 
 		// Set transforms
-		DatasmithActorElement->SetTranslation(NodeTransform.GetTranslation());
-		DatasmithActorElement->SetScale(NodeTransform.GetScale3D());
-		DatasmithActorElement->SetRotation(NodeTransform.GetRotation());
-
-		// Setup mesh actor transform with pivot's
-		if (bNeedPivotComponent && Instances.HasMesh()) 
+		if (bNeedPivotComponent) 
 		{
+			// Remove pivot from the node actor transform
+			FTransform NodeTransform = Pivot.Inverse() * ObjectTransform;
+
+			DatasmithActorElement->SetTranslation(NodeTransform.GetTranslation());
+			DatasmithActorElement->SetScale(NodeTransform.GetScale3D());
+			DatasmithActorElement->SetRotation(NodeTransform.GetRotation());
+
+			// Setup mesh actor with (relative)pivot transform
 			DatasmithMeshActor->SetTranslation(Pivot.GetTranslation());
 			DatasmithMeshActor->SetRotation(Pivot.GetRotation());
 			DatasmithMeshActor->SetScale(Pivot.GetScale3D());
@@ -2160,6 +2163,14 @@ public:
 
 			DatasmithActorElement->AddChild(DatasmithMeshActor, EDatasmithActorAttachmentRule::KeepRelativeTransform);
 		}
+		else
+		{
+			FTransform NodeTransform = ObjectTransform;
+
+			DatasmithActorElement->SetTranslation(NodeTransform.GetTranslation());
+			DatasmithActorElement->SetScale(NodeTransform.GetScale3D());
+			DatasmithActorElement->SetRotation(NodeTransform.GetRotation());
+		}
 
 		FNodeConverted& Converted = NodeTracker.CreateConverted();
 		Converted.DatasmithActorElement = DatasmithActorElement;
