Fixed issue with wrong Ruby framework required for SketchUp 2022 exporter on Mac
Solution: Remove the link directive to use the Ruby framework as it is already done for the SketchUpAPI framework

diff --git a/Engine/Source/Programs/Enterprise/Datasmith/DatasmithSketchUpRubyExporter/Build/build_mac.py b/Engine/Source/Programs/Enterprise/Datasmith/DatasmithSketchUpRubyExporter/Build/build_mac.py
--- a/Engine/Source/Programs/Enterprise/Datasmith/DatasmithSketchUpRubyExporter/Build/build_mac.py
+++ b/Engine/Source/Programs/Enterprise/Datasmith/DatasmithSketchUpRubyExporter/Build/build_mac.py
@@ -202,7 +202,7 @@ link_cmd.append('-bundle') # link ruby extension as a MacOs 'bundle' not dynamic
 
 #link_cmd += ['-framework', 'Ruby'] # link to RUby framework (found by framework_search_paths)
 
-# Handle SketchUp API
+# Handle not linking against SketchUpAPI and Ruby frameworks. Undefined will be solved at runtime
 link_cmd += [
     '-undefined', 'dynamic_lookup', # We are referencing SketchUp API symbols that are defined in the SketchUp app
     # looks like the following(referencing SketchUp app) is not needed when we have the above(-undefined dynamic_lookup)
