Datasmith Revit: Changes in Linked Files not reflected in direct link

When doing independent changes and then reloading the link, the master document did not detect the change.
On link reload, we need to mark the whole linked document as changed and resend on next update.

diff --git a/Engine/Source/Programs/Enterprise/Datasmith/DatasmithRevitExporter/Private/DatasmithRevitDirectLink.cs b/Engine/Source/Programs/Enterprise/Datasmith/DatasmithRevitExporter/Private/DatasmithRevitDirectLink.cs
--- a/Engine/Source/Programs/Enterprise/Datasmith/DatasmithRevitExporter/Private/DatasmithRevitDirectLink.cs
+++ b/Engine/Source/Programs/Enterprise/Datasmith/DatasmithRevitExporter/Private/DatasmithRevitDirectLink.cs
@@ -338,25 +338,43 @@ namespace DatasmithRevitExporter
 			
 				DirectLink.bHasChanges = true;
 
-				if (ModifiedElement.GetType() == typeof(RevitLinkInstance))
+				if (ModifiedElement != null)
 				{
-					DirectLink.ModifiedLinkedDocuments.Add((ModifiedElement as RevitLinkInstance).GetLinkDocument());
-				}
-				else
-				{
-					// Handles a case where Revit won't notify us about modified mullions and their transform remains obsolte, thus wrong.
-					ElementCategoryFilter Filter = new ElementCategoryFilter(BuiltInCategory.OST_CurtainWallMullions);
-					IList<ElementId> DependentElements = ModifiedElement.GetDependentElements(Filter);
-					if (DependentElements != null && DependentElements.Count > 0)
+					if (ModifiedElement.GetType() == typeof(RevitLinkInstance))
+					{
+						DirectLink.ModifiedLinkedDocuments.Add((ModifiedElement as RevitLinkInstance).GetLinkDocument());
+					}
+					else if (ModifiedElement.GetType() == typeof(RevitLinkType))
 					{
-						foreach (ElementId DepElemId in DependentElements)
+						foreach (KeyValuePair<ElementId, FCachedDocumentData> Link in DirectLink.RootCache.LinkedDocumentsCache)
 						{
-							DirectLink.RootCache.SetElementModified(true, DepElemId);
+							RevitLinkInstance LinkInstance = DirectLink.RootCache.SourceDocument.GetElement(Link.Key) as RevitLinkInstance;
+							if (LinkInstance != null)
+							{
+								RevitLinkType LinkType = DirectLink.RootCache.SourceDocument.GetElement(LinkInstance.GetTypeId()) as RevitLinkType;
+								if (LinkType != null && LinkType.Id == ModifiedElement.Id && RevitLinkType.IsLoaded(DirectLink.RootCache.SourceDocument, LinkType.Id))
+								{
+									Link.Value.SetAllElementsModified();
+								}
+							}
+						}
+					}
+					else
+					{
+						// Handles a case where Revit won't notify us about modified mullions and their transform remains obsolte, thus wrong.
+						ElementCategoryFilter Filter = new ElementCategoryFilter(BuiltInCategory.OST_CurtainWallMullions);
+						IList<ElementId> DependentElements = ModifiedElement.GetDependentElements(Filter);
+						if (DependentElements != null && DependentElements.Count > 0)
+						{
+							foreach (ElementId DepElemId in DependentElements)
+							{
+								DirectLink.RootCache.SetElementModified(true, DepElemId);
+							}
 						}
 					}
-				}
 
-				DirectLink.RootCache.SetElementModified(true, ElemId);
+					DirectLink.RootCache.SetElementModified(true, ElemId);
+				}
 			}
 
 			if (DirectLink.bHasChanges && bAutoSync)
@@ -836,6 +854,11 @@ namespace DatasmithRevitExporter
 						continue;
 					}
 
+					if (!CacheData.SourceDocument.IsValidObject)
+					{
+						return;
+					}
+
 					Element RevitElement = CacheData.SourceDocument.GetElement(Entry.Key);
 
 					if (RevitElement == null)
diff --git a/Engine/Source/Programs/Enterprise/Datasmith/DatasmithRevitExporter/Private/DatasmithRevitDocumentData.cs b/Engine/Source/Programs/Enterprise/Datasmith/DatasmithRevitExporter/Private/DatasmithRevitDocumentData.cs
--- a/Engine/Source/Programs/Enterprise/Datasmith/DatasmithRevitExporter/Private/DatasmithRevitDocumentData.cs
+++ b/Engine/Source/Programs/Enterprise/Datasmith/DatasmithRevitExporter/Private/DatasmithRevitDocumentData.cs
@@ -1227,6 +1227,8 @@ namespace DatasmithRevitExporter
 							(ExistingActor as FDatasmithFacadeActorMesh)?.SetMesh(null);
 						}
 
+						ElementData.BaseElementType = InElement.Document.GetElement(InElement.GetTypeId()) as ElementType;
+						ElementData.CurrentElement = InElement;
 						ElementData.InitializePivotPlacement(ref InWorldTransform);
 						ElementData.InitializeElement(InWorldTransform, ElementData);
 						ElementData.MeshMaterialsMap.Clear();
