[3ds Max] DirectLink: MaxPhysicalMtl - CutOut map

diff --git a/Engine/Source/Programs/Enterprise/Datasmith/DatasmithMaxExporter/Private/MaxMaterialsToUEPbr/DatasmithMaxPhysicalMaterialToUEPbr.cpp b/Engine/Source/Programs/Enterprise/Datasmith/DatasmithMaxExporter/Private/MaxMaterialsToUEPbr/DatasmithMaxPhysicalMaterialToUEPbr.cpp
--- a/Engine/Source/Programs/Enterprise/Datasmith/DatasmithMaxExporter/Private/MaxMaterialsToUEPbr/DatasmithMaxPhysicalMaterialToUEPbr.cpp
+++ b/Engine/Source/Programs/Enterprise/Datasmith/DatasmithMaxExporter/Private/MaxMaterialsToUEPbr/DatasmithMaxPhysicalMaterialToUEPbr.cpp
@@ -295,7 +295,6 @@ struct FMaxPhysicalMaterial
 	bool bTransparencyRoughessMapOn = true;
 	bool bRoughnessInverted = true;
 	bool bThinWalled = true;
-	bool bCutoutMapOn = true;
 
 	float DiffuseWeight = 0, Roughness = 0, Metalness = 0, Ior = 0;
 	float Transparency = 0;
@@ -317,7 +316,7 @@ struct FMaxPhysicalMaterial
 	DatasmithMaxTexmapParser::FMapParameter ReflectionColorMap;
 	DatasmithMaxTexmapParser::FMapParameter EmittanceMap;
 	DatasmithMaxTexmapParser::FMapParameter EmittanceColorMap;
-	Texmap* CutoutMap = NULL;
+	DatasmithMaxTexmapParser::FMapParameter CutoutMap;
 
 	int MaterialMode = 0; //0 means simple, just metalness, 1 means specular + metalness
 
@@ -446,7 +445,7 @@ struct FMaxPhysicalMaterial
 				{
 					if (ParamBlock2->GetInt(ParamDefinition.ID, GetCOREInterface()->GetTime()) == 0)
 					{
-						bCutoutMapOn = false;
+						CutoutMap.bEnabled = false;
 					}
 				}
 
@@ -561,7 +560,7 @@ struct FMaxPhysicalMaterial
 				}
 				else if (FCString::Stricmp(ParamDefinition.int_name, TEXT("cutout_map")) == 0)
 				{
-					CutoutMap = ParamBlock2->GetTexmap(ParamDefinition.ID, GetCOREInterface()->GetTime());
+					CutoutMap.Map = ParamBlock2->GetTexmap(ParamDefinition.ID, GetCOREInterface()->GetTime());
 				}
 			}
 
@@ -708,6 +707,12 @@ void FDatasmithMaxPhysicalMaterialToUEPbr::Convert( TSharedRef< IDatasmithScene
 		}
 	}
 
+	IDatasmithMaterialExpression* CutOutExpression = nullptr;
+	if (PhysicalMaterialProperties.CutoutMap.IsMapPresentAndEnabled())
+	{
+		CutOutExpression = COMPOSE_OR_NULL(Desaturate, ConvertTexmap(PhysicalMaterialProperties.CutoutMap));
+	}
+
 	Connect(PbrMaterialElement->GetEmissiveColor(), 
 		ApplyWeightExpression(
 			TextureOrColor(TEXT("Emissive Color"), PhysicalMaterialProperties.EmittanceColorMap, FDatasmithMaxMatHelper::MaxLinearColorToFLinearColor(PhysicalMaterialProperties.EmittanceColor)),
@@ -761,8 +766,7 @@ void FDatasmithMaxPhysicalMaterialToUEPbr::Convert( TSharedRef< IDatasmithScene
 	// Build Material based on estimated material type
 	if (TransparencyExpression)
 	{
-		IDatasmithMaterialExpression& TransparencyExpressionRef = *TransparencyExpression;
-
+		IDatasmithMaterialExpression& OpacityFromTransparency = OneMinus(Desaturate(*TransparencyExpression));
 		// ThinTranslucent, including when Ior is nearly 1
 		if (PhysicalMaterialProperties.bThinWalled || FMath::IsNearlyEqual(PhysicalMaterialProperties.Ior, 1.0f))
 		{
@@ -772,7 +776,7 @@ void FDatasmithMaxPhysicalMaterialToUEPbr::Convert( TSharedRef< IDatasmithScene
 
 			TransparencyExpression->ConnectExpression(*ThinTranslucencyMaterialOutput->GetInput(0));
 
-			Connect(PbrMaterialElement->GetOpacity(), OneMinus(Desaturate(TransparencyExpressionRef)));
+			Connect(PbrMaterialElement->GetOpacity(), COMPOSE_OR_DEFAULT2(&OpacityFromTransparency, Multiply, &OpacityFromTransparency, CutOutExpression));
 		}
 		else
 		{
@@ -781,9 +785,12 @@ void FDatasmithMaxPhysicalMaterialToUEPbr::Convert( TSharedRef< IDatasmithScene
 				// Ported FDatasmithMaterialExpressions::AddGlassNode to compute base color and opacity
 				float Ior = PhysicalMaterialProperties.Ior;
 
-				IDatasmithMaterialExpression& Multiply09 = Multiply(Scalar(0.9), TransparencyExpressionRef);
+				IDatasmithMaterialExpression& Multiply09 = Multiply(Scalar(0.9), *TransparencyExpression);
+
+				// Modulate opacity by fresnel 
+				IDatasmithMaterialExpression& OpacityWithFresnelExpression = OneMinus(Desaturate(CalcIORComplex(Ior, 0, Multiply(Scalar(0.5), *TransparencyExpression), Multiply09)));
 
-				Connect(PbrMaterialElement->GetOpacity(), OneMinus(Desaturate(CalcIORComplex(Ior, 0, Multiply(Scalar(0.5), TransparencyExpressionRef), Multiply09))));
+				Connect(PbrMaterialElement->GetOpacity(), COMPOSE_OR_DEFAULT2(&OpacityWithFresnelExpression, Multiply, &OpacityWithFresnelExpression, CutOutExpression));
 				Connect(PbrMaterialElement->GetRefraction(), Lerp(Scalar(1.0f), Scalar(PhysicalMaterialProperties.Ior), Fresnel()));
 
 				// todo: may add Diffuse color into the mix to accomodate non-clear opacity
@@ -791,10 +798,15 @@ void FDatasmithMaxPhysicalMaterialToUEPbr::Convert( TSharedRef< IDatasmithScene
 			}
 			else
 			{
-				Connect(PbrMaterialElement->GetOpacity(), OneMinus(Desaturate(TransparencyExpressionRef)));
+				Connect(PbrMaterialElement->GetOpacity(), COMPOSE_OR_DEFAULT2(&OpacityFromTransparency, Multiply, &OpacityFromTransparency, CutOutExpression));
 			}
 		}
 	}
+	else
+	{
+		Connect(PbrMaterialElement->GetOpacity(),  CutOutExpression);
+	}
+
 	Connect(PbrMaterialElement->GetBaseColor(), BaseColorExpression);
 
 	FMaxPhysicalMaterialClearCoat ClearCoatProperties;
