[3ds Max] DirectLink: Messages Window rewrite with Slate

Fixed(compared to previous):
- Selection
- Copy(to clipboard) Button
- Colored message types
- Clear button
- Resizing
- Tooltips

diff --git a/Engine/Source/Programs/Enterprise/Datasmith/DatasmithMaxExporter/DatasmithMax2017.Build.cs b/Engine/Source/Programs/Enterprise/Datasmith/DatasmithMaxExporter/DatasmithMax2017.Build.cs
--- a/Engine/Source/Programs/Enterprise/Datasmith/DatasmithMaxExporter/DatasmithMax2017.Build.cs
+++ b/Engine/Source/Programs/Enterprise/Datasmith/DatasmithMaxExporter/DatasmithMax2017.Build.cs
@@ -11,6 +11,10 @@ namespace UnrealBuildTool.Rules
 
 			bUseRTTI = true;
 
+			// To avoid clashes with Max SDK
+			// todo: separate Slate code from 3ds max to a module
+			bUseUnity = false; 
+
 
 			PublicDefinitions.Add("NEW_DIRECTLINK_PLUGIN=1");
 
@@ -22,6 +26,9 @@ namespace UnrealBuildTool.Rules
 
 					"UdpMessaging", // required for DirectLink networking
 					"UEOpenExr",
+
+					"Slate",
+					"SlateCore",
 				}
 			);
 
diff --git a/Engine/Source/Programs/Enterprise/Datasmith/DatasmithMaxExporter/Private/DatasmithMaxDirectLinkLogging.cpp b/Engine/Source/Programs/Enterprise/Datasmith/DatasmithMaxExporter/Private/DatasmithMaxDirectLinkLogging.cpp
--- a/Engine/Source/Programs/Enterprise/Datasmith/DatasmithMaxExporter/Private/DatasmithMaxDirectLinkLogging.cpp
+++ b/Engine/Source/Programs/Enterprise/Datasmith/DatasmithMaxExporter/Private/DatasmithMaxDirectLinkLogging.cpp
@@ -35,6 +35,7 @@ void LogDebug(const TCHAR* Msg)
 	mprintf(L"[%s]%s\n", *FDateTime::UtcNow().ToString(TEXT("%Y.%m.%d-%H.%M.%S:%s")), Msg);
 	UE_LOG(LogDatasmithMaxExporter, Error, TEXT("%s"), Msg);
 	LogFlush();
+	LogDebugDialog(Msg);
 #endif
 }
 
@@ -50,20 +51,31 @@ void LogWarning(const TCHAR* Msg)
 	LogWarningDialog(Msg);
 }
 
-void LogFlush()
+void LogCompletion(const TCHAR* Msg)
 {
-	Async(EAsyncExecution::TaskGraphMainThread,
-		[]()
-		{
-			GLog->FlushThreadedLogs();
-			GLog->Flush();
-		});
+	UE_LOG(LogDatasmithMaxExporter, Display, TEXT("%s"), Msg);
+	LogCompletionDialog(Msg);
+}
+
+void LogError(const TCHAR* Msg)
+{
+	UE_LOG(LogDatasmithMaxExporter, Display, TEXT("%s"), Msg);
+	LogErrorDialog(Msg);
 }
 
+void LogError(const FString& Msg)
+{
+	LogError(*Msg);
+}
 
-void LogDebug(const FString& Msg)
+void LogWarning(const FString& Msg)
 {
-	LogDebug(*Msg);
+	LogWarning(*Msg);
+}
+
+void LogCompletion(const FString& Msg)
+{
+	LogCompletion(*Msg);
 }
 
 void LogInfo(const FString& Msg)
@@ -71,9 +83,20 @@ void LogInfo(const FString& Msg)
 	LogInfo(*Msg);
 }
 
-void LogWarning(const FString& Msg)
+void LogDebug(const FString& Msg)
 {
-	LogWarning(*Msg);
+	LogDebug(*Msg);
+}
+
+
+void LogFlush()
+{
+	Async(EAsyncExecution::TaskGraphMainThread,
+		[]()
+		{
+			GLog->FlushThreadedLogs();
+			GLog->Flush();
+		});
 }
 
 FString GetNodeDesc(INode* Node)
diff --git a/Engine/Source/Programs/Enterprise/Datasmith/DatasmithMaxExporter/Private/DatasmithMaxDirectLinkMaxScript.cpp b/Engine/Source/Programs/Enterprise/Datasmith/DatasmithMaxExporter/Private/DatasmithMaxDirectLinkMaxScript.cpp
--- a/Engine/Source/Programs/Enterprise/Datasmith/DatasmithMaxExporter/Private/DatasmithMaxDirectLinkMaxScript.cpp
+++ b/Engine/Source/Programs/Enterprise/Datasmith/DatasmithMaxExporter/Private/DatasmithMaxDirectLinkMaxScript.cpp
@@ -12,6 +12,8 @@
 
 #include "Windows/AllowWindowsPlatformTypes.h"
 
+#include "UI/IDatasmith3dsMaxUI.h"
+
 
 MAX_INCLUDES_START
 	#include "impexp.h"
@@ -35,9 +37,34 @@ extern HINSTANCE HInstanceMax;
 namespace DatasmithMaxDirectLink
 {
 
-/************************************* MaxScript exports *********************************/
+class FDummyDatasmith3dsMaxUI: public Ui::IMessagesWindow
+{
+public:
+	virtual void OpenWindow() override
+	{
+	}
+
+	virtual void AddError(const FString& Message)  override
+	{
+	}
+	virtual void AddWarning(const FString& Message)  override
+	{
+	}
+	virtual void AddInfo(const FString& Message)  override
+	{
+	}
+	virtual void AddCompletion(const FString& Message)  override
+	{
+	}
+	virtual void ClearMessages()  override
+	{
+	}
+};
 
-void InitScripts();
+FDummyDatasmith3dsMaxUI DummyDatasmith3dsMaxUI;
+Ui::IMessagesWindow* MessagesDialog = &DummyDatasmith3dsMaxUI;
+
+/************************************* MaxScript exports *********************************/
 
 Value* OnLoad_cf(Value**, int);
 Primitive OnLoad_pf(_M("Datasmith_OnLoad"), OnLoad_cf);
@@ -52,10 +79,15 @@ Value* OnLoad_cf(Value **arg_list, int count)
 
 	const TCHAR* EnginePathUnreal = (const TCHAR*)pEnginePath->to_string();
 
-	InitScripts();
+	bool bResult = CreateExporter(bEnableUI, EnginePathUnreal);
 
+	if (bResult)
+	{
+		// Create Slate UI only when CreateExporter returns true(it needs to successfully init  engine loop there for Slate to work)
+		MessagesDialog = Ui::CreateMessagesWindow();
+	}
 
-	return bool_result(CreateExporter(bEnableUI, EnginePathUnreal));
+	return bool_result(bResult);
 }
 
 Value* OnUnload_cf(Value**, int);
@@ -307,7 +339,7 @@ Primitive Crash_pf(_M("Datasmith_Crash"), Crash_cf);
 
 Value* LogInfo_cf(Value** arg_list, int count)
 {
-	check_arg_count(CreateScene, 1, count);
+	check_arg_count(LogInfo, 1, count);
 	Value* Message = arg_list[0];
 
 	LogInfo(Message->to_string());
@@ -316,6 +348,39 @@ Value* LogInfo_cf(Value** arg_list, int count)
 }
 Primitive LogInfo_pf(_M("Datasmith_LogInfo"), LogInfo_cf);
 
+Value* LogWarning_cf(Value** arg_list, int count)
+{
+	check_arg_count(LogWarning, 1, count);
+	Value* Message = arg_list[0];
+
+	LogWarning(Message->to_string());
+
+	return bool_result(true);
+}
+Primitive LogWarning_pf(_M("Datasmith_LogWarning"), LogWarning_cf);
+
+Value* LogError_cf(Value** arg_list, int count)
+{
+	check_arg_count(LogError, 1, count);
+	Value* Message = arg_list[0];
+
+	LogError(Message->to_string());
+
+	return bool_result(true);
+}
+Primitive LogError_pf(_M("Datasmith_LogError"), LogError_cf);
+
+Value* LogCompletion_cf(Value** arg_list, int count)
+{
+	check_arg_count(LogCompletion, 1, count);
+	Value* Message = arg_list[0];
+
+	LogCompletion(Message->to_string());
+
+	return bool_result(true);
+}
+Primitive LogCompletion_pf(_M("Datasmith_LogCompletion"), LogCompletion_cf);
+
 
 class FMessagesDialog
 {
@@ -429,21 +494,29 @@ public:
 	TQueue<FString> Messages;
 };
 
-TUniquePtr<FMessagesDialog> MessagesDialog;
-
-void LogInfoDialog(const TCHAR* Msg)
+void LogErrorDialog(const FString& Msg)
 {
-	MessagesDialog->AddInfo(Msg);
+	MessagesDialog->AddError(Msg);
 }
 
 void LogWarningDialog(const FString& Msg)
 {
-	LogWarningDialog(*Msg);
+	MessagesDialog->AddWarning(Msg);
+}
+
+void LogCompletionDialog(const FString& Msg)
+{
+	MessagesDialog->AddCompletion(Msg);
 }
 
-void LogWarningDialog(const TCHAR* Msg)
+void LogInfoDialog(const FString& Msg)
 {
-	MessagesDialog->AddWarning(Msg);
+	MessagesDialog->AddInfo(Msg);
+}
+
+void LogDebugDialog(const FString& Msg)
+{
+	MessagesDialog->AddInfo(Msg);
 }
 
 // Setup ActionTable with Datasmith commands exposed as actions
@@ -568,7 +641,7 @@ public:
 			}
 			case ID_SHOWLOG_ACTION_ID:
 			{
-				MessagesDialog->Toggle();
+				MessagesDialog->OpenWindow();
 				return true;
 			}
 			}
@@ -614,14 +687,8 @@ private:
 
 TUniquePtr<FDatasmithActions> Actions;
 
-void InitScripts()
-{
-	MessagesDialog.Reset(new FMessagesDialog);
-}
-
 void ShutdownScripts()
 {
-	MessagesDialog.Reset();
 	Actions.Reset();
 }
 
@@ -636,20 +703,6 @@ Value* SetupActions_cf(Value** arg_list, int count)
 }
 Primitive SetupActions_pf(_M("Datasmith_SetupActions"), SetupActions_cf);
 
-
-Value* AddWarning_cf(Value**, int);
-Primitive AddWarning_pf(_M("Datasmith_AddWarning"), AddWarning_cf);
-
-Value* AddWarning_cf(Value** arg_list, int count)
-{
-	check_arg_count(AddWarning, 1, count);
-
-	MessagesDialog->AddWarning(arg_list[0]->to_string());
-
-	return bool_result(true);
-}
-
-
 }
 
 #include "Windows/HideWindowsPlatformTypes.h"
diff --git a/Engine/Source/Programs/Enterprise/Datasmith/DatasmithMaxExporter/Private/DatasmithMaxExporterWithDirectLink.cpp b/Engine/Source/Programs/Enterprise/Datasmith/DatasmithMaxExporter/Private/DatasmithMaxExporterWithDirectLink.cpp
--- a/Engine/Source/Programs/Enterprise/Datasmith/DatasmithMaxExporter/Private/DatasmithMaxExporterWithDirectLink.cpp
+++ b/Engine/Source/Programs/Enterprise/Datasmith/DatasmithMaxExporter/Private/DatasmithMaxExporterWithDirectLink.cpp
@@ -346,7 +346,7 @@ public:
 
 	void PrintStage(FStage& Stage, FString Indent=TEXT(""))
 	{
-		LogInfo(Indent + FString::Printf(TEXT("    %s - %s"), *Stage.Name, *(Stage.TimeFinish-Stage.TimeStart).ToString()));
+		LogCompletion(Indent + FString::Printf(TEXT("    %s - %s"), *Stage.Name, *(Stage.TimeFinish-Stage.TimeStart).ToString()));
 		if (!Stage.Result.IsEmpty())
 		{
 			LogInfo(Indent + TEXT("      #") + Stage.Result);
@@ -2747,7 +2747,7 @@ public:
 
 		if (Options.bStatSync)
 		{
-			LogInfo("Sync completed:");
+			LogCompletion("Sync completed:");
 			ProgressManager.PrintStatisticss();
 		}
 	}
@@ -2771,7 +2771,7 @@ public:
 
 		if (Options.bStatSync)
 		{
-			LogInfo("AutoSync completed:");
+			LogCompletion("AutoSync completed:");
 			ProgressManager.PrintStatisticss();
 		}
 	}
diff --git a/Engine/Source/Programs/Enterprise/Datasmith/DatasmithMaxExporter/Private/UI/Datasmith3dsMaxUI.cpp b/Engine/Source/Programs/Enterprise/Datasmith/DatasmithMaxExporter/Private/UI/Datasmith3dsMaxUI.cpp
new file mode 100644
--- /dev/null
+++ b/Engine/Source/Programs/Enterprise/Datasmith/DatasmithMaxExporter/Private/UI/Datasmith3dsMaxUI.cpp
@@ -0,0 +1,374 @@
+// Copyright Epic Games, Inc. All Rights Reserved.
+
+#include "CoreMinimal.h"
+
+#include "IDatasmith3dsMaxUI.h"
+
+#include "Containers/Queue.h"
+#include "CoreGlobals.h"
+#include "Delegates/Delegate.h"
+#include "Framework/Application/SWindowTitleBar.h"
+#include "Framework/Application/SlateApplication.h"
+#include "HAL/Platform.h"
+#include "Misc/ConfigCacheIni.h"
+#include "Misc/ScopeLock.h"
+#include "SlateCore/Public/Widgets/DeclarativeSyntaxSupport.h"
+#include "Styling/CoreStyle.h"
+#include "Templates/SharedPointer.h"
+#include "Widgets/Input/SButton.h"
+#include "Widgets/Layout/SBorder.h"
+#include "Widgets/Layout/SScrollBox.h"
+#include "Widgets/SBoxPanel.h"
+#include "Widgets/SWidget.h"
+#include "Widgets/SWindow.h"
+#include "Widgets/Views/SListView.h"
+
+#include "HAL/PlatformApplicationMisc.h"  // For ClipboardCopy
+
+#include "DatasmithExporterManager.h"  // Using FDatasmithExporterManager::PushCommandIntoGameThread
+
+
+#define LOCTEXT_NAMESPACE "Datasmith3dsMaxUI"
+
+namespace DatasmithMaxDirectLink
+{
+namespace Ui
+{
+
+struct FMessageData
+{
+	enum ECategory
+	{
+		Default,
+		Info,
+		Debug,
+		Error,
+		Warning,
+		Completion
+	};
+
+	FString MessageText;
+	ECategory Category;
+};
+
+class SMessagesWidget : public SCompoundWidget
+{
+public:
+
+	SLATE_BEGIN_ARGS(SMessagesWidget) {}
+	SLATE_END_ARGS()
+
+	virtual bool SupportsKeyboardFocus() const override
+	{
+		// Handle keys
+		return true;
+	}
+
+	virtual void Tick(const FGeometry& AllottedGeometry, const double InCurrentTime, const float InDeltaTime) override
+	{
+		// Copy queued messages
+		if (!MessagesQueue->IsEmpty())
+		{
+			TSharedPtr<FMessageData> Message;
+			while (MessagesQueue->Dequeue(Message))
+			{
+				Messages.Add(Message);
+			}
+			MessagesView->RequestListRefresh();
+			MessagesView->ScrollToBottom();
+		}
+
+		// Always call base Widget Tick
+		SCompoundWidget::Tick(AllottedGeometry, InCurrentTime, InDeltaTime);
+	}
+
+	FReply OnKeyDown( const FGeometry& MyGeometry, const FKeyEvent& InKeyEvent )
+	{
+		if (InKeyEvent.IsControlDown())
+		{
+			if (InKeyEvent.GetKey() == EKeys::C)
+			{
+				CopyLog();
+			}
+		}
+		return FReply::Unhandled();
+	}
+
+	FReply OnCopyMessagesClicked()
+	{
+		CopyLog();
+		return FReply::Handled();
+	}
+
+	FReply OnClearMessagesClicked()
+	{
+		ClearMessages();
+		return FReply::Handled();
+	}
+
+	TSharedRef<ITableRow> OnGenerateRow(TSharedPtr<FMessageData> Item, const TSharedRef<STableViewBase>& Owner)
+	{
+		FSlateColor Color = FLinearColor::Gray;
+
+		switch (Item->Category)
+		{
+		case FMessageData::Default: break;
+		case FMessageData::Info: break;
+		case FMessageData::Debug: break;
+		case FMessageData::Error: 
+			Color = FLinearColor::Red;
+			break;
+		case FMessageData::Warning: 
+			Color = FLinearColor::Yellow;
+			break;
+		case FMessageData::Completion: 
+			Color = FLinearColor::Green;
+			break;
+		default: ;
+		}
+
+		return SNew(STableRow<TSharedPtr<FMessageData> >, Owner)
+		[
+			SNew(SHorizontalBox)
+
+			+ SHorizontalBox::Slot()
+			.AutoWidth()
+			.Padding(0)
+			[
+				SNew(STextBlock)
+				.Text(FText::FromString(Item->MessageText))
+				.ColorAndOpacity(Color)
+			]
+		];
+
+	}
+
+	void Construct(const FArguments& InArgs)
+	{
+		ChildSlot
+		[
+			SNew( SVerticalBox )
+
+			// Messages list view
+			+ SVerticalBox::Slot()
+			.Padding( 2.f )
+			.VAlign( VAlign_Fill )
+			[
+				SAssignNew( MessagesView, SListView<TSharedPtr<FMessageData>> )
+				.ItemHeight(32)
+				.ListItemsSource( &Messages )
+				.OnGenerateRow( this, &SMessagesWidget::OnGenerateRow )
+				.SelectionMode(ESelectionMode::Multi)
+			]
+
+			// Buttons
+			+ SVerticalBox::Slot()
+			.Padding( 2.f )
+			.AutoHeight()
+			.VAlign( VAlign_Bottom )
+			[
+				SNew( SHorizontalBox )
+				+ SHorizontalBox::Slot()
+				.Padding( 2.f )
+				.AutoWidth()
+				.HAlign( HAlign_Right )
+				[
+					SNew( SButton )
+					.OnClicked( this, &SMessagesWidget::OnCopyMessagesClicked )
+					.ToolTipText( LOCTEXT("CopyMessages_Tooltip", "Copy Messages To Clipboard") )
+					[
+						SNew( SVerticalBox )
+						+ SVerticalBox::Slot()
+						.VAlign( VAlign_Center )
+						[
+							SNew( STextBlock )
+							.Text( LOCTEXT("CopyMessages_Label", "Copy") )
+						]
+					]
+				]
+				+ SHorizontalBox::Slot()
+				.Padding( 2.f )
+				.AutoWidth()
+				.HAlign( HAlign_Right )
+				[
+					SNew( SButton )
+					.OnClicked( this, &SMessagesWidget::OnClearMessagesClicked )
+					.ToolTipText( LOCTEXT("ClearMessages_Tooltip", "Clear Messages") )
+					[
+						SNew( SVerticalBox )
+						+ SVerticalBox::Slot()
+						.VAlign( VAlign_Center )
+						[
+							SNew( STextBlock )
+							.Text( LOCTEXT("ClearMessages_Label", "Clear") )
+						]
+					]
+				]
+			]
+
+		];
+	}
+
+	void CopyLog( )
+	{
+		TArray<TSharedPtr<FMessageData>> SelectedItems = MessagesView->GetSelectedItems();
+		
+		if (SelectedItems.Num() > 0)
+		{
+			FString SelectedText;
+		
+			for( int32 Index = 0; Index < SelectedItems.Num(); ++Index )
+			{
+				SelectedText += SelectedItems[Index]->MessageText;
+				SelectedText += LINE_TERMINATOR;
+			}
+		
+			FPlatformApplicationMisc::ClipboardCopy( *SelectedText );
+		}
+	}
+
+	void ClearMessages()
+	{
+		check(IsInGameThread());
+		Messages.Empty();
+		MessagesView->RequestListRefresh();
+	}
+
+	TSharedPtr<SListView<TSharedPtr<FMessageData>>> MessagesView;
+
+	TArray<TSharedPtr<FMessageData>> Messages;
+	TQueue<TSharedPtr<FMessageData>>* MessagesQueue;
+};
+
+
+
+class FMessagesWindow : public IMessagesWindow
+{
+public:
+	virtual void OpenWindow() override
+	{
+		FSimpleDelegate RunOnUIThread;
+
+		RunOnUIThread.BindLambda( [this]()
+		{
+			if ( TSharedPtr<SWindow> SlateWindow = SlateWindowPtr.Pin() )
+			{
+				// Cancel the minimize
+				SlateWindow->BringToFront();
+
+				SlateWindow->HACK_ForceToFront();
+				SlateWindow->FlashWindow();
+			}
+			else
+			{
+				TSharedRef<SWindow> Window = SNew( SWindow )
+					.CreateTitleBar( false )
+					.ClientSize( FVector2D( 640, 480 ) )
+					.AutoCenter( EAutoCenter::PrimaryWorkArea )
+					.SizingRule( ESizingRule::UserSized )
+					.FocusWhenFirstShown( true )
+					.Title( LOCTEXT("DatasmithMessagesWindowTitle", "Datasmith Messages") );
+
+				TSharedRef<SWindowTitleBar> WindowTitleBar = SNew( SWindowTitleBar, Window, nullptr, GetWindowTitleAlignement() )
+					.Visibility( EVisibility::Visible )
+					.ShowAppIcon( false );
+
+				Window->SetTitleBar( WindowTitleBar );
+
+				Window->SetContent(
+						SNew(SVerticalBox)
+						+ SVerticalBox::Slot()
+						.AutoHeight()
+						[
+							WindowTitleBar
+						]
+						+ SVerticalBox::Slot()
+						.FillHeight( 1.f )
+						[
+							SNew( SBorder )
+							.BorderImage( FCoreStyle::Get().GetBrush( "ToolPanel.GroupBorder" ) )
+							[
+								SAssignNew(MessagesWidget, SMessagesWidget)
+							]
+						]
+					);
+
+				FSlateApplication::Get().AddWindow( Window, true );
+				Window->HACK_ForceToFront();
+				MessagesWidget->MessagesQueue = &MessagesQueue;
+				SlateWindowPtr = Window;
+			}
+		});
+
+#if IS_PROGRAM
+		FDatasmithExporterManager::PushCommandIntoGameThread(MoveTemp(RunOnUIThread));
+#else
+		RunOnUIThread.Execute();
+#endif
+	}
+
+	TQueue<TSharedPtr<FMessageData>> MessagesQueue;
+	void AddMessage(const FString& Message, FMessageData::ECategory Category)
+	{
+		MessagesQueue.Enqueue(MakeShareable(new FMessageData{Message, Category}));
+	}
+
+	virtual void AddError(const FString& Message) override
+	{
+		AddMessage(Message, FMessageData::Error);
+	}
+
+	virtual void AddWarning(const FString& Message) override
+	{
+		AddMessage(Message, FMessageData::Warning);
+	}
+
+	virtual void AddInfo(const FString& Message) override
+	{
+		AddMessage(Message, FMessageData::Info);
+	}
+
+	virtual void AddCompletion(const FString& Message) override
+	{
+		AddMessage(Message, FMessageData::Completion);
+	}
+
+	virtual void ClearMessages() override
+	{
+		MessagesWidget->ClearMessages();
+	}
+
+	EHorizontalAlignment GetWindowTitleAlignement()
+	{
+		EWindowTitleAlignment::Type TitleAlignment = FSlateApplicationBase::Get().GetPlatformApplication()->GetWindowTitleAlignment();
+		EHorizontalAlignment TitleContentAlignment;
+
+		if ( TitleAlignment == EWindowTitleAlignment::Left )
+		{
+			TitleContentAlignment = HAlign_Left;
+		}
+		else if ( TitleAlignment == EWindowTitleAlignment::Center )
+		{
+			TitleContentAlignment = HAlign_Center;
+		}
+		else
+		{
+			TitleContentAlignment = HAlign_Right;
+		}
+		return TitleContentAlignment;
+	}
+
+	TWeakPtr<SWindow> SlateWindowPtr;
+
+	TSharedPtr<SMessagesWidget> MessagesWidget;
+};
+
+IMessagesWindow* CreateMessagesWindow()
+{
+	static FMessagesWindow MessagesWindow;
+	return &MessagesWindow;
+}
+
+}
+}
+#undef LOCTEXT_NAMESPACE
diff --git a/Engine/Source/Programs/Enterprise/Datasmith/DatasmithMaxExporter/Private/UI/IDatasmith3dsMaxUI.h b/Engine/Source/Programs/Enterprise/Datasmith/DatasmithMaxExporter/Private/UI/IDatasmith3dsMaxUI.h
new file mode 100644
--- /dev/null
+++ b/Engine/Source/Programs/Enterprise/Datasmith/DatasmithMaxExporter/Private/UI/IDatasmith3dsMaxUI.h
@@ -0,0 +1,28 @@
+// Copyright Epic Games, Inc. All Rights Reserved.
+
+#pragma once
+
+namespace DatasmithMaxDirectLink
+{
+namespace Ui
+{
+
+class IMessagesWindow
+{
+public:
+
+	virtual void OpenWindow() = 0;
+
+	virtual void AddError(const FString& Message) = 0;
+	virtual void AddWarning(const FString& Message) = 0;
+	virtual void AddInfo(const FString& Message) = 0;
+	virtual void AddCompletion(const FString& Message) = 0;
+	virtual void ClearMessages() = 0;
+
+	virtual ~IMessagesWindow() = default;
+};
+
+IMessagesWindow* CreateMessagesWindow();
+
+}
+}
diff --git a/Engine/Source/Programs/Enterprise/Datasmith/DatasmithMaxExporter/Public/DatasmithMaxDirectLink.h b/Engine/Source/Programs/Enterprise/Datasmith/DatasmithMaxExporter/Public/DatasmithMaxDirectLink.h
--- a/Engine/Source/Programs/Enterprise/Datasmith/DatasmithMaxExporter/Public/DatasmithMaxDirectLink.h
+++ b/Engine/Source/Programs/Enterprise/Datasmith/DatasmithMaxExporter/Public/DatasmithMaxDirectLink.h
@@ -657,20 +657,27 @@ public:
 	#define LOG_DEBUG_HEAVY(message) 
 #endif
 
-void LogDebug(const TCHAR* Msg);
-void LogDebug(const FString& Msg);
-void LogInfo(const TCHAR* Msg);
-void LogInfo(const FString& Msg);
 void LogDebugNode(const FString& Name, class INode* Node);
 void LogNodeEvent(const MCHAR* Name, INodeEventCallback::NodeKeyTab& nodes);
 void LogFlush();
+
+void LogError(const TCHAR* Msg);
 void LogWarning(const TCHAR* Msg);
+void LogCompletion(const TCHAR* Msg);
+void LogInfo(const TCHAR* Msg);
+void LogDebug(const TCHAR* Msg);
+
+void LogError(const FString& Msg);
 void LogWarning(const FString& Msg);
+void LogCompletion(const FString& Msg);
+void LogInfo(const FString& Msg);
+void LogDebug(const FString& Msg);
 
-void LogWarningDialog(const TCHAR* Msg);
+void LogErrorDialog(const FString& Msg);
 void LogWarningDialog(const FString& Msg);
-void LogInfoDialog(const TCHAR* Msg);
-
+void LogCompletionDialog(const FString& Msg);
+void LogInfoDialog(const FString& Msg);
+void LogDebugDialog(const FString& Msg);
 
 }
 
