Solidworks datasmith exporter does not generate udatasmith file

diff --git a/Engine/Source/Programs/Enterprise/Datasmith/DatasmithSolidworksExporter/Private/DatasmithExporter.cs b/Engine/Source/Programs/Enterprise/Datasmith/DatasmithSolidworksExporter/Private/DatasmithExporter.cs
--- a/Engine/Source/Programs/Enterprise/Datasmith/DatasmithSolidworksExporter/Private/DatasmithExporter.cs
+++ b/Engine/Source/Programs/Enterprise/Datasmith/DatasmithSolidworksExporter/Private/DatasmithExporter.cs
@@ -450,6 +450,11 @@ namespace DatasmithSolidworks
 
 		public void ExportLevelVariantSets(List<FConfigurationData> InConfigs)
 		{
+			if (InConfigs == null)
+			{
+				return;
+			}
+
 			// Request existing VariantSet, or create a new one
 			FDatasmithFacadeLevelVariantSets LevelVariantSets = null;
 
diff --git a/Engine/Source/Programs/Enterprise/Datasmith/DatasmithSolidworksExporter/Private/Document.cs b/Engine/Source/Programs/Enterprise/Datasmith/DatasmithSolidworksExporter/Private/Document.cs
--- a/Engine/Source/Programs/Enterprise/Datasmith/DatasmithSolidworksExporter/Private/Document.cs
+++ b/Engine/Source/Programs/Enterprise/Datasmith/DatasmithSolidworksExporter/Private/Document.cs
@@ -113,7 +113,11 @@ namespace DatasmithSolidworks
 			List<FConfigurationData> Configs = FConfigurationExporter.ExportConfigurations(this);
 
 			Exporter.ExportMaterials(ExportedMaterialsMap);
-			Exporter.ExportLevelVariantSets(Configs);
+
+			if (Configs != null)
+			{
+				Exporter.ExportLevelVariantSets(Configs);
+			}
 		}
 
 		private void CheckForMaterialUpdatesProc()
