Revit - DirectLink frequently stops sending updates upon deleting scene entities

diff --git a/Engine/Source/Programs/Enterprise/Datasmith/DatasmithRevitExporter/Private/DatasmithRevitDirectLink.cs b/Engine/Source/Programs/Enterprise/Datasmith/DatasmithRevitExporter/Private/DatasmithRevitDirectLink.cs
--- a/Engine/Source/Programs/Enterprise/Datasmith/DatasmithRevitExporter/Private/DatasmithRevitDirectLink.cs
+++ b/Engine/Source/Programs/Enterprise/Datasmith/DatasmithRevitExporter/Private/DatasmithRevitDirectLink.cs
@@ -377,6 +377,9 @@ namespace DatasmithRevitExporter
 				}
 			}
 
+			DirectLink.bHasChanges = DirectLink.bHasChanges || InArgs.GetDeletedElementIds().Any();
+			DirectLink.bHasChanges = DirectLink.bHasChanges || InArgs.GetAddedElementIds().Any();
+
 			if (DirectLink.bHasChanges && bAutoSync)
 			{
 				DirectLink.RunAutoSync();
