Datasmith Revit: Compile Datasmith Revit Exporter - DatasmithRevitDocumentData.cs error: The type or namespace name could not be found

diff --git a/Engine/Source/Programs/Enterprise/Datasmith/DatasmithRevitExporter/DatasmithRevit2023/DatasmithRevit2023.csproj b/Engine/Source/Programs/Enterprise/Datasmith/DatasmithRevitExporter/DatasmithRevit2023/DatasmithRevit2023.csproj
--- a/Engine/Source/Programs/Enterprise/Datasmith/DatasmithRevitExporter/DatasmithRevit2023/DatasmithRevit2023.csproj
+++ b/Engine/Source/Programs/Enterprise/Datasmith/DatasmithRevitExporter/DatasmithRevit2023/DatasmithRevit2023.csproj
@@ -64,6 +64,9 @@
     <Compile Include="..\Private\DatasmithRevitDirectLink.cs">
       <Link>DatasmithRevitExporter\DatasmithRevitDirectLink.cs</Link>
     </Compile>
+    <Compile Include="..\Private\DatasmithRevitDecal.cs">
+      <Link>DatasmithRevitExporter\DatasmithRevitDecal.cs</Link>
+    </Compile>
     <Compile Include="..\Private\DatasmithRevitApplication.cs">
       <Link>DatasmithRevitExporter\DatasmithRevitApplication.cs</Link>
     </Compile>
