[3ds Max] DirectLink: Max 2023 support

diff --git a/Engine/Source/Programs/Enterprise/Datasmith/DatasmithMaxExporter/Plugin/PackageContents.xml b/Engine/Source/Programs/Enterprise/Datasmith/DatasmithMaxExporter/Plugin/PackageContents.xml
--- a/Engine/Source/Programs/Enterprise/Datasmith/DatasmithMaxExporter/Plugin/PackageContents.xml
+++ b/Engine/Source/Programs/Enterprise/Datasmith/DatasmithMaxExporter/Plugin/PackageContents.xml
@@ -16,7 +16,12 @@
 	<!--  todo: Icon="./DatasmithIcon32.png" -->
     <CompanyDetails Name="EpicGames, Inc." Url="www.epicgames.com"/>
 
-    <RuntimeRequirements OS="Win64" Platform="3ds Max" SeriesMin="2017" SeriesMax="2022" />
+    <RuntimeRequirements OS="Win64" Platform="3ds Max" SeriesMin="2017" SeriesMax="2023" />
+
+    <Components Description="plugins parts">
+        <RuntimeRequirements OS="Win64" Platform="3ds Max" SeriesMin="2023" SeriesMax="2023" />
+        <ComponentEntry ModuleName="./2023/DatasmithMax2023.gup" />
+    </Components>
 
     <Components Description="plugins parts">
         <RuntimeRequirements OS="Win64" Platform="3ds Max" SeriesMin="2022" SeriesMax="2022" />
@@ -49,18 +54,18 @@
     </Components>
 
     <Components Description="pre-start-up scripts parts">
-        <RuntimeRequirements OS="Win64" Platform="3ds Max" SeriesMin="2017" SeriesMax="2022" />
+        <RuntimeRequirements OS="Win64" Platform="3ds Max" SeriesMin="2017" SeriesMax="2023" />
         <ComponentEntry ModuleName="./UI/pre-start-up.ms" />
         <ComponentEntry ModuleName="./UI/UnrealStudio_Datasmith_Attributes_Holder_Modifier.ms" />
     </Components>
 
     <Components Description="macroscripts parts">
-        <RuntimeRequirements OS="Win64" Platform="3ds Max" SeriesMin="2017" SeriesMax="2022" />
+        <RuntimeRequirements OS="Win64" Platform="3ds Max" SeriesMin="2017" SeriesMax="2023" />
         <ComponentEntry ModuleName="./UI/DatasmithMacros.mcr" />
     </Components>
 
     <Components Description="post-start-up scripts parts">
-        <RuntimeRequirements OS="Win64" Platform="3ds Max" SeriesMin="2017" SeriesMax="2022" />
+        <RuntimeRequirements OS="Win64" Platform="3ds Max" SeriesMin="2017" SeriesMax="2023" />
         <ComponentEntry ModuleName="./UI/DatasmithDialog.ms" />
     </Components>
 
