[3ds Max] DirectLink: Cancel Export/Sync

- Sync can correctly restarts after being cancelled

diff --git a/Engine/Source/Programs/Enterprise/Datasmith/DatasmithMaxExporter/Private/DatasmithMaxConverters.cpp b/Engine/Source/Programs/Enterprise/Datasmith/DatasmithMaxExporter/Private/DatasmithMaxConverters.cpp
--- a/Engine/Source/Programs/Enterprise/Datasmith/DatasmithMaxExporter/Private/DatasmithMaxConverters.cpp
+++ b/Engine/Source/Programs/Enterprise/Datasmith/DatasmithMaxExporter/Private/DatasmithMaxConverters.cpp
@@ -29,7 +29,7 @@ void FHelperNodeConverter::ConvertToDatasmith(ISceneTracker& SceneTracker, FNode
 	SceneTracker.SetupActor(NodeTracker);
 
 	// todo: update validity from node
-	NodeTracker.NarrowValidityToInterval(FOREVER);
+	NodeTracker.Validity.NarrowValidityToInterval(FOREVER);
 }
 
 void FHelperNodeConverter::RemoveFromTracked(ISceneTracker& SceneTracker, FNodeTracker& NodeTracker)
@@ -58,7 +58,7 @@ void FCameraNodeConverter::ConvertToDatasmith(ISceneTracker& SceneTracker, FNode
 	NodeTracker.GetConverted().DatasmithActorElement->SetRotation(Rotation);
 
 	// todo: update validity from node
-	NodeTracker.NarrowValidityToInterval(FOREVER);
+	NodeTracker.Validity.NarrowValidityToInterval(FOREVER);
 }
 
 void FCameraNodeConverter::RemoveFromTracked(ISceneTracker& SceneTracker, FNodeTracker& NodeTracker)
@@ -130,7 +130,7 @@ void FLightNodeConverter::ConvertToDatasmith(ISceneTracker& SceneTracker, FNodeT
 		LightElement->SetRotation(Rotation);
 
 		// todo: update validity from node
-		NodeTracker.NarrowValidityToInterval(ValidityInterval);
+		NodeTracker.Validity.NarrowValidityToInterval(ValidityInterval);
 	}
 
 }
@@ -144,13 +144,10 @@ void FMeshNodeConverter::Parse(ISceneTracker& SceneTracker, FNodeTracker& NodeTr
 	ObjectState ObjState = NodeTracker.Node->EvalWorldState(SceneTracker.CurrentSyncPoint.Time);
 	Object* Obj = ObjState.obj;
 
-	if (Obj)
-	{
-		// AnimHandle is unique and never reused for new objects
-		InstanceHandle = Animatable::GetHandleByAnim(Obj); // Record instance handle into Converter to identify it later
+	// AnimHandle is unique and never reused for new objects
+	InstanceHandle = Animatable::GetHandleByAnim(Obj); // Record instance handle into Converter to identify it later
 
-		SceneTracker.AddGeometryNodeInstance(NodeTracker, *this, Obj);
-	}
+	SceneTracker.AddGeometryNodeInstance(NodeTracker, *this, Obj);
 }
 
 void FMeshNodeConverter::RemoveFromTracked(ISceneTracker& SceneTracker, FNodeTracker& NodeTracker)
@@ -168,7 +165,7 @@ void FMeshNodeConverter::ConvertToDatasmith(ISceneTracker& SceneTracker, FNodeTr
 
 	// Mark node as updated as soon as it is - in order for next nodes to be able to use its DatasmithActor
 	// todo: update validity from node
-	NodeTracker.NarrowValidityToInterval(FOREVER);
+	NodeTracker.Validity.NarrowValidityToInterval(FOREVER);
 }
 
 void FHismNodeConverter::RemoveFromTracked(ISceneTracker& SceneTracker, FNodeTracker& NodeTracker)
diff --git a/Engine/Source/Programs/Enterprise/Datasmith/DatasmithMaxExporter/Private/DatasmithMaxExporterWithDirectLink.cpp b/Engine/Source/Programs/Enterprise/Datasmith/DatasmithMaxExporter/Private/DatasmithMaxExporterWithDirectLink.cpp
--- a/Engine/Source/Programs/Enterprise/Datasmith/DatasmithMaxExporter/Private/DatasmithMaxExporterWithDirectLink.cpp
+++ b/Engine/Source/Programs/Enterprise/Datasmith/DatasmithMaxExporter/Private/DatasmithMaxExporterWithDirectLink.cpp
@@ -61,6 +61,8 @@ MAX_INCLUDES_START
 	#include "IFileResolutionManager.h" // for GetActualPath
 
 	#include "xref/iXrefObj.h"
+
+	#include "maxscript/maxscript.h"
 MAX_INCLUDES_END
 
 namespace DatasmithMaxDirectLink
@@ -660,25 +662,13 @@ public:
 class FInvalidatedNodeTrackers: FNoncopyable
 {
 public:
+
 	void Add(FNodeTracker& NodeTracker)
 	{
+		check(!IteratorUsageCount);
 		InvalidatedNodeTrackers.Add(&NodeTracker);
 	}
 
-	/**
-	 * @return if anything was deleted
-	 */
-	bool PurgeDeletedNodeTrackers(class FSceneTracker& Scene);
-
-	template <typename Func>
-	void EnumerateAll(Func Callable)
-	{
-		for (FNodeTracker* NodeTracker : InvalidatedNodeTrackers)
-		{
-			Callable(*NodeTracker);
-		}
-	}
-
 	int32 Num()
 	{
 		return InvalidatedNodeTrackers.Num();
@@ -686,18 +676,21 @@ public:
 
 	void Append(const TSet<FNodeTracker*>& NodeTrackers)
 	{
+		check(!IteratorUsageCount);
 		InvalidatedNodeTrackers.Append(NodeTrackers);
 	}
 
 	// Called when update is finished and all changes are processed and recorded
 	void Finish()
 	{
+		check(!IteratorUsageCount);
 		InvalidatedNodeTrackers.Reset();
 	}
 
 	// Scene is reset so invalidation is reset too
 	void Reset()
 	{
+		check(!IteratorUsageCount);
 		InvalidatedNodeTrackers.Reset();
 	}
 
@@ -708,11 +701,86 @@ public:
 
 	void RemoveFromInvalidated(FNodeTracker& NodeTracker)
 	{
+		check(!IteratorUsageCount);
 		InvalidatedNodeTrackers.Remove(&NodeTracker);
 	}
 
+	/**
+	 * @return if anything was deleted
+	 */
+	bool PurgeDeletedNodeTrackers(class FSceneTracker& Scene);
+
+	TArray<FNodeTracker*> Copy()
+	{
+		return InvalidatedNodeTrackers.Array();
+	}
+
+	// Iterate invalidated nodes checking that set of invalidated nodes is not changed during iteration
+	class TBaseIterator: private FNoncopyable
+	{
+	private:
+
+		typedef FNodeTracker ItElementType;
+
+	public:
+
+		typedef TSet<FNodeTracker*>::TRangedForIterator ElementItType;
+
+		FInvalidatedNodeTrackers& InvalidatedNodeTrackers;
+
+		FORCEINLINE TBaseIterator(FInvalidatedNodeTrackers& InInvalidatedNodeTrackers, const ElementItType& InElementIt)
+			: InvalidatedNodeTrackers(InInvalidatedNodeTrackers), ElementIt(InElementIt)
+		{
+			InvalidatedNodeTrackers.IteratorUsageCount++;
+		}
+
+		FORCEINLINE ~TBaseIterator()
+		{
+			InvalidatedNodeTrackers.IteratorUsageCount--;
+		}
+
+		FORCEINLINE TBaseIterator& operator++()
+		{
+			++Index;
+			++ElementIt;
+			return *this;
+		}
+
+		FORCEINLINE explicit operator bool() const
+		{ 
+			return !!ElementIt; 
+		}
+		
+		FORCEINLINE bool operator !() const 
+		{
+			return !static_cast<bool>(*this);
+		}
+
+		FORCEINLINE ItElementType* operator->() const
+		{
+			return *ElementIt;
+		}
+		FORCEINLINE ItElementType& operator*() const
+		{
+			return **ElementIt;
+		}
+
+		FORCEINLINE friend bool operator==(const TBaseIterator& Lhs, const TBaseIterator& Rhs) { return Lhs.ElementIt == Rhs.ElementIt; }
+		FORCEINLINE friend bool operator!=(const TBaseIterator& Lhs, const TBaseIterator& Rhs) { return Lhs.ElementIt != Rhs.ElementIt; }
+
+		ElementItType ElementIt;
+		int32 Index = 0; // For simpler debugging
+	};
+
+	using TRangedForIterator=TBaseIterator;
+
+	FORCEINLINE TRangedForIterator begin() { return TRangedForIterator(*this, InvalidatedNodeTrackers.begin()); }
+	FORCEINLINE TRangedForIterator end() { return TRangedForIterator(*this, InvalidatedNodeTrackers.end()); }
+
 private:
 	TSet<FNodeTracker*> InvalidatedNodeTrackers;
+
+	int32 IteratorUsageCount = 0;
 };
 
 class FNodeTrackersNames: FNoncopyable
@@ -771,6 +839,65 @@ public:
 	}
 };
 
+#ifdef CANCEL_DEBUG_ENABLE
+
+	// Cancel handling implementation that allows cancelling from Max Script (e.g. to trigger it faster than human can)
+	class FCancel
+	{
+	public:
+
+		void Reset()
+		{
+			bUseGetCancelMxsCallback = false;
+		}
+
+		bool GetCancelMxsCallback(bool& bSuccess)
+		{
+			MCHAR* ScriptCode = _T("Datasmith_GetCancel()");
+			FPValue Result;
+	#if MAX_PRODUCT_YEAR_NUMBER >= 2022
+			bSuccess = ExecuteMAXScriptScript(ScriptCode, MAXScript::ScriptSource::NonEmbedded, FALSE, &Result) != FALSE;
+	#else
+			bSuccess = ExecuteMAXScriptScript(ScriptCode, FALSE, &Result) != FALSE;
+	#endif
+			if (bSuccess)
+			{
+				check (Result.type == TYPE_BOOL);
+				return Result.b != FALSE;
+			}
+			return false;
+		}
+
+		bool GetCancel()
+		{
+			if (bUseGetCancelMxsCallback)
+			{
+				return GetCancelMxsCallback(bUseGetCancelMxsCallback); // Stop using callback if it fails(i.e. no mxs function defined or another error)
+			}
+
+			return GetCOREInterface()->GetCancel() != FALSE;
+		}
+
+	private:
+		bool bUseGetCancelMxsCallback = false;
+	};
+#else
+	class FCancel
+	{
+	public:
+
+		void Reset()
+		{
+		}
+
+		bool GetCancel()
+		{
+			return GetCOREInterface()->GetCancel() != FALSE;
+		}
+	};
+#endif
+
+
 
 // Holds states of entities for syncronization and handles change events
 class FSceneTracker: public ISceneTracker
@@ -787,14 +914,6 @@ public:
 		return ExportedScene.GetDatasmithScene().ToSharedRef();
 	}
 
-	bool ParseScene()
-	{
-		FIncludeXRefGuard IncludeXRefGuard(bIncludeXRefWhileParsing);
-		INode* Node = GetCOREInterface()->GetRootNode();
-		bSceneParsed = ParseScene(Node);
-		return bSceneParsed;
-	}
-
 	bool bParseXRefScenes = true;
 	bool bIncludeXRefWhileParsing = false; 
 
@@ -861,30 +980,37 @@ public:
 	FNodeTracker* ParseNode(INode* Node)
 	{
 		LogDebugNode(TEXT("ParseNode"), Node);
-
 		SCENE_UPDATE_STAT_INC(ParseNode, NodesEncountered);
 
 		FNodeTracker* NodeTracker =  GetNodeTracker(Node);
 
 		if (NodeTracker)
 		{
-			// Node being added might already be tracked(e.g. if it was deleted before but Update wasn't called to SceneTracker yet)
-			//ensure(NodeTracker->bDeleted);
-			NodeTracker->bDeleted = false;
-			InvalidateNode(*NodeTracker);
-			return NodeTracker;
+			InvalidateNode(*NodeTracker, false);
 		}
 		else
 		{
-			NodeTracker = AddNode(NodeEventNamespace::GetKeyByNode(Node), Node);
+			NodeTracker = AddNode(Node);
+		}
 
-			// Parse children
-			int32 ChildNum = Node->NumberOfChildren();
-			for (int32 ChildIndex = 0; ChildIndex < ChildNum; ++ChildIndex)
-			{
-				ParseNode(Node->GetChildNode(ChildIndex));
-			}
-			return NodeTracker;
+		// Parse children
+		int32 ChildNum = Node->NumberOfChildren();
+		for (int32 ChildIndex = 0; ChildIndex < ChildNum; ++ChildIndex)
+		{
+			ParseNode(Node->GetChildNode(ChildIndex));
+		}
+		return NodeTracker;
+	}
+
+	void ParseNode(FNodeTracker& NodeTracker)
+	{
+		if (NodeTracker.Node->IsRootNode())
+		{
+			ParseScene(NodeTracker.Node);
+		}
+		else
+		{
+			ParseNode(NodeTracker.Node);
 		}
 	}
 
@@ -988,6 +1114,18 @@ public:
 		return bResult;
 	}
 
+	FCancel Cancel;
+
+	bool GetCancel()
+	{
+		if (Cancel.GetCancel())
+		{
+			LogDebug("Update Cancelled");
+			return true;
+		}
+		return false;
+	}
+
 	bool UpdateInternalSafe(FUpdateProgress::FStage& MainStage)
 	{
 		bool bResult = false;
@@ -1002,6 +1140,7 @@ public:
 		return bResult;
 	}
 
+	 // returns if any change in scene was encountered and scene update completed(i.e. DirectLink Sync can be run)
 	bool UpdateInternal(FUpdateProgress::FStage& MainStage)
 	{
 
@@ -1011,20 +1150,15 @@ public:
 
 		Stats.Reset();
 
-		if (!bSceneParsed) // Parse whole scene only once
 		{
-			PROGRESS_STAGE("Parse Scene")
-			PROGRESS_STAGE_RESULT_DEFERRED
-			{
-				return FormatStatsParseScene();
-			};
-
-			ParseScene();
+			PROGRESS_STAGE("Refresh layers")
+			bChangeEncountered |= UpdateLayers();
 		}
 
+		Cancel.Reset();
+		if(GetCancel())
 		{
-			PROGRESS_STAGE("Refresh layers")
-			bChangeEncountered = UpdateLayers() && bChangeEncountered;
+			return false;	
 		}
 
 		// Changes present only when there are modified layers(changes checked manually), nodes(notified by Max) or materials(notified by Max with all changes in dependencies)
@@ -1041,7 +1175,7 @@ public:
 		}
 
 		{
-			PROGRESS_STAGE("Check Time Slider Validity")
+			PROGRESS_STAGE("Check Nodes Validity")
 			PROGRESS_STAGE_RESULT_DEFERRED
 			{
 				return FormatStatsCheckTimeSliderValidity();
@@ -1049,31 +1183,60 @@ public:
 
 			FIncludeXRefGuard IncludeXRefGuard(bIncludeXRefWhileParsing);
 
-			INode* SceneRootNode = GetCOREInterface()->GetRootNode();
+			FNodeTracker* NodeTracker = GetNodeTracker(GetCOREInterface()->GetRootNode());
+			if (NodeTracker)
+			{
+				InvalidateOutdatedNodeTracker(*NodeTracker);
+			}
+			else
+			{
+				// Add root node if it wasn't before(first Update)
+				AddNode(GetCOREInterface()->GetRootNode());
+			}
+		}
 
-			int32 ChildNum = SceneRootNode->NumberOfChildren();
-			for (int32 ChildIndex = 0; ChildIndex < ChildNum; ++ChildIndex)
+		{
+			PROGRESS_STAGE("Parse Invalidated Node Hierachy")
+			PROGRESS_STAGE_RESULT_DEFERRED
 			{
-				INode* Node = SceneRootNode->GetChildNode(ChildIndex);
+				return FormatStatsParseScene();
+			};
+
+			FIncludeXRefGuard IncludeXRefGuard(bIncludeXRefWhileParsing);
 
-				if (FNodeTracker* NodeTracker = GetNodeTracker(Node))
+			PROGRESS_STAGE_COUNTER(InvalidatedNodeTrackers.Num());
+			for (FNodeTracker* NodeTracker: InvalidatedNodeTrackers.Copy())
+			{
+				ProgressCounter.Next();
+				if(GetCancel())
 				{
-					InvalidateOutdatedNodeTracker(*NodeTracker);
+					return false;	
 				}
+
+				ParseNode(*NodeTracker);
 			}
 		}
 
 		{
 			PROGRESS_STAGE("Refresh collisions") // Update set of nodes used for collision
-			PROGRESS_STAGE_COUNTER(InvalidatedNodeTrackers.Num());
+			PROGRESS_STAGE_RESULT_DEFERRED
+			{
+				return FormatStatsRefreshCollisions();
+			};
 
+			PROGRESS_STAGE_COUNTER(InvalidatedNodeTrackers.Num());
 			TSet<FNodeTracker*> NodesWithChangedCollisionStatus; // Need to invalidate these nodes to make them renderable or to keep them from renderable depending on collision status
 
-			InvalidatedNodeTrackers.EnumerateAll([&](FNodeTracker& NodeTracker)
+			for (FNodeTracker& NodeTracker: InvalidatedNodeTrackers)
+			{
+				ProgressCounter.Next();
+				if(GetCancel())
 				{
-					ProgressCounter.Next();
-					UpdateCollisionStatus(NodeTracker, NodesWithChangedCollisionStatus);
-				});
+					return false;	
+				}
+
+				UpdateCollisionStatus(NodeTracker, NodesWithChangedCollisionStatus);
+			}
 
 			// Rebuild all nodes that has changed them being colliders
 			for (FNodeTracker* NodeTracker : NodesWithChangedCollisionStatus)
@@ -1083,10 +1246,6 @@ public:
 
 			SCENE_UPDATE_STAT_SET(RefreshCollisions, ChangedNodes, NodesWithChangedCollisionStatus.Num());
 
-			PROGRESS_STAGE_RESULT_DEFERRED
-			{
-				return FormatStatsRefreshCollisions();
-			};
 		}
 
 		{
@@ -1097,15 +1256,21 @@ public:
 			};
 
 			PROGRESS_STAGE_COUNTER(InvalidatedNodeTrackers.Num());
-			InvalidatedNodeTrackers.EnumerateAll([&](FNodeTracker& NodeTracker)
+
+			for (FNodeTracker& NodeTracker: InvalidatedNodeTrackers)
+			{
+				ProgressCounter.Next();
+				if(GetCancel())
 				{
-					ProgressCounter.Next();
-					UpdateNode(NodeTracker);
-				});
+					return false;	
+				}
+
+				UpdateNode(NodeTracker);
+			}
 		}
 
 		{
-			PROGRESS_STAGE("Process invalidated instances")
+			PROGRESS_STAGE("Process geometry")
 			PROGRESS_STAGE_RESULT_DEFERRED
 			{
 				return FormatStatsProcessInvalidatedInstances();
@@ -1114,6 +1279,11 @@ public:
 			for (FInstances* Instances : InvalidatedInstances)
 			{
 				ProgressCounter.Next();
+				if(GetCancel())
+				{
+					return false;	
+				}
+
 				UpdateInstances(*Instances);
 
 				// Need to re-convert and reattach all instances of an updated node
@@ -1131,16 +1301,20 @@ public:
 			};
 
 			PROGRESS_STAGE_COUNTER(InvalidatedNodeTrackers.Num());
-			InvalidatedNodeTrackers.EnumerateAll([&](FNodeTracker& NodeTracker)
+			for (FNodeTracker& NodeTracker: InvalidatedNodeTrackers)
+			{
+				ProgressCounter.Next();
+				if(GetCancel())
 				{
-					ProgressCounter.Next();
+					return false;	
+				}
 
-					if (NodeTracker.HasConverter())
-					{
-						SCENE_UPDATE_STAT_INC(ConvertNodes, Converted)
-						NodeTracker.GetConverter().ConvertToDatasmith(*this, NodeTracker);
-					}
-				});
+				if (NodeTracker.HasConverter())
+				{
+					SCENE_UPDATE_STAT_INC(ConvertNodes, Converted)
+					NodeTracker.GetConverter().ConvertToDatasmith(*this, NodeTracker);
+				}
+			}
 		}
 
 		{
@@ -1150,11 +1324,14 @@ public:
 				return FormatStatsReparentDatasmithActors();
 			};
 			
-			InvalidatedNodeTrackers.EnumerateAll([&](FNodeTracker& NodeTracker)
-				{
-					AttachNodeToDatasmithScene(NodeTracker);
-				});
-			;
+			for (FNodeTracker& NodeTracker: InvalidatedNodeTrackers)
+			{
+				// if(GetCancel())
+				// {
+				// 	return false;	
+				// }
+				AttachNodeToDatasmithScene(NodeTracker);
+			}
 		}
 
 		{
@@ -1162,26 +1339,17 @@ public:
 			// Finally mark all nodetrackers as valid(we didn't interrupt update as it went to this point)
 			// And calculate subtree validity for each subtree of each updated node
 
-			// Maximize each invalidated node's subtree validity interval before recalculating it from updated nodes
-			InvalidatedNodeTrackers.EnumerateAll([](FNodeTracker& NodeTracker)
-				{
-					NodeTracker.SubtreeValidity.Invalidate();
-					NodeTracker.SubtreeValidity.ResetValidityInterval();
-				});
+			bChangeEncountered |= InvalidatedNodeTrackers.HasInvalidated(); // Right before resetting invalidated nodes, record that anything was invalidated
 
-			// Recalculate subtree validity 
-			InvalidatedNodeTrackers.EnumerateAll([&](FNodeTracker& NodeTracker)
-				{
-					PromoteValidity(NodeTracker, NodeTracker.Validity);
-				});
+			// Recalculate subtree validity using updated nodes validity
+			UpdateSubtreeValidity(*GetNodeTracker(GetCOREInterface()->GetRootNode()));
 
-			InvalidatedNodeTrackers.EnumerateAll([](FNodeTracker& NodeTracker)
-				{
-					NodeTracker.SubtreeValidity.SetValid();
-					NodeTracker.SetValid();
-				});
+			for (FNodeTracker& NodeTracker: InvalidatedNodeTrackers)
+			{
+				NodeTracker.SubtreeValidity.SetValid();
+				NodeTracker.Validity.SetValid();
+			}
 
-			bChangeEncountered |= InvalidatedNodeTrackers.HasInvalidated(); // Right before resetting invalidated nodes, record that anything was invalidated
 			InvalidatedNodeTrackers.Finish();
 		}
 
@@ -1199,6 +1367,11 @@ public:
 			for (FMaterialTracker* MaterialTracker : MaterialsCollectionTracker.GetInvalidatedMaterials())
 			{
 				ProgressCounter.Next();
+				if(GetCancel())
+				{
+					return false;	
+				}
+
 				SCENE_UPDATE_STAT_INC(ProcessInvalidatedMaterials, Invalidated);
 
 				MaterialsCollectionTracker.UpdateMaterial(MaterialTracker);
@@ -1225,6 +1398,10 @@ public:
 			for (Mtl* ActualMaterial: ActualMaterialToUpdate)
 			{
 				ProgressCounter.Next();
+				if(GetCancel())
+				{
+					return false;	
+				}
 
 				MaterialsCollectionTracker.ConvertMaterial(ActualMaterial, ExportedScene.GetDatasmithScene().ToSharedRef(), ExportedScene.GetSceneExporter().GetAssetsOutputPath(), ActualTexmapsToUpdate);
 			}
@@ -1241,6 +1418,10 @@ public:
 			for (Texmap* Texture : ActualTexmapsToUpdate)
 			{
 				ProgressCounter.Next();
+				if(GetCancel())
+				{
+					return false;	
+				}
 				SCENE_UPDATE_STAT_INC(UpdateTextures, Total);
 
 				TArray<TSharedPtr<IDatasmithTextureElement>> TextureElements;
@@ -1316,7 +1497,7 @@ public:
 
 	FString FormatStatsProcessInvalidatedInstances()
 	{
-		return FString::Printf(TEXT("Instances: %ld updated"), InvalidatedInstances.Num());
+		return FString::Printf(TEXT("Unique meshes: %ld updated"), SCENE_UPDATE_STAT_GET(UpdateInstances, GeometryUpdated));
 	}
 
 	FString FormatStatsReparentDatasmithActors()
@@ -1372,18 +1553,6 @@ public:
 		}
 	}
 
-	FORCENOINLINE
-	FNodeTracker* AddNode(FNodeKey NodeKey, INode* Node)
-	{
-		LogDebugNode(TEXT("AddNode"), Node);
-		FNodeTracker* NodeTracker = NodeTrackers.Emplace(NodeKey, FNodeTrackerHandle(NodeKey, Node)).GetNodeTracker();
-
-		NodeTrackersNames.Add(*NodeTracker);
-		InvalidatedNodeTrackers.Add(*NodeTracker);
-
-		return NodeTracker;
-	}
-
 	virtual void RemoveMaterial(const TSharedPtr<IDatasmithBaseMaterialElement>& DatasmithMaterial) override
 	{
 		ExportedScene.DatasmithSceneRef->RemoveMaterial(DatasmithMaterial);
@@ -1467,38 +1636,82 @@ public:
 			}
 		}
 
-		if (NodeTracker.IsInvalidated())
+		if (NodeTracker.Validity.IsInvalidated())
 		{
 			// Don't do work twice - already invalidated node would have its subhierarchy invalidated
 			//   in case subhierarchy was changed for already invalidated node - the changed nodes would invalidate too(responding to reparent event)
 			return;
 		}
 
-		NodeTracker.Invalidate();
+		NodeTracker.Validity.Invalidate();
 		InvalidatedNodeTrackers.Add(NodeTracker);
+		PromoteValidity(NodeTracker, NodeTracker.Validity);
+	}
 
-		// Invalidate whole sub-hierarchy of nodes are now children.
-		// E.g. a node could have been hidden so its children were attached to grandparent(parent of hidden node)
-		// Need to invalidate those to reattach
-		FIncludeXRefGuard IncludeXRefGuard(bIncludeXRefWhileParsing);
-		int32 ChildNum = NodeTracker.Node->NumberOfChildren();
-		for (int32 ChildIndex = 0; ChildIndex < ChildNum; ++ChildIndex)
+	// todo: make fine invalidates - full only something like geometry change, but finer for transform, name change and more
+	FNodeTracker* InvalidateNode(FNodeKey NodeKey, bool bCheckCalledInProgress = true)
+	{
+		LogDebugNode(TEXT("InvalidateNode"), NodeEventNamespace::GetNodeByKey(NodeKey));
+		// We don't expect node chances while Update inprogress(unless Invalidate called explicitly)
+		if (bCheckCalledInProgress)
+		{
+			ensure(!bUpdateInProgress);
+		}
+
+		if (FNodeTracker* NodeTracker =  GetNodeTracker(NodeKey))
+		{
+			if (NodeEventNamespace::GetNodeByKey(NodeKey))
+			{
+				InvalidateNode(*NodeTracker, bCheckCalledInProgress);
+				return NodeTracker;
+			}
+			else
+			{
+				// Sometimes node update received without node Delete event
+				// Test case: create container, add node to it. Close it, open it, close again, then sync
+				InvalidatedNodeTrackers.Add(*NodeTracker);
+				NodeTracker->bDeleted = true;
+				PromoteValidity(*NodeTracker, NodeTracker->Validity);
+			}
+		}
+		else
 		{
-			InvalidateNode(NodeEventNamespace::GetKeyByNode(NodeTracker.Node->GetChildNode(ChildIndex)), bCheckCalledInProgress);
+			AddNode(NodeKey, NodeEventNamespace::GetNodeByKey(NodeKey));
 		}
+		return nullptr;
+	}
+
+	FNodeTracker* AddNode(FNodeKey NodeKey, INode* Node)
+	{
+		LogDebugNode(TEXT("AddNode"), Node);
+		check(!NodeTrackers.Contains(NodeKey));
+
+		FNodeTracker* NodeTracker = NodeTrackers.Emplace(NodeKey, FNodeTrackerHandle(NodeKey, Node)).GetNodeTracker();
+
+		NodeTrackersNames.Add(*NodeTracker);
+		InvalidatedNodeTrackers.Add(*NodeTracker);
+		PromoteValidity(*NodeTracker, NodeTracker->Validity);
+
+		return NodeTracker;
+	}
+
+	FNodeTracker* AddNode(INode* Node)
+	{
+		return AddNode(NodeEventNamespace::GetKeyByNode(Node), Node);
 	}
 
 	void InvalidateOutdatedNodeTracker(FNodeTracker& NodeTracker)
 	{
 		SCENE_UPDATE_STAT_INC(CheckTimeSlider, TotalChecks);
-		// Skip check in case whole subtree valid for current time OR node was already invalidated(and so its whole subtree)
-		if (NodeTracker.IsInvalidated())
+		// Skip node that was already invalidated(and so its whole subtree)
+		if (NodeTracker.Validity.IsInvalidated())
 		{
 			SCENE_UPDATE_STAT_INC(CheckTimeSlider, SkippedAsAlreadyInvalidated);
 			return;
 		}
 
-		if (NodeTracker.IsSubtreeValidForSyncPoint(CurrentSyncPoint))
+		// Skip check in case whole subtree valid for current time
+		if (NodeTracker.SubtreeValidity.IsValidForSyncPoint(CurrentSyncPoint))
 		{
 			SCENE_UPDATE_STAT_INC(CheckTimeSlider, SkippedAsSubtreeValid);
 			return;
@@ -1508,7 +1721,7 @@ public:
 		// todo: it it possible to optimize reevaluation of whole subtree?
 		//   certain types of invalidation need not to propagate down to descendants(e.g. geometry change)
 		//   but some need to, like transform change
-		if (!NodeTracker.IsValidForSyncPoint(CurrentSyncPoint))
+		if (!NodeTracker.Validity.IsValidForSyncPoint(CurrentSyncPoint))
 		{
 			SCENE_UPDATE_STAT_INC(CheckTimeSlider, Invalidated);
 			InvalidateNode(NodeTracker, false);
@@ -1526,41 +1739,40 @@ public:
 			}
 		}
 	}
-	
-	// todo: make fine invalidates - full only something like geometry change, but finer for transform, name change and more
-	FNodeTracker* InvalidateNode(FNodeKey NodeKey, bool bCheckCalledInProgress = true)
+
+	bool UpdateSubtreeValidity(FNodeTracker& NodeTracker)
 	{
-		FIncludeXRefGuard IncludeXRefGuard(bIncludeXRefWhileParsing);
+		// FIncludeXRefGuard IncludeXRefGuard(bIncludeXRefWhileParsing);
 
-		LogDebugNode(TEXT("InvalidateNode"), NodeEventNamespace::GetNodeByKey(NodeKey));
-		// We don't expect node chances while Update inprogress(unless Invalidate called explicitly)
-		if (bCheckCalledInProgress)
+		// Skip check in case whole subtree was valid for current time
+		if (NodeTracker.SubtreeValidity.IsValidForSyncPoint(CurrentSyncPoint))
 		{
-			ensure(!bUpdateInProgress);
+			return true;
 		}
 
-		if (FNodeTracker* NodeTracker =  GetNodeTracker(NodeKey))
+		NodeTracker.SubtreeValidity.Invalidate();
+		NodeTracker.SubtreeValidity.ResetValidityInterval();
+		NodeTracker.SubtreeValidity.NarrowValidityToInterval(NodeTracker.Validity);
+
+		int32 ChildNum = NodeTracker.Node->NumberOfChildren();
+		for (int32 ChildIndex = 0; ChildIndex < ChildNum; ++ChildIndex)
 		{
-			if (NodeEventNamespace::GetNodeByKey(NodeKey))
-			{
-				InvalidateNode(*NodeTracker, bCheckCalledInProgress);
-				return NodeTracker;
-			}
-			else
+			if (FNodeTracker* ChildNodeTracker = GetNodeTracker(NodeTracker.Node->GetChildNode(ChildIndex)))
 			{
-				// Sometimes node update received without node Delete event
-				// Test case: create container, add node to it. Close it, open it, close again, then sync
-				InvalidatedNodeTrackers.Add(*NodeTracker);
-				NodeTracker->bDeleted = true;
+				if (UpdateSubtreeValidity(*ChildNodeTracker))
+				{
+					NodeTracker.SubtreeValidity.NarrowValidityToInterval(ChildNodeTracker->SubtreeValidity);
+				}
+				else
+				{
+					NodeTracker.SubtreeValidity.Invalidate();
+					return false;
+				}
 			}
 		}
-		else
-		{
-			NodeAdded(NodeEventNamespace::GetNodeByKey(NodeKey));
-		}
-		return nullptr;
+		return true;
 	}
-
+	
 	void ClearNodeFromDatasmithScene(FNodeTracker& NodeTracker)
 	{
 		ReleaseNodeTrackerFromDatasmithMetadata(NodeTracker);
@@ -1762,7 +1974,7 @@ public:
 		ClearNodeFromDatasmithScene(NodeTracker);
 		RemoveFromTracked(NodeTracker); // todo: might keep it tracked by complicating conversions later(e.g. to avoid extra work if object stays the same)
 
-		NodeTracker.ResetValidityInterval(); // Make infinite validity to narrow it during update
+		NodeTracker.Validity.ResetValidityInterval(); // Make infinite validity to narrow it during update
 
 		ConvertNodeObject(NodeTracker);
 	}
@@ -1801,7 +2013,6 @@ public:
 
 		ObjectState ObjState = NodeTracker.Node->EvalWorldState(CurrentSyncPoint.Time);
 		Object* Obj = ObjState.obj;
-
 		if (!Obj)
 		{
 			return;
@@ -1934,6 +2145,7 @@ public:
 			//    - have single call to BeginEnumeration and EndEnumeration
 			//    - track all Begin'd nodes to End them together after all is updated(to prevent duplicated Begin's of referenced objects that might be shared by different noвes)
 			NodesPreparer.PrepareNode(NodeTracker.Node);
+			SCENE_UPDATE_STAT_INC(UpdateInstances, GeometryUpdated);
 			UpdateInstancesGeometry(Instances, NodeTracker);
 
 			// assign materials to static mesh for the first instance(others will use override on mesh actors)
@@ -2019,7 +2231,7 @@ public:
 		Rotation.Normalize();
 		ObjectTransform = FTransform(Rotation, Translation, Scale);
 
-		NodeTracker.NarrowValidityToInterval(ValidityInterval);
+		NodeTracker.Validity.NarrowValidityToInterval(ValidityInterval);
 	}
 
 	void RegisterNodeForMaterial(FNodeTracker& NodeTracker, Mtl* Material)
@@ -2386,7 +2598,18 @@ public:
 			NotificationsHandler->AddNode(Node);
 		}
 
-		ParseNode(Node);
+		FNodeKey NodeKey = NodeEventNamespace::GetKeyByNode(Node);
+
+		if (FNodeTracker* NodeTracker = GetNodeTracker(NodeKey))
+		{
+			// Re-added trackd node, probably deleted before(e.g. delete then undo)
+			InvalidateNode(*NodeTracker);
+			NodeTracker->bDeleted = false;
+		}
+		else
+		{
+			AddNode(Node);
+		}
 	}
 
 	virtual void NodeXRefMerged(INode* Node) override
@@ -2409,9 +2632,6 @@ public:
 
 		NotificationsHandler->AddNode(Node);
 
-		FNodeKey NodeKey = NodeEventNamespace::GetKeyByNode(Node);
-		InvalidateNode(NodeKey);
-
 		if (!bIncludeXRefWhileParsing)
 		{
 			ParseScene(Node, FXRefScene{SceneRootNode, XRefIndex}); // Parse xref hierarchy - it won't add itself! Or will it?
@@ -2426,6 +2646,7 @@ public:
 		{
 			InvalidatedNodeTrackers.Add(*NodeTracker);
 			NodeTracker->bDeleted = true;
+			PromoteValidity(*NodeTracker, NodeTracker->Validity);
 		}
 	}
 
@@ -2525,7 +2746,6 @@ public:
 	
 	void Reset()
 	{
-		bSceneParsed = false;
 
 		NodeTrackers.Reset();
 		NodeTrackersNames.Reset();
@@ -2554,7 +2774,6 @@ public:
 	bool bUpdateInProgress = false;
 
 	// Scene tracked/converted entities and connections beetween them
-	bool bSceneParsed = false;
 	TMap<FNodeKey, FNodeTrackerHandle> NodeTrackers; // All scene nodes
 	FNodeTrackersNames NodeTrackersNames; // Nodes grouped by name, for faster access
 	FInstancesManager InstancesManager; // Groups geometry nodes by their shared mesh
@@ -2606,11 +2825,6 @@ public:
 		ExportedScene.SetupScene();
 	}
 
-	virtual void ParseScene() override
-	{
-		SceneTracker.ParseScene();
-	}
-
 	virtual void InitializeDirectLinkForScene() override
 	{
 		if (DirectLinkImpl)
diff --git a/Engine/Source/Programs/Enterprise/Datasmith/DatasmithMaxExporter/Public/DatasmithMaxDirectLink.h b/Engine/Source/Programs/Enterprise/Datasmith/DatasmithMaxExporter/Public/DatasmithMaxDirectLink.h
--- a/Engine/Source/Programs/Enterprise/Datasmith/DatasmithMaxExporter/Public/DatasmithMaxDirectLink.h
+++ b/Engine/Source/Programs/Enterprise/Datasmith/DatasmithMaxExporter/Public/DatasmithMaxDirectLink.h
@@ -29,6 +29,7 @@ MAX_INCLUDES_END
 // Log all messages
 // #define LOG_DEBUG_HEAVY_ENABLE 1
 // #define LOG_DEBUG_ENABLE 1
+// #define CANCEL_DEBUG_ENABLE 1
 
 
 void AssignMeshMaterials(TSharedPtr<IDatasmithMeshElement>&MeshElement, Mtl * Material, const TSet<uint16>&SupportedChannels);
@@ -76,7 +77,6 @@ class IExporter
 	virtual void InitializeScene() = 0;
 
 	// Scene update
-	virtual void ParseScene() = 0;
 	virtual bool UpdateScene(bool bQuiet) = 0;
 
 	virtual void ResetSceneTracking() = 0;
@@ -133,9 +133,7 @@ struct FValidity
 
 	bool IsInvalidated() const
 	{
-		bool bResult = ValidityInterval.Empty() != 0; // Invalidated equals interval is empty
-		check(bIsInvalidated == bResult);
-		return bResult;
+		return bIsInvalidated;
 	}
 
 	void SetValid()
@@ -182,9 +180,10 @@ struct FValidity
 
 private:
 	Interval ValidityInterval;
-	// todo: this flag indicates, that 'invalidate' called specifically(no force rebuild
+	// todo: this flag indicates that ValidityInterval is determined to be recalculated
 	//  this is done in order to distinguish from Empty Interval that might happen somehow when updated
-	//  probably this is not needed - need to think about this and just rely on Empty for invalidated state 
+	//  Until ValidityInterval is recalculated bIsInvalidated stays - this supports cancelling Update
+	//  at any point. So on the next Update nodes with bIsInvalidated are updated(again, if their update wasn't finished)
 	bool bIsInvalidated; 
 };
 
@@ -198,44 +197,6 @@ class FNodeTracker: FNoncopyable
 public:
 	explicit FNodeTracker(FNodeKey InNodeKey, INode* InNode) : NodeKey(InNodeKey), Node(InNode) {}
 
-	void Invalidate()
-	{
-		Validity.Invalidate();
-	}
-
-	void SetValid()
-	{
-		Validity.SetValid();
-	}
-
-	void ResetValidityInterval()
-	{
-		Validity.ResetValidityInterval();
-	}
-
-	bool IsValidForSyncPoint(const FSyncPoint& State) const
-	{
-		return Validity.IsValidForSyncPoint(State);
-	}
-
-	bool IsSubtreeValidForSyncPoint(const FSyncPoint& State) const
-	{
-		return SubtreeValidity.IsValidForSyncPoint(State);
-	}
-
-	// Forcefully invalidated
-	bool IsInvalidated() const
-	{
-		return Validity.IsInvalidated();
-	}
-
-	// todo: rename. Intersect validity with currently set one
-	//   e.g. validity interval of a geometry should be intersection of validity of transform, mesh, etc
-	void NarrowValidityToInterval(Interval ValidityInterval) 
-	{
-		return Validity.NarrowValidityToInterval(ValidityInterval);
-	}
-
 	void SetXRefIndex(FXRefScene InXRefScene)
 	{
 		XRefScene = InXRefScene;
@@ -314,6 +275,7 @@ public:
 	FString Name;
 
 	// Node validity
+	bool bParsed = false; // 
 	bool bDeleted = false;
 	FValidity Validity;
 	FValidity SubtreeValidity;
@@ -385,6 +347,7 @@ struct FSceneUpdateStats
 	int32 CheckTimeSliderSkippedAsSubtreeValid = 0;
 	int32 CheckTimeSliderInvalidated = 0;
 	int32 ConvertNodesConverted = 0;
+	int32 UpdateInstancesGeometryUpdated = 0;
 };
 
 #define SCENE_UPDATE_STAT_INC(Category, Name) {Stats.##Category##Name++;}
