[3ds Max] DirectLink: cancel during scene parsing

- initial scene parsing of big scenes(or when lots of nodes are invalidated) might take some noticeable of time

diff --git a/Engine/Source/Programs/Enterprise/Datasmith/DatasmithMaxExporter/Private/DatasmithMaxExporterWithDirectLink.cpp b/Engine/Source/Programs/Enterprise/Datasmith/DatasmithMaxExporter/Private/DatasmithMaxExporterWithDirectLink.cpp
--- a/Engine/Source/Programs/Enterprise/Datasmith/DatasmithMaxExporter/Private/DatasmithMaxExporterWithDirectLink.cpp
+++ b/Engine/Source/Programs/Enterprise/Datasmith/DatasmithMaxExporter/Private/DatasmithMaxExporterWithDirectLink.cpp
@@ -383,7 +383,7 @@ public:
 	void Next()
 	{
 		double CurrentTime = FPlatformTime::Seconds();
-		if (CurrentTime - SecondsOfLastUpdate > UpdateIntervalMin) // Don't span progress bar
+		if (CurrentTime - SecondsOfLastUpdate > UpdateIntervalMin) // Don't spam progress bar
 		{
 			ProgressStage.ProgressEvent(float(Index) / Count, *FString::Printf(TEXT("%d of %d"), Index, Count) );
 			SecondsOfLastUpdate = CurrentTime;
@@ -848,7 +848,7 @@ public:
 
 		void Reset()
 		{
-			bUseGetCancelMxsCallback = false;
+			bUseGetCancelMxsCallback = true;
 		}
 
 		bool GetCancelMxsCallback(bool& bSuccess)
@@ -952,14 +952,15 @@ public:
 				if (FPaths::FileExists(Path) == false)
 				{
 					SCENE_UPDATE_STAT_INC(ParseScene, XRefFileMissing);
-					FString Error = FString("XRefScene file \"") + FPaths::GetCleanFilename(*Path) + FString("\" cannot be found");
-					// todo: logging
-					// DatasmithMaxLogger::Get().AddMissingAssetError(*Error);
+					LogWarning(FString::Printf(TEXT("XRefScene file \"%s\" cannot be found"), *FPaths::GetCleanFilename(*Path)));
 				}
 				else
 				{
 					SCENE_UPDATE_STAT_INC(ParseScene, XRefFileToParse);
-					ParseScene(SceneRootNode->GetXRefTree(XRefChild), FXRefScene{SceneRootNode, XRefChild});
+					if (!ParseScene(SceneRootNode->GetXRefTree(XRefChild), FXRefScene{SceneRootNode, XRefChild}))
+					{
+						return false;
+					}
 				}
 			}
 		}
@@ -973,6 +974,10 @@ public:
 				// It's needed restore hierarchy when converting
 				NodeTracker->SetXRefIndex(XRefScene);
 			}
+			else
+			{
+				return false;
+			}
 		}
 		return true;
 	}
@@ -984,6 +989,11 @@ public:
 
 		FNodeTracker* NodeTracker =  GetNodeTracker(Node);
 
+		if (GetCancel())
+		{
+			return nullptr;
+		}
+
 		if (NodeTracker)
 		{
 			InvalidateNode(*NodeTracker, false);
@@ -997,7 +1007,10 @@ public:
 		int32 ChildNum = Node->NumberOfChildren();
 		for (int32 ChildIndex = 0; ChildIndex < ChildNum; ++ChildIndex)
 		{
-			ParseNode(Node->GetChildNode(ChildIndex));
+			if(!ParseNode(Node->GetChildNode(ChildIndex)))
+			{
+				return nullptr;
+			}
 		}
 		return NodeTracker;
 	}
@@ -1048,7 +1061,7 @@ public:
 				bChangeEncountered = true;
 				if (TSet<FNodeTracker*>* NodeTrackersPtr = NodesPerLayer.Find(LayerTracker.Get()))
 				{
-					for(FNodeTracker* NodeTracker:* NodeTrackersPtr)
+					for(FNodeTracker* NodeTracker: *NodeTrackersPtr)
 					{
 						InvalidateNode(*NodeTracker, false);
 					}
@@ -1120,7 +1133,7 @@ public:
 	{
 		if (Cancel.GetCancel())
 		{
-			LogDebug("Update Cancelled");
+			LogWarning("Update Cancelled");
 			return true;
 		}
 		return false;
