[3ds Max] DirectLink: cancelling don't write .udatasmith

- Export of DatasmithSceneExporteris not called when cancelled
- also translating animations can be now cancelled

diff --git a/Engine/Source/Programs/Enterprise/Datasmith/DatasmithMaxExporter/Private/DatasmithMaxExporterWithDirectLink.cpp b/Engine/Source/Programs/Enterprise/Datasmith/DatasmithMaxExporter/Private/DatasmithMaxExporterWithDirectLink.cpp
--- a/Engine/Source/Programs/Enterprise/Datasmith/DatasmithMaxExporter/Private/DatasmithMaxExporterWithDirectLink.cpp
+++ b/Engine/Source/Programs/Enterprise/Datasmith/DatasmithMaxExporter/Private/DatasmithMaxExporterWithDirectLink.cpp
@@ -1533,15 +1533,24 @@ public:
 		return FString::Printf(TEXT("Texmaps: %d updated"), SCENE_UPDATE_STAT_GET(UpdateTextures, Total));
 	}
 
-	void ExportAnimations()
+	bool ExportAnimations(FUpdateProgress::FStage& MainStage)
 	{
+		PROGRESS_STAGE("Export Animations");
+
 		FDatasmithConverter Converter;
 		// Use the same name for the unique level sequence as the scene name
 		TSharedRef<IDatasmithLevelSequenceElement> LevelSequence = FDatasmithSceneFactory::CreateLevelSequence(GetDatasmithScene().GetName());
 		LevelSequence->SetFrameRate(GetFrameRate());
 
+		PROGRESS_STAGE_COUNTER(NodeTrackers.Num());
 		for (TPair<FNodeKey, FNodeTrackerHandle> NodeKeyAndNodeTracker: NodeTrackers)
 		{
+			ProgressCounter.Next();
+			if (GetCancel())
+			{
+				return false;
+			}
+
 			FNodeTracker* NodeTracker = NodeKeyAndNodeTracker.Value.GetNodeTracker();
 
 			if (NodeTracker->HasConverted())
@@ -1564,6 +1573,7 @@ public:
 		{
 			GetDatasmithScene().AddLevelSequence(LevelSequence);
 		}
+		return true;
 	}
 
 	virtual void RemoveMaterial(const TSharedPtr<IDatasmithBaseMaterialElement>& DatasmithMaterial) override
@@ -2964,7 +2974,7 @@ public:
 		FUpdateProgress ProgressManager(!bQuiet, 1);
 		FUpdateProgress::FStage& MainStage = ProgressManager.MainStage;
 
-		SceneTracker.Update(MainStage, false);
+		if (SceneTracker.Update(MainStage, false))
 		{
 			PROGRESS_STAGE("Sync With DirectLink")
 			UpdateDirectLinkScene();
@@ -2974,7 +2984,8 @@ public:
 
 		if (Options.bStatSync)
 		{
-			LogCompletion("Sync completed:");
+			LogCompletion("Sync completed");
+
 			ProgressManager.PrintStatisticss();
 		}
 	}
@@ -3121,14 +3132,22 @@ bool Export(const TCHAR* Name, const TCHAR* OutputPath, bool bQuiet)
 
 	FSceneTracker SceneTracker(PersistentExportOptions.Options, ExportedScene, nullptr);
 
-	SceneTracker.Update(MainStage, true);
+	bool bCancelled = false;
 
-	if (PersistentExportOptions.Options.bAnimatedTransforms)
+	if (!SceneTracker.Update(MainStage, true))
 	{
-		PROGRESS_STAGE("Export Animations");
-		SceneTracker.ExportAnimations();
+		bCancelled = true;
 	}
 
+	if (PersistentExportOptions.Options.bAnimatedTransforms && !bCancelled)
+	{
+		if (!SceneTracker.ExportAnimations(MainStage))
+		{
+			bCancelled = true;
+		}
+	}
+
+	if (!bCancelled)
 	{
 		PROGRESS_STAGE("Save Datasmith Scene");
 
@@ -3145,7 +3164,15 @@ bool Export(const TCHAR* Name, const TCHAR* OutputPath, bool bQuiet)
 
 	ProgressManager.Finished();
 
-	LogInfo(TEXT("Export completed:"));
+	if (!bCancelled)
+	{
+		LogInfo(TEXT("Export completed:"));
+	}
+	else
+	{
+		LogWarning(TEXT("Export cancelled by User"));
+	}
+
 	ProgressManager.PrintStatisticss();
 
 	return true;
