[3ds Max] DirectLink: Sync TimeSlider changes - basic geometry

- regular geometry synced when animation slider time changes
- testcase: create box, animate it's size on timeline, toggle autosync - moving slider should send updates
- note: no other objects are re-translated unnecessarily

diff --git a/Engine/Source/Programs/Enterprise/Datasmith/DatasmithMaxExporter/Private/DatasmithMaxDirectLinkItoo.cpp b/Engine/Source/Programs/Enterprise/Datasmith/DatasmithMaxExporter/Private/DatasmithMaxDirectLinkItoo.cpp
--- a/Engine/Source/Programs/Enterprise/Datasmith/DatasmithMaxExporter/Private/DatasmithMaxDirectLinkItoo.cpp
+++ b/Engine/Source/Programs/Enterprise/Datasmith/DatasmithMaxExporter/Private/DatasmithMaxDirectLinkItoo.cpp
@@ -194,11 +194,9 @@ bool ConvertForest(ISceneTracker& Scene, FNodeTracker& NodeTracker)
 		{
 			INode* GeometryNode = Instances.GeometryNode;
 
-			Object* GeomObj = GeometryNode->EvalWorldState(GetCOREInterface()->GetTime()).obj;
-
 			FMeshConverterSource MeshSource = {
 				GeometryNode, TEXT(""),
-				GetMeshForGeomObject(CurrentTime, GeometryNode, GeomObj), false,
+				GetMeshForGeomObject(CurrentTime, GeometryNode), false,
 				FRenderMeshForConversion()
 			};
 
diff --git a/Engine/Source/Programs/Enterprise/Datasmith/DatasmithMaxExporter/Private/DatasmithMaxDirectLinkMesh.cpp b/Engine/Source/Programs/Enterprise/Datasmith/DatasmithMaxExporter/Private/DatasmithMaxDirectLinkMesh.cpp
--- a/Engine/Source/Programs/Enterprise/Datasmith/DatasmithMaxExporter/Private/DatasmithMaxDirectLinkMesh.cpp
+++ b/Engine/Source/Programs/Enterprise/Datasmith/DatasmithMaxExporter/Private/DatasmithMaxDirectLinkMesh.cpp
@@ -45,27 +45,20 @@ public:
 	}
 };
 
-FRenderMeshForConversion GetMeshForGeomObject(TimeValue CurrentTime, INode* Node, Object* Obj)
+FRenderMeshForConversion GetMeshForGeomObject(TimeValue CurrentTime, INode* Node)
 {
 	// todo: baseline exporter uses GetBaseObject which takes result of EvalWorldState
 	// and searched down DerivedObject pipeline(by taking GetObjRef) 
 	// This is STRANGE as EvalWorldState shouldn't return DerivedObject in the first place(it should return result of pipeline evaluation)
 
 	BOOL bNeedsDelete;
-	Mesh* RenderMesh = GetMeshFromRenderMesh(CurrentTime, Node, bNeedsDelete);
+	Interval ValidityInterval;
+	ValidityInterval.SetInfinite();
+	Mesh* RenderMesh = GetMeshFromRenderMesh(CurrentTime, Node, ValidityInterval, bNeedsDelete);
 
-	return FRenderMeshForConversion(Node, RenderMesh, bNeedsDelete);
-}
-
-FRenderMeshForConversion GetMeshForNode(TimeValue CurrentTime, INode* Node, FTransform Pivot)
-{
-	if (!Node)
-	{
-		return FRenderMeshForConversion();
-	}
-	BOOL bNeedsDelete;
-	Mesh* RenderMesh = GetMeshFromRenderMesh(CurrentTime, Node, bNeedsDelete);
-	return FRenderMeshForConversion(Node, RenderMesh, bNeedsDelete, Pivot);
+	FRenderMeshForConversion Result(Node, RenderMesh, bNeedsDelete);
+	Result.SetValidityInterval(ValidityInterval);
+	return MoveTemp(Result);
 }
 
 /** Convert Max to UE coordinates, handle scene master unit
@@ -158,7 +151,13 @@ FRenderMeshForConversion GetMeshForCollision(TimeValue CurrentTime, ISceneTracke
 			CollisionPivot = BakedTransform;
 		}
 	}
-	return GetMeshForNode(CurrentTime, CollisionNode, CollisionPivot);
+	if (!CollisionNode)
+	{
+		return FRenderMeshForConversion();
+	}
+	FRenderMeshForConversion RenderMesh = GetMeshForGeomObject(CurrentTime, CollisionNode);
+	RenderMesh.SetPivot(CollisionPivot);
+	return MoveTemp(RenderMesh);
 }
 
 
@@ -319,30 +318,24 @@ bool FillDatasmithMeshFromBoundingBox(TimeValue CurrentTime, FDatasmithMesh& Dat
 	}
 	FDatasmithConverter Converter;
 
-	BOOL bNeedsDelete;
-
-	Mesh* MaxMesh = GetMeshFromRenderMesh(CurrentTime, ExportNode, bNeedsDelete);
+	FRenderMeshForConversion MaxMesh = GetMeshForGeomObject(CurrentTime, ExportNode);
 
-	if (!MaxMesh)
+	if (!MaxMesh.IsValid())
 	{
 		return false;
 	}
 
-	if (MaxMesh->getNumFaces() == 0)
+	if (MaxMesh.GetMesh()->getNumFaces() == 0)
 	{
-		if (bNeedsDelete)
-		{
-			MaxMesh->DeleteThis();
-		}
 		return false;
 	}
 
-	if (MaxMesh->getBoundingBox().IsEmpty())
+	if (MaxMesh.GetMesh()->getBoundingBox().IsEmpty())
 	{
-		MaxMesh->buildBoundingBox();
+		MaxMesh.GetMesh()->buildBoundingBox();
 	}
 
-	Box3 BoundingBox = MaxMesh->getBoundingBox();
+	Box3 BoundingBox = MaxMesh.GetMesh()->getBoundingBox();
 
 	DatasmithMesh.SetVerticesCount(8);
 	DatasmithMesh.SetFacesCount(12);
@@ -411,11 +404,6 @@ bool FillDatasmithMeshFromBoundingBox(TimeValue CurrentTime, FDatasmithMesh& Dat
 	SetNormalForAFace(DatasmithMesh, 10, NormalVector);
 	SetNormalForAFace(DatasmithMesh, 11, NormalVector);
 
-	if (bNeedsDelete)
-	{
-		MaxMesh->DeleteThis();
-	}
-
 	return true;
 }
 
@@ -439,7 +427,6 @@ bool CreateDatasmithMeshFromMaxMesh(FDatasmithMesh& DatasmithMesh, const MeshCon
 
 		if (CachedMesh.getNumFaces() > 0)
 		{
-			// todo: pivot
 			FillDatasmithMeshFromMaxMesh(DatasmithMesh, CachedMesh, Params.Node, Params.bConsolidateMaterialIds, SupportedChannels, UVChannelsMap, Params.RenderMesh.GetPivot());
 
 			bResult = true; // Set to true, don't care what ExportToUObject does here - we need to move it to a thread anyway
@@ -501,9 +488,10 @@ int SetObjectParamValue(TimeValue CurrentTime, Object *Obj, const FString& Param
 }
 
 
-Mesh* GetMeshFromRenderMesh(TimeValue CurrentTime, INode* Node, BOOL& bNeedsDelete)
+Mesh* GetMeshFromRenderMesh(TimeValue CurrentTime, INode* Node, Interval& ValidityInterval, BOOL& bNeedsDelete)
 {
 	Object* Obj = GetBaseObject(Node, CurrentTime);
+
 	const Class_ID& ObjectClassID = Obj->ClassID();
 	const FString VRayProxyParamName(TEXT("display"));
 	const FString BodyObjectViewportMeshParamName(TEXT("RenderViewportMeshRA"));
@@ -539,6 +527,10 @@ Mesh* GetMeshFromRenderMesh(TimeValue CurrentTime, INode* Node, BOOL& bNeedsDele
 		SetObjectParamValue(CurrentTime, Obj, BodyObjectViewportMeshParamName, PreviousMeshDisplayValue);
 	}
 
+	if (RenderMesh)
+	{
+		ValidityInterval &= Obj->ObjectValidity(CurrentTime); // Update validity only when actual mesh is returned
+	}
 	return RenderMesh;
 }
 
diff --git a/Engine/Source/Programs/Enterprise/Datasmith/DatasmithMaxExporter/Private/DatasmithMaxExporterWithDirectLink.cpp b/Engine/Source/Programs/Enterprise/Datasmith/DatasmithMaxExporter/Private/DatasmithMaxExporterWithDirectLink.cpp
--- a/Engine/Source/Programs/Enterprise/Datasmith/DatasmithMaxExporter/Private/DatasmithMaxExporterWithDirectLink.cpp
+++ b/Engine/Source/Programs/Enterprise/Datasmith/DatasmithMaxExporter/Private/DatasmithMaxExporterWithDirectLink.cpp
@@ -165,6 +165,7 @@ struct FInstances: FNoncopyable
 
 	// Datasmith mesh conversion results
 	FMeshConverted Converted;
+	Interval ValidityInterval;
 
 	bool HasMesh()
 	{
@@ -2179,6 +2180,11 @@ public:
 			}
 			break;
 		}
+
+		for (FNodeTracker* NodeTrackerPtr : Instances.NodeTrackers)
+		{
+			NodeTrackerPtr->Validity.NarrowValidityToInterval(Instances.ValidityInterval);
+		}
 	}
 
 	void UpdateNodeMetadata(FNodeTracker& NodeTracker)
@@ -2454,16 +2460,17 @@ public:
 	void UpdateInstancesGeometry(FInstances& Instances, FNodeTracker& NodeTracker)
 	{
 		INode* Node = NodeTracker.Node;
-		Object* Obj = Instances.EvaluatedObj;
 
 		FString MeshName = FString::FromInt(Node->GetHandle());
 
 		FMeshConverterSource MeshSource = {
 			Node, MeshName,
-			GeomUtils::GetMeshForGeomObject(CurrentSyncPoint.Time, Node, Obj), false,
+			GeomUtils::GetMeshForGeomObject(CurrentSyncPoint.Time, Node), false,
 			GeomUtils::GetMeshForCollision(CurrentSyncPoint.Time, *this, Node),
 		};
 
+		Instances.ValidityInterval = MeshSource.RenderMesh.GetValidityInterval() & MeshSource.CollisionMesh.GetValidityInterval();
+
 		if (MeshSource.RenderMesh.GetMesh())
 		{
 			bool bHasInstanceWithMultimat = false;
diff --git a/Engine/Source/Programs/Enterprise/Datasmith/DatasmithMaxExporter/Public/DatasmithMaxGeomUtils.h b/Engine/Source/Programs/Enterprise/Datasmith/DatasmithMaxExporter/Public/DatasmithMaxGeomUtils.h
--- a/Engine/Source/Programs/Enterprise/Datasmith/DatasmithMaxExporter/Public/DatasmithMaxGeomUtils.h
+++ b/Engine/Source/Programs/Enterprise/Datasmith/DatasmithMaxExporter/Public/DatasmithMaxGeomUtils.h
@@ -26,13 +26,43 @@ namespace GeomUtils
 class FRenderMeshForConversion: FNoncopyable
 {
 public:
-	explicit FRenderMeshForConversion(){}
-	explicit FRenderMeshForConversion(INode* InNode, Mesh* InMaxMesh, bool bInNeedsDelete, FTransform InPivot=FTransform::Identity)
+	explicit FRenderMeshForConversion()
+		: Node(nullptr)
+		, MaxMesh(nullptr)
+		, bNeedsDelete(false)
+		, Pivot(FTransform::Identity)
+	{
+		ValidityInterval.SetInfinite();
+	}
+
+	explicit FRenderMeshForConversion(INode* InNode, Mesh* InMaxMesh, bool bInNeedsDelete)
 		: Node(InNode)
 		, MaxMesh(InMaxMesh)
 		, bNeedsDelete(bInNeedsDelete)
-		, Pivot(InPivot)
-	{}
+		, Pivot(FTransform::Identity)
+	{
+		ValidityInterval.SetInfinite();
+	}
+
+	FRenderMeshForConversion(FRenderMeshForConversion&& Other)
+	{
+		Node = Other.Node;
+		Other.Node = nullptr;
+
+		MaxMesh = Other.MaxMesh;
+		Other.MaxMesh = nullptr;
+
+		Pivot = Other.Pivot;
+		Other.Pivot = FTransform::Identity;
+
+		ValidityInterval = Other.ValidityInterval;
+		Other.ValidityInterval.SetInfinite();
+
+		//Transfer ownership
+		bNeedsDelete = Other.bNeedsDelete;
+		Other.bNeedsDelete = false;
+	}
+
 	~FRenderMeshForConversion()
 	{
 		if (bNeedsDelete)
@@ -61,19 +91,33 @@ public:
 		return Pivot;
 	}
 
+	void SetPivot(const FTransform& InPivot)
+	{
+		Pivot = InPivot;
+	}
+
+	void SetValidityInterval(const Interval& Interval)
+	{
+		ValidityInterval = Interval;
+	}
 
+	Interval GetValidityInterval()
+	{
+		return ValidityInterval;
+	}
 private:
-	INode* Node = nullptr;
-	Mesh* MaxMesh = nullptr;
-	bool bNeedsDelete = false;
+	INode* Node;
+	Mesh* MaxMesh;
+	bool bNeedsDelete;
 	FTransform Pivot;
+	Interval ValidityInterval;
 };
 
 bool ConvertRailClone(DatasmithMaxDirectLink::ISceneTracker& SceneTracker, DatasmithMaxDirectLink::FNodeTracker& NodeTracker);
 bool ConvertForest(DatasmithMaxDirectLink::ISceneTracker& Scene, DatasmithMaxDirectLink::FNodeTracker& NodeTracker);
 
-Mesh* GetMeshFromRenderMesh(TimeValue CurrentTime, INode* Node, BOOL& bNeedsDelete);
-FRenderMeshForConversion GetMeshForGeomObject(TimeValue CurrentTime, INode* Node, Object* Obj); // Extract mesh using already evaluated object
+Mesh* GetMeshFromRenderMesh(TimeValue CurrentTime, INode* Node, Interval& ValidityInterval, BOOL& bNeedsDelete);
+FRenderMeshForConversion GetMeshForGeomObject(TimeValue CurrentTime, INode* Node); // Extract mesh using already evaluated object
 FRenderMeshForConversion GetMeshForNode(TimeValue CurrentTime, INode* Node, FTransform Pivot); // Extract mesh evaluating node object
 
 FRenderMeshForConversion GetMeshForCollision(TimeValue CurrentTime, DatasmithMaxDirectLink::ISceneTracker& SceneTracker, INode* Node);
