[3ds Max] DirectLink: fixed VRayProxy to export correct mesh

- after previous change (using ObjectValidity to validate geometry for time slider) vray proxy stated exporting display mesh set in UI instead of the correct, full-size mesh. This is fixed.

diff --git a/Engine/Source/Programs/Enterprise/Datasmith/DatasmithMaxExporter/Private/DatasmithMaxDirectLinkMesh.cpp b/Engine/Source/Programs/Enterprise/Datasmith/DatasmithMaxExporter/Private/DatasmithMaxDirectLinkMesh.cpp
--- a/Engine/Source/Programs/Enterprise/Datasmith/DatasmithMaxExporter/Private/DatasmithMaxDirectLinkMesh.cpp
+++ b/Engine/Source/Programs/Enterprise/Datasmith/DatasmithMaxExporter/Private/DatasmithMaxDirectLinkMesh.cpp
@@ -492,6 +492,12 @@ Mesh* GetMeshFromRenderMesh(TimeValue CurrentTime, INode* Node, Interval& Validi
 {
 	Object* Obj = GetBaseObject(Node, CurrentTime);
 
+	// Read validity interval before changing display settings
+	// Else VRay will somehow react to a call to ObjectValidity and change RenderMesh contents to what it was originally.
+	// E.g. display was set 'Box', we change display to 'Mesh' here before calling GetRenderMesh to retrieve actual mesh later(rather than a simple box)
+	//  but if ObjectValidity is called after GetRenderMesh that RenderMesh will become box again!
+	Interval ObjectValidity = Obj->ObjectValidity(CurrentTime);
+
 	const Class_ID& ObjectClassID = Obj->ClassID();
 	const FString VRayProxyParamName(TEXT("display"));
 	const FString BodyObjectViewportMeshParamName(TEXT("RenderViewportMeshRA"));
@@ -529,7 +535,7 @@ Mesh* GetMeshFromRenderMesh(TimeValue CurrentTime, INode* Node, Interval& Validi
 
 	if (RenderMesh)
 	{
-		ValidityInterval &= Obj->ObjectValidity(CurrentTime); // Update validity only when actual mesh is returned
+		ValidityInterval &= ObjectValidity; // Update validity only when actual mesh is returned
 	}
 	return RenderMesh;
 }
