Datasmith Revit: Decals import to use the M_StdDecal

diff --git a/Engine/Plugins/Enterprise/DatasmithImporter/Source/DatasmithTranslator/Private/MasterMaterials/DatasmithRevitMaterialSelector.cpp b/Engine/Plugins/Enterprise/DatasmithImporter/Source/DatasmithTranslator/Private/MasterMaterials/DatasmithRevitMaterialSelector.cpp
--- a/Engine/Plugins/Enterprise/DatasmithImporter/Source/DatasmithTranslator/Private/MasterMaterials/DatasmithRevitMaterialSelector.cpp
+++ b/Engine/Plugins/Enterprise/DatasmithImporter/Source/DatasmithTranslator/Private/MasterMaterials/DatasmithRevitMaterialSelector.cpp
@@ -13,6 +13,7 @@ FDatasmithRevitMaterialSelector::FDatasmithRevitMaterialSelector()
 {
 	// Master
 	MasterMaterial.FromSoftObjectPath( FSoftObjectPath("/DatasmithContent/Materials/RevitMaster.RevitMaster") );
+	MasterMaterialDecal.FromSoftObjectPath(FSoftObjectPath("/DatasmithContent/Materials/StdDecal/M_StdDecal.M_StdDecal"));
 }
 
 bool FDatasmithRevitMaterialSelector::IsValid() const
@@ -20,8 +21,13 @@ bool FDatasmithRevitMaterialSelector::IsValid() const
 	return MasterMaterial.IsValid();
 }
 
-const FDatasmithMasterMaterial& FDatasmithRevitMaterialSelector::GetMasterMaterial( const TSharedPtr< IDatasmithMasterMaterialElement >& /*InDatasmithMaterial*/ ) const
+const FDatasmithMasterMaterial& FDatasmithRevitMaterialSelector::GetMasterMaterial( const TSharedPtr< IDatasmithMasterMaterialElement >& InDatasmithMaterial ) const
 {
+	if (InDatasmithMaterial->GetMaterialType() == EDatasmithMasterMaterialType::Decal)
+	{
+		return MasterMaterialDecal;
+	}
+
 	return MasterMaterial;
 }
 
diff --git a/Engine/Plugins/Enterprise/DatasmithImporter/Source/DatasmithTranslator/Private/MasterMaterials/DatasmithRevitMaterialSelector.h b/Engine/Plugins/Enterprise/DatasmithImporter/Source/DatasmithTranslator/Private/MasterMaterials/DatasmithRevitMaterialSelector.h
--- a/Engine/Plugins/Enterprise/DatasmithImporter/Source/DatasmithTranslator/Private/MasterMaterials/DatasmithRevitMaterialSelector.h
+++ b/Engine/Plugins/Enterprise/DatasmithImporter/Source/DatasmithTranslator/Private/MasterMaterials/DatasmithRevitMaterialSelector.h
@@ -16,9 +16,10 @@ public:
 	virtual ~FDatasmithRevitMaterialSelector() = default;
 
 	virtual bool IsValid() const override;
-	virtual const FDatasmithMasterMaterial& GetMasterMaterial( const TSharedPtr< IDatasmithMasterMaterialElement >& /*InDatasmithMaterial*/ ) const override;
+	virtual const FDatasmithMasterMaterial& GetMasterMaterial( const TSharedPtr< IDatasmithMasterMaterialElement >& InDatasmithMaterial ) const override;
 	virtual void FinalizeMaterialInstance(const TSharedPtr< IDatasmithMasterMaterialElement >& InDatasmithMaterial, UMaterialInstanceConstant* MaterialInstance) const override;
 
 private:
 	FDatasmithMasterMaterial MasterMaterial;
+	FDatasmithMasterMaterial MasterMaterialDecal;
 };
diff --git a/Engine/Source/Developer/Datasmith/DatasmithFacade/Public/DatasmithFacadeMaterial.h b/Engine/Source/Developer/Datasmith/DatasmithFacade/Public/DatasmithFacadeMaterial.h
--- a/Engine/Source/Developer/Datasmith/DatasmithFacade/Public/DatasmithFacadeMaterial.h
+++ b/Engine/Source/Developer/Datasmith/DatasmithFacade/Public/DatasmithFacadeMaterial.h
@@ -59,6 +59,7 @@ public:
 		/** Material has a transparent cutout map */
 		CutOut,
 		Emissive,
+		Decal,
 		/** Dummy element to count the number of types */
 		Count
 	};
diff --git a/Engine/Source/Programs/Enterprise/Datasmith/DatasmithFacadeCSharp/Public/FDatasmithFacadeMasterMaterial.cs b/Engine/Source/Programs/Enterprise/Datasmith/DatasmithFacadeCSharp/Public/FDatasmithFacadeMasterMaterial.cs
--- a/Engine/Source/Programs/Enterprise/Datasmith/DatasmithFacadeCSharp/Public/FDatasmithFacadeMasterMaterial.cs
+++ b/Engine/Source/Programs/Enterprise/Datasmith/DatasmithFacadeCSharp/Public/FDatasmithFacadeMasterMaterial.cs
@@ -126,6 +126,7 @@ public class FDatasmithFacadeMasterMaterial : FDatasmithFacadeBaseMaterial {
     Custom,
     CutOut,
     Emissive,
+    Decal,
     Count
   }
 
diff --git a/Engine/Source/Programs/Enterprise/Datasmith/DatasmithRevitExporter/Private/DatasmithRevitDecal.cs b/Engine/Source/Programs/Enterprise/Datasmith/DatasmithRevitExporter/Private/DatasmithRevitDecal.cs
--- a/Engine/Source/Programs/Enterprise/Datasmith/DatasmithRevitExporter/Private/DatasmithRevitDecal.cs
+++ b/Engine/Source/Programs/Enterprise/Datasmith/DatasmithRevitExporter/Private/DatasmithRevitDecal.cs
@@ -11,12 +11,20 @@ namespace DatasmithRevitExporter
 		public string MaterialName { get; private set; }
 		public string DiffuseTexturePath { get; private set; }
 		public string BumpTexturePath { get; private set; }
+		public string CutoutTexturePath { get; private set; }
+		public double BumpAmount { get; private set; }
+		public double Transparency { get; private set; }
+		public double Luminance { get; private set; }
 
 		private void SetMaterialProperties(Asset InMaterialAsset)
 		{
 			const string DiffuseMapPropName = "decApp_diffuse";
 			const string BumpMapPropName = "decApp_bump_map";
 			const string TexturePathName = "unifiedbitmap_Bitmap";
+			const string BumpAmountPropName = "decApp_bump_amount";
+			const string CutoutMapPropName = "decApp_cutout_opacity";
+			const string LuminancePropName = "decApp_self_illum_luminance";
+			const string TransparencyPropName = "decApp_transparency";
 
 			string GetConnectedImagePath(AssetProperty InProp)
 			{
@@ -47,9 +55,6 @@ namespace DatasmithRevitExporter
 				return null;
 			}
 
-			bool bDiffusePropertyProcessed = false;
-			bool bBumpPropertyProcessed = false;
-
 			for (int PropIndex = 0; PropIndex < InMaterialAsset.Size; ++PropIndex)
 			{
 #if REVIT_API_2018
@@ -60,17 +65,26 @@ namespace DatasmithRevitExporter
 				if (Prop.Name == DiffuseMapPropName)
 				{
 					DiffuseTexturePath = GetConnectedImagePath(Prop);
-					bDiffusePropertyProcessed = true;
 				}
 				else if (Prop.Name == BumpMapPropName)
 				{
 					BumpTexturePath = GetConnectedImagePath(Prop);
-					bBumpPropertyProcessed = true;
 				}
-
-				if (bDiffusePropertyProcessed && bBumpPropertyProcessed)
+				else if (Prop.Name == BumpAmountPropName)
 				{
-					break;
+					BumpAmount = (Prop as AssetPropertyDouble)?.Value ?? 0;
+				}
+				else if (Prop.Name == CutoutMapPropName)
+				{
+					CutoutTexturePath = GetConnectedImagePath(Prop);
+				}
+				else if (Prop.Name == LuminancePropName)
+				{
+					Luminance = (Prop as AssetPropertyDouble)?.Value ?? 0;
+				}
+				else if (Prop.Name == TransparencyPropName)
+				{
+					Transparency = (Prop as AssetPropertyDouble)?.Value ?? 0;
 				}
 			}
 		}
diff --git a/Engine/Source/Programs/Enterprise/Datasmith/DatasmithRevitExporter/Private/DatasmithRevitDocumentData.cs b/Engine/Source/Programs/Enterprise/Datasmith/DatasmithRevitExporter/Private/DatasmithRevitDocumentData.cs
--- a/Engine/Source/Programs/Enterprise/Datasmith/DatasmithRevitExporter/Private/DatasmithRevitDocumentData.cs
+++ b/Engine/Source/Programs/Enterprise/Datasmith/DatasmithRevitExporter/Private/DatasmithRevitDocumentData.cs
@@ -818,8 +818,9 @@ namespace DatasmithRevitExporter
 					else
 					{
 						InElement.ElementActor = new FDatasmithFacadeActorMesh(HashedActorName);
-						InElement.ElementActor.SetLabel(GetActorLabel());
 					}
+
+					InElement.ElementActor.SetLabel(GetActorLabel());
 				}
 
 				// Set the world transform of the Datasmith mesh actor.
@@ -1906,7 +1907,8 @@ namespace DatasmithRevitExporter
 					}
 
 					Transform TranslationMatrix = Transform.CreateTranslation(BasePointPosition);
-					SetActorTransform(TranslationMatrix.Multiply(InWorldTransform), BasePointActor);
+					// Don't apply offset since basepoints aren't yet initialized
+					SetActorTransform(TranslationMatrix.Multiply(InWorldTransform), BasePointActor, false);
 
 					// Set the Datasmith placeholder actor layer to the base point category name.
 					BasePointActor.SetLayer(BasePointLocation.Category.Name);
@@ -2049,7 +2051,8 @@ namespace DatasmithRevitExporter
 
 				FDecalMaterial DecalMaterial = DecalMaterialPair.Value;
 
-				FDatasmithFacadeDecalMaterial DatasmithMaterial = new FDatasmithFacadeDecalMaterial(DecalMaterial.MaterialName);
+				FDatasmithFacadeMasterMaterial DatasmithMaterial = new FDatasmithFacadeMasterMaterial(DecalMaterial.MaterialName);
+				DatasmithMaterial.SetMaterialType(FDatasmithFacadeMasterMaterial.EMasterMaterialType.Decal);
 
 				if (!string.IsNullOrEmpty(DecalMaterial.DiffuseTexturePath))
 				{
@@ -2057,7 +2060,9 @@ namespace DatasmithRevitExporter
 					DiffuseTexture.SetSRGB(FDatasmithFacadeTexture.EColorSpace.sRGB);
 					DiffuseTexture.SetTextureMode(FDatasmithFacadeTexture.ETextureMode.Diffuse);
 					DiffuseTexture.SetFile(DecalMaterial.DiffuseTexturePath);
-					DatasmithMaterial.SetDiffuseTexturePathName(DiffuseTexture.GetName());
+
+					DatasmithMaterial.AddTexture("ColorMap", DiffuseTexture);
+
 					InDatasmithScene.AddTexture(DiffuseTexture);
 				}
 
@@ -2067,10 +2072,36 @@ namespace DatasmithRevitExporter
 					BumpTexture.SetSRGB(FDatasmithFacadeTexture.EColorSpace.sRGB);
 					BumpTexture.SetTextureMode(FDatasmithFacadeTexture.ETextureMode.Bump);
 					BumpTexture.SetFile(DecalMaterial.BumpTexturePath);
-					DatasmithMaterial.SetDiffuseTexturePathName(BumpTexture.GetName());
+
+					DatasmithMaterial.AddTexture("NormalMap", BumpTexture);
+					DatasmithMaterial.AddFloat("NormalMapAmount", (float)DecalMaterial.BumpAmount);
+
 					InDatasmithScene.AddTexture(BumpTexture);
 				}
 
+				if (!string.IsNullOrEmpty(DecalMaterial.CutoutTexturePath))
+				{
+					FDatasmithFacadeTexture CutoutTexture = FDatasmithFacadeMaterialsUtils.CreateSimpleTextureElement(DecalMaterial.CutoutTexturePath);
+					CutoutTexture.SetSRGB(FDatasmithFacadeTexture.EColorSpace.sRGB);
+					CutoutTexture.SetTextureMode(FDatasmithFacadeTexture.ETextureMode.Other);
+					CutoutTexture.SetFile(DecalMaterial.CutoutTexturePath);
+
+					DatasmithMaterial.AddBoolean("UseCustomOpacityMap", true);
+					DatasmithMaterial.AddTexture("OpacityMap", CutoutTexture);
+
+					InDatasmithScene.AddTexture(CutoutTexture);
+				}
+
+				if (DecalMaterial.Luminance > 0f)
+				{
+					DatasmithMaterial.AddFloat("LuminanceAmount", (float)DecalMaterial.Luminance);
+				}
+
+				if (DecalMaterial.Transparency > 0f)
+				{
+					DatasmithMaterial.AddFloat("Opacity", (float)DecalMaterial.Transparency);
+				}
+
 				InDatasmithScene.AddMaterial(DatasmithMaterial);
 
 				Transform DecalTransform = null;
@@ -2278,8 +2309,8 @@ namespace DatasmithRevitExporter
 
 		private void SetActorTransform(
 			Transform InWorldTransform,
-			FDatasmithFacadeActor IOActor
-		)
+			FDatasmithFacadeActor IOActor,
+			bool bInApplyOffset = true)
 		{
 			XYZ transformBasisX = InWorldTransform.BasisX;
 			XYZ transformBasisY = InWorldTransform.BasisY;
@@ -2287,12 +2318,25 @@ namespace DatasmithRevitExporter
 			XYZ transformOrigin = InWorldTransform.Origin;
 
 			// Check if need to apply world offset to element transform
-			if (InsertionPoint != FSettings.EInsertionPoint.Default)
+			if (bInApplyOffset && InsertionPoint != FSettings.EInsertionPoint.Default)
 			{
 				switch (InsertionPoint)
 				{
-					case FSettings.EInsertionPoint.BasePoint: transformOrigin -= ProjectBasePoint; break;
-					case FSettings.EInsertionPoint.SurveyPoint: transformOrigin -= ProjectSurveyPoint; break;
+					case FSettings.EInsertionPoint.BasePoint: 
+					{
+						if (ProjectBasePoint != null)
+						{
+							transformOrigin -= ProjectBasePoint;
+						}
+					} break;
+
+					case FSettings.EInsertionPoint.SurveyPoint: 
+					{
+						if (ProjectSurveyPoint != null)
+						{
+							transformOrigin -= ProjectSurveyPoint;
+						}
+					} break;
 				}
 			}
 
diff --git a/Engine/Source/Runtime/Datasmith/DatasmithCore/Public/DatasmithDefinitions.h b/Engine/Source/Runtime/Datasmith/DatasmithCore/Public/DatasmithDefinitions.h
--- a/Engine/Source/Runtime/Datasmith/DatasmithCore/Public/DatasmithDefinitions.h
+++ b/Engine/Source/Runtime/Datasmith/DatasmithCore/Public/DatasmithDefinitions.h
@@ -235,6 +235,7 @@ enum class EDatasmithMasterMaterialType : uint8
 	/** Material has a transparent cutout map */
 	CutOut,
 	Emissive,
+	Decal,
 	/** Dummy element to count the number of types */
 	Count
 };
