DatasmithSDK: Enable directlink for Shipping target configuration

diff --git a/Engine/Source/Programs/Enterprise/Datasmith/DatasmithSDK/DatasmithSDK.Target.cs b/Engine/Source/Programs/Enterprise/Datasmith/DatasmithSDK/DatasmithSDK.Target.cs
--- a/Engine/Source/Programs/Enterprise/Datasmith/DatasmithSDK/DatasmithSDK.Target.cs
+++ b/Engine/Source/Programs/Enterprise/Datasmith/DatasmithSDK/DatasmithSDK.Target.cs
@@ -38,6 +38,13 @@ public class DatasmithSDKTarget : TargetRules
 		{
 			AddWindowsPostBuildSteps();
 		}
+
+		// Enable UDP in shipping (used by DirectLink)
+		if (BuildEnvironment == TargetBuildEnvironment.Unique)
+		{
+			GlobalDefinitions.Add("ALLOW_UDP_MESSAGING_SHIPPING=1"); // bypasses the 'if shipping' of UdpMessagingModule.cpp
+			GlobalDefinitions.Add("PLATFORM_SUPPORTS_MESSAGEBUS=1"); // required to enable the default MessageBus in MessagingModule.cpp
+		}
 	}
 
 	public void PostBuildCopy(string SrcPath, string DestPath)
