[3ds Max] DirectLink: notifications rework

- Changed geometry-related events to use ReferenceMaker notifications
- Result is that notifications are immediate..
- ..therefore could be ignored when sent from within Update. This allows setting Geom nodes to produce Render mesh during update(switching to Render mesh sends notification).
- Test case - AutoSync for optimize/turbosmooth/vrayproxy etc objects exports Render mesh(not viewport mesh) and doesn't spam DirectLink updates

diff --git a/Engine/Source/Programs/Enterprise/Datasmith/DatasmithMaxExporter/Private/DatasmithMaxDirectLinkLogging.cpp b/Engine/Source/Programs/Enterprise/Datasmith/DatasmithMaxExporter/Private/DatasmithMaxDirectLinkLogging.cpp
--- a/Engine/Source/Programs/Enterprise/Datasmith/DatasmithMaxExporter/Private/DatasmithMaxDirectLinkLogging.cpp
+++ b/Engine/Source/Programs/Enterprise/Datasmith/DatasmithMaxExporter/Private/DatasmithMaxDirectLinkLogging.cpp
@@ -25,15 +25,14 @@ MAX_INCLUDES_START
 MAX_INCLUDES_END
 
 
-
 namespace DatasmithMaxDirectLink
 {
 
-void LogDebug(const TCHAR* Msg)
+void LogDebugImpl(const FString& Msg)
 {
 #ifdef LOG_DEBUG_ENABLE
-	mprintf(L"[%s]%s\n", *FDateTime::UtcNow().ToString(TEXT("%Y.%m.%d-%H.%M.%S:%s")), Msg);
-	UE_LOG(LogDatasmithMaxExporter, Error, TEXT("%s"), Msg);
+	// mprintf(L"[%s]%s\n", *FDateTime::UtcNow().ToString(TEXT("%Y.%m.%d-%H.%M.%S:%s")), *Msg);
+	UE_LOG(LogDatasmithMaxExporter, Error, TEXT("%s"), *Msg);
 	LogFlush();
 	LogDebugDialog(Msg);
 #endif
@@ -83,12 +82,6 @@ void LogInfo(const FString& Msg)
 	LogInfo(*Msg);
 }
 
-void LogDebug(const FString& Msg)
-{
-	LogDebug(*Msg);
-}
-
-
 void LogFlush()
 {
 	Async(EAsyncExecution::TaskGraphMainThread,
@@ -127,13 +120,13 @@ void LogDebugNode(const FString& Name, INode* Node)
 #endif
 }
 
-void LogNodeEvent(const MCHAR* Name, INodeEventCallback::NodeKeyTab& nodes)
+void LogNodeEvent(const MCHAR* Name, INodeEventCallback::NodeKeyTab& Nodes)
 {
 #ifdef LOG_DEBUG_HEAVY_ENABLE
 	LogDebug(FString::Printf(TEXT("NodeEventCallback:%s"), Name));
-	for (int NodeIndex = 0; NodeIndex < nodes.Count(); ++NodeIndex)
+	for (int NodeIndex = 0; NodeIndex < Nodes.Count(); ++NodeIndex)
 	{
-		FNodeKey NodeKey = nodes[NodeIndex];
+		FNodeKey NodeKey = Nodes[NodeIndex];
 
 		Animatable* anim = Animatable::GetAnimByHandle(NodeKey);
 		if (INode* Node = NodeEventNamespace::GetNodeByKey(NodeKey)) // Node sometimes is null. Not sure why
diff --git a/Engine/Source/Programs/Enterprise/Datasmith/DatasmithMaxExporter/Private/DatasmithMaxDirectLinkNotifications.cpp b/Engine/Source/Programs/Enterprise/Datasmith/DatasmithMaxExporter/Private/DatasmithMaxDirectLinkNotifications.cpp
--- a/Engine/Source/Programs/Enterprise/Datasmith/DatasmithMaxExporter/Private/DatasmithMaxDirectLinkNotifications.cpp
+++ b/Engine/Source/Programs/Enterprise/Datasmith/DatasmithMaxExporter/Private/DatasmithMaxDirectLinkNotifications.cpp
@@ -32,126 +32,84 @@ public:
 
 	virtual BOOL VerboseDeleted() { return TRUE; }
 
-	virtual void GeometryChanged(NodeKeyTab& nodes) override
-	{
-		LogNodeEvent(L"GeometryChanged", nodes);
-
-		InvalidateGeometry(nodes);
-	}
-
-	virtual void TopologyChanged(NodeKeyTab& nodes) override
-	{
-		LogNodeEvent(L"TopologyChanged", nodes);
-
-		InvalidateGeometry(nodes);
-	}
-
-
-	virtual void MappingChanged(NodeKeyTab& nodes) override
-	{
-		LogNodeEvent(L"MappingChanged", nodes);
-
-		InvalidateGeometry(nodes);
-	}
-
-	void InvalidateGeometry(NodeKeyTab& nodes)
+	// Tracks node's material parameter change(even if it's a submaterial of multimat that is assigned)
+	virtual void MaterialOtherEvent(NodeKeyTab& nodes) override
 	{
+		LogNodeEvent(L"MaterialOtherEvent", nodes);
 		for (int NodeIndex = 0; NodeIndex < nodes.Count(); ++NodeIndex)
 		{
-			SceneTracker.NodeGeometryChanged(nodes[NodeIndex]);
+			SceneTracker.NodeMaterialGraphModified(nodes[NodeIndex]);
 		}
 	}
 
-	// Fired when node transform changes
-	virtual void ControllerOtherEvent(NodeKeyTab& nodes) override
+	virtual void NameChanged(NodeKeyTab& nodes) override
 	{
-		LogNodeEvent(L"ControllerOtherEvent", nodes);
+		LogNodeEvent(L"NameChanged", nodes);
 
 		for (int NodeIndex = 0; NodeIndex < nodes.Count(); ++NodeIndex)
 		{
-			SceneTracker.NodeTransformChanged(nodes[NodeIndex]);
+			SceneTracker.NodeNameChanged(nodes[NodeIndex]);
 		}
 	}
 
-	// Tracks material assignment on node
-	virtual void MaterialStructured(NodeKeyTab& nodes) override
+	virtual void LinkChanged(NodeKeyTab& nodes) override
 	{
-		LogNodeEvent(L"MaterialStructured", nodes);
+		LogNodeEvent(L"LinkChanged", nodes);
 		for (int NodeIndex = 0; NodeIndex < nodes.Count(); ++NodeIndex)
 		{
-			SceneTracker.NodeMaterialAssignmentChanged(nodes[NodeIndex]);
+			SceneTracker.NodeLinkChanged(nodes[NodeIndex]);
 		}
 	}
 
-	// Tracks node's material parameter change(even if it's a submaterial of multimat that is assigned)
-	virtual void MaterialOtherEvent(NodeKeyTab& nodes) override
+	// Not used:
+
+	virtual void LayerChanged(NodeKeyTab& nodes) override
 	{
-		LogNodeEvent(L"MaterialOtherEvent", nodes);
-		for (int NodeIndex = 0; NodeIndex < nodes.Count(); ++NodeIndex)
-		{
-			SceneTracker.NodeMaterialGraphModified(nodes[NodeIndex]);
-		}
+		LogNodeEvent(L"LayerChanged", nodes);
+	}
+
+	virtual void UserPropertiesChanged(NodeKeyTab& nodes) override
+	{
+		LogNodeEvent(L"UserPropertiesChanged", nodes);
 	}
 
 	virtual void HideChanged(NodeKeyTab& nodes) override
 	{
 		LogNodeEvent(L"HideChanged", nodes);
-		for (int NodeIndex = 0; NodeIndex < nodes.Count(); ++NodeIndex)
-		{
-			SceneTracker.NodeHideChanged(nodes[NodeIndex]);
-		}
 	}
 
 	virtual void RenderPropertiesChanged(NodeKeyTab& nodes) override
 	{
 		LogNodeEvent(L"RenderPropertiesChanged", nodes);
-		// Handle Renderable flag change. mxs: box.setRenderable
-		for (int NodeIndex = 0; NodeIndex < nodes.Count(); ++NodeIndex)
-		{
-			SceneTracker.NodePropertiesChanged(nodes[NodeIndex]);
-		}
 	}
 
-	virtual void UserPropertiesChanged(NodeKeyTab& nodes) override
+	virtual void GeometryChanged(NodeKeyTab& nodes) override
 	{
-		LogNodeEvent(L"UserPropertiesChanged", nodes);
-		// Handle user properties change for metadata update
-		for (int NodeIndex = 0; NodeIndex < nodes.Count(); ++NodeIndex)
-		{
-			SceneTracker.NodePropertiesChanged(nodes[NodeIndex]);
-		}
+		LogNodeEvent(L"GeometryChanged", nodes);
 	}
 
-	virtual void NameChanged(NodeKeyTab& nodes) override
+	virtual void TopologyChanged(NodeKeyTab& nodes) override
 	{
-		LogNodeEvent(L"NameChanged", nodes);
-
-		for (int NodeIndex = 0; NodeIndex < nodes.Count(); ++NodeIndex)
-		{
-			SceneTracker.NodeNameChanged(nodes[NodeIndex]);
-		}
+		LogNodeEvent(L"TopologyChanged", nodes);
 	}
 
-	virtual void LayerChanged(NodeKeyTab& nodes) override
-	{
-		LogNodeEvent(L"LayerChanged", nodes);
 
-		for (int NodeIndex = 0; NodeIndex < nodes.Count(); ++NodeIndex)
-		{
-			SceneTracker.NodePropertiesChanged(nodes[NodeIndex]);
-		}
+	virtual void MappingChanged(NodeKeyTab& nodes) override
+	{
+		LogNodeEvent(L"MappingChanged", nodes);
 	}
 
-	virtual void LinkChanged(NodeKeyTab& nodes) override
+	// Fired when node transform changes
+	virtual void ControllerOtherEvent(NodeKeyTab& nodes) override
 	{
-		LogNodeEvent(L"LinkChanged", nodes);
-		for (int NodeIndex = 0; NodeIndex < nodes.Count(); ++NodeIndex)
-		{
-			SceneTracker.NodeLinkChanged(nodes[NodeIndex]);
-		}
+		LogNodeEvent(L"ControllerOtherEvent", nodes);
 	}
 
-	// Not used:
+	// Tracks material assignment on node
+	virtual void MaterialStructured(NodeKeyTab& nodes) override
+	{
+		LogNodeEvent(L"MaterialStructured", nodes);
+	}
 
 	virtual void Added(NodeKeyTab& nodes) override
 	{
@@ -318,107 +276,200 @@ private:
 };
 
 
-class FNodeObserver : public ReferenceMaker
+// Every instance watches single node
+class FNodeNotifications: public ReferenceMaker
 {
-	typedef int FItemIndex;
-
 public:
-	~FNodeObserver()
+	ISceneTracker& SceneTracker;
+
+	FNodeNotifications(ISceneTracker& InSceneTracker, INode& Node)
+		: SceneTracker(InSceneTracker)
+		, Node(Node)
 	{
-		DeleteAllRefs(); // Required to be called in destructor
+		LogDebug(FString::Printf(TEXT("FNodeNotifications(%s - %p)::FNodeNotifications"), Node.GetName(), &Node));
+
+		// Don't record setting reference for Undo/Redo  - we control this ourselves(e.g. on ADDED_NODE event for addition)
+		HoldSuspend SuspendUndoRedoGuard; 
+		ReplaceReference(0, &Node);
 	}
 
-	void Reset()
+	void Delete()
 	{
-		IndexToReferencedItem.Reset();
-		ReferencedItemToIndex.Reset();
+		LogDebug(FString::Printf(TEXT("FNodeNotifications(%p)::Delete"), &Node));
+		
+		HoldSuspend SuspendUndoRedoGuard; 
+		DeleteMe(); //Instead of calling DeleteAllRefs() in destructor, calling DeleteMe under HoldSuspend to disable undo for this object
 	}
 
-	RefResult NotifyRefChanged(const Interval& ChangeInterval, RefTargetHandle TargetHandle, PartID& PartId, RefMessage Message, BOOL propagate) override
+	// Same INode pointer could be added to the scene - handle this in case notifier for this pointer still kept
+	void ReAdded()
 	{
-		// todo: remove material handling???
-		if (!ReferencedItemToIndex.Contains(TargetHandle))
+		ReplaceReference(0, &Node);
+	}
+
+protected:
+
+	virtual RefResult NotifyRefChanged(const Interval& ChangeInterval, RefTargetHandle TargetHandle, PartID& PartId, RefMessage Message, BOOL propagate) override
+	{
+		if (!bReferenceSet)
 		{
-			LogDebugNode(TEXT("FNodeObserver::NotifyRefChanged(MISSING)"), dynamic_cast<INode*>(TargetHandle));
+			LOG_DEBUG_HEAVY(FString::Printf(TEXT("FNodeNotifications(UNSET - %p)::NotifyRefChanged:  Target - %p, Message - %x"), &Node, TargetHandle, Message));
+			return REF_DONTCARE;
 		}
-		else
+		
+		LOG_DEBUG_HEAVY(FString::Printf(TEXT("FNodeNotifications(%s - %p)::NotifyRefChanged:  Target - %p, Message - %x"), Node.GetName(), &Node, TargetHandle, Message));
+		LogDebugNode(TEXT("     Node:"), &Node);
+
+		if (SceneTracker.IsUpdateInProgress())
 		{
-			LogDebugNode(TEXT("FNodeObserver::NotifyRefChanged"), dynamic_cast<INode*>(TargetHandle));
+			// Don't propagate ReferenceMaker events when Update is in progress - these events are result of accessing render geometry on some objects, including
+			// Smooth/TurboSmooth modifiers, VRay proxy and similar. Plugin retrieves render mesh from them and those objects fire geometry change events
+			// Plugin reverts settings back to previously set(i.e. viewport for smooth modifiers or whatever vray proxy had for display)
+			return REF_DONTCARE;
 		}
 
+		switch (Message)
+		{
+		case REFMSG_REF_DELETED:
+			// 
+		break;
+	    case REFMSG_CHANGE:
+			LOG_DEBUG_HEAVY(FString::Printf(TEXT("   REFMSG_CHANGE.PartId:  %llx"), PartId));
+
+            if (PartId & (PART_TOPO|PART_GEOM|PART_TEXMAP))
+            {
+				SceneTracker.NodeGeometryChanged(&Node);
+            }
+
+            if (PartId & PART_TM)
+            {
+				SceneTracker.NodeTransformChanged(&Node);
+            }
+
+		break;
+
+		case REFMSG_NODE_MATERIAL_CHANGED:
+			SceneTracker.NodeMaterialAssignmentChanged(&Node);
+		break;
+
+		case REFMSG_DISPLAY_MATERIAL_CHANGE:
+			SceneTracker.NodeMaterialGraphModified(&Node);
+		break;
+
+		case REFMSG_SUBANIM_STRUCTURE_CHANGED:
+			if (Node.GetMtl() == TargetHandle)
+			{
+				SceneTracker.NodeMaterialGraphModified(&Node);
+			}
+		break;
+
+		case REFMSG_NODE_RENDERING_PROP_CHANGED:
+			// Handle Renderable flag change. mxs: box.setRenderable
+			SceneTracker.NodePropertiesChanged(&Node);
+		break;
+
+		// todo: appears in Max SDK since 2020 
+		// case REFMSG_NODE_USER_PROPERTY_CHANGED:
+		// break;
+
+		// SceneTracker.NodeGeometryChanged(nodes[NodeIndex]);
+		default: 
+			break;
+		};
+
+
 		return REF_SUCCEED;
 	}
 
-	void AddItem(INode* Node)
+	virtual int NumRefs() override
+	{
+		return 1;
+	}
+
+	virtual RefTargetHandle GetReference(int ReferenceIndex) override
 	{
-		if (!ReferencedItemToIndex.Contains(Node))
+		if (bReferenceSet)
+		{
+			LOG_DEBUG_HEAVY(FString::Printf(TEXT("FNodeNotifications(%s - %p)::GetReference: %d"), Node.GetName(), &Node, ReferenceIndex));
+			return &Node;
+		}
+		else
 		{
-			ReplaceReference(NumRefs(), Node);
+			LOG_DEBUG_HEAVY(FString::Printf(TEXT("FNodeNotifications(UNSET - %p)::GetReference: %d"), &Node, ReferenceIndex));
+			return nullptr;
 		}
 	}
 
-	// todo: unused
-	// RECONSIDER: when this method is used - removed material reduces NumRefs result so adding new material will overwrite already existing reference 
-	// e.g. was two materials added, with index 0 and 1, material 0 removed, NumRefs becomes 1 so next call ReplaceReference(NumRefs(), Material) will replace material 1 in the map
-	void RemoveItem(Mtl* Node)
+	virtual void SetReference(int ReferenceIndex, RefTargetHandle Handle) override
 	{
-		FItemIndex NodeIndex;
-		if (ReferencedItemToIndex.RemoveAndCopyValue(Node, NodeIndex))
+		LOG_DEBUG_HEAVY(FString::Printf(TEXT("FNodeNotifications(%s - %p)::SetReference: %d, %p"), Node.GetName(), &Node, ReferenceIndex, Handle));
+
+		// Either clears reference or passes the same node
+		check((Handle == nullptr ) || (dynamic_cast<INode*>(Handle) == &Node));
+
+		bool bReferenceSetNew = Handle != nullptr;
+
+		if (bReferenceSet && !bReferenceSetNew)
 		{
-			IndexToReferencedItem.Remove(NodeIndex);
+			// Sometimes node deleted without another notification event but with removing it from ReferenceMakers
+			// This is visible with Containers opening/closing them
+			SceneTracker.NodeDeleted(&Node);
 		}
+
+		bReferenceSet = bReferenceSetNew;
 	}
 
-	int NumRefs() override
+private:
+	INode& Node;
+	bool bReferenceSet = false;
+
+
+protected:
+	~FNodeNotifications(){} // Should call Delete
+};
+
+class FNodeObserver
+{
+public:
+	ISceneTracker& SceneTracker;
+
+	FNodeObserver(ISceneTracker& InSceneTracker) : SceneTracker(InSceneTracker)
 	{
-		return IndexToReferencedItem.Num();
 	}
 
-	RefTargetHandle GetReference(int ReferenceIndex) override
+	~FNodeObserver()
 	{
-		RefTargetHandle TargetHandle = IndexToReferencedItem[ReferenceIndex];
-		LOG_DEBUG_HEAVY(FString::Printf(TEXT("FNodeObserver::GetReference: %d, %s"), ReferenceIndex, TargetHandle ? dynamic_cast<INode*>(TargetHandle)->GetName() : TEXT("<null>")));
-		return TargetHandle;
 	}
 
-	void SetReference(int ReferenceIndex, RefTargetHandle TargetHandle) override
+	void AddItem(INode* Node)
 	{
-		LOG_DEBUG_HEAVY(FString::Printf(TEXT("FNodeObserver::SetReference: %d, %s"), ReferenceIndex, TargetHandle?dynamic_cast<INode*>(TargetHandle)->GetName():TEXT("<null>")));
-
-		// todo: investigate why NodeEventNamespace::GetNodeByKey may stil return NULL
-		// testcase - add XRef Material - this will immediately have this 
-		// even though NOTIFY_SCENE_ADDED_NODE was called for node and NOTIFY_SCENE_PRE_DELETED_NODE wasn't!
-		// BUT SeetReference with NULL handle is called
-		// also REFMSG_REF_DELETED and TARGETMSG_DELETING_NODE messages are sent to NotifyRefChanged
-
-		check(!ReferencedItemToIndex.Contains(TargetHandle)); // Not expecting to have same handle under two indices(back-indexing breaks)
-
-		if (TargetHandle)
+		if (FNodeNotifications** NotifierPtr = Nodes.Find(Node))
 		{
-			ReferencedItemToIndex.Add(TargetHandle, ReferenceIndex);
+			(*NotifierPtr)->ReAdded();
 		}
-
-		if (RefTargetHandle* HandlePtr = IndexToReferencedItem.Find(ReferenceIndex))
+		else
 		{
-			if (*HandlePtr)
-			{
-				ReferencedItemToIndex.Remove(*HandlePtr);
-			}
-			*HandlePtr = TargetHandle;
+			Nodes.Add(Node, new FNodeNotifications(SceneTracker, *Node));
 		}
-		else
+	}
+
+	void Reset()
+	{
+		for (TPair<INode*, FNodeNotifications*> Node : Nodes)
 		{
-			IndexToReferencedItem.Add(ReferenceIndex, TargetHandle);
+			Node.Value->Delete();
 		}
+		Nodes.Reset();
 	}
+
 private:
-	TMap<FItemIndex, RefTargetHandle> IndexToReferencedItem;
-	TMap<RefTargetHandle, FItemIndex> ReferencedItemToIndex;
+	
+	TMap<INode*, FNodeNotifications*> Nodes;
 };
 
 FNotifications::FNotifications(IExporter& InExporter)
 	: Exporter(InExporter)
-	, NodeObserver(MakeUnique<FNodeObserver>())
+	, NodeObserver(MakeUnique<FNodeObserver>(InExporter.GetSceneTracker()))
 	, MaterialObserver(MakeUnique<FMaterialObserver>())
 {
 	RegisterForSystemNotifications(); // Always follow system events like file reset, new, open and others. But some system notifications also track changes within scene.
@@ -591,6 +642,15 @@ void FNotifications::On3dsMaxNotification(void* param, NotifyInfo* info)
 			LogDebugNode(NotificationsHandler.ConvertNotificationCodeToString(info->intcode), Node);
 			break;
 		}
+
+		case NOTIFY_NODE_HIDE:
+		case NOTIFY_NODE_UNHIDE:
+		{
+			INode* Node = reinterpret_cast<INode*>(info->callParam);
+			SceneTracker.NodeHideChanged(reinterpret_cast<INode*>(info->callParam));
+			break;
+		}
+
 		case NOTIFY_SCENE_POST_DELETED_NODE:
 		{
 			INode* Node = reinterpret_cast<INode*>(info->callParam);
diff --git a/Engine/Source/Programs/Enterprise/Datasmith/DatasmithMaxExporter/Private/DatasmithMaxExporterWithDirectLink.cpp b/Engine/Source/Programs/Enterprise/Datasmith/DatasmithMaxExporter/Private/DatasmithMaxExporterWithDirectLink.cpp
--- a/Engine/Source/Programs/Enterprise/Datasmith/DatasmithMaxExporter/Private/DatasmithMaxExporterWithDirectLink.cpp
+++ b/Engine/Source/Programs/Enterprise/Datasmith/DatasmithMaxExporter/Private/DatasmithMaxExporterWithDirectLink.cpp
@@ -916,7 +916,12 @@ public:
 	}
 
 	bool bParseXRefScenes = true;
-	bool bIncludeXRefWhileParsing = false; 
+	bool bIncludeXRefWhileParsing = false;
+
+	virtual bool IsUpdateInProgress() override
+	{
+		return bUpdateInProgress;
+	}
 
 	// Parse scene or XRef scene(in this case attach to parent datasmith actor)
 	bool ParseScene(INode* SceneRootNode, FXRefScene XRefScene=FXRefScene())
@@ -1090,9 +1095,8 @@ public:
 
 		DatasmithMaxLogger::Get().Purge();
 
-		NodesPreparer.Start(GetCOREInterface()->GetTime(), bRenderQuality);
-
 		bUpdateInProgress = true;
+		NodesPreparer.Start(GetCOREInterface()->GetTime(), bRenderQuality);
 
 		bool bResult = false;
 
@@ -1117,10 +1121,13 @@ public:
 			bResult = UpdateInternalSafe(ProgressStage.Stage);
 
 		}
-		bUpdateInProgress = false;
 
 		NodesPreparer.Finish();
 
+		// Revert Update flag after RenderEnd was called for all nodes. It fires change events too when
+		// switching from Render to Viewport geometry so we need to know to ignore those events.
+		bUpdateInProgress = false; 
+
 		UberSuspend.Resume();
 		GetCOREInterface()->EnableSceneRedraw();
 		GetCOREInterface()->EnableUndo(true);
@@ -1716,6 +1723,11 @@ public:
 		InvalidatedNodeTrackers.Add(*NodeTracker);
 		PromoteValidity(*NodeTracker, NodeTracker->Validity);
 
+		if (NotificationsHandler)
+		{
+			NotificationsHandler->AddNode(Node);
+		}
+
 		return NodeTracker;
 	}
 
@@ -2623,11 +2635,6 @@ public:
 			return;
 		}
 
-		if (NotificationsHandler)
-		{
-			NotificationsHandler->AddNode(Node);
-		}
-
 		FNodeKey NodeKey = NodeEventNamespace::GetKeyByNode(Node);
 
 		if (FNodeTracker* NodeTracker = GetNodeTracker(NodeKey))
@@ -2680,12 +2687,13 @@ public:
 		}
 	}
 
-	virtual void NodeTransformChanged(FNodeKey NodeKey) override
+	virtual void NodeTransformChanged(INode* Node) override
 	{
-		// todo: invalidate transform only
-
-		// todo: grouping makes this crash. Need to handle event before?
-		InvalidateNode(NodeKey);
+		if (FNodeTracker* NodeTracker =  GetNodeTracker(Node))
+		{
+			//todo: handle more precisely
+			InvalidateNode(*NodeTracker);
+		}
 	}
 
 	virtual void NodeMaterialAssignmentChanged(FNodeKey NodeKey) override
@@ -2694,44 +2702,61 @@ public:
 		InvalidateNode(NodeKey);
 	}
 
+	virtual void NodeMaterialAssignmentChanged(INode* Node) override
+	{
+		if (FNodeTracker* NodeTracker =  GetNodeTracker(Node))
+		{
+			//todo: handle more precisely
+			InvalidateNode(*NodeTracker);
+		}
+	}
+
 	virtual void NodeMaterialGraphModified(FNodeKey NodeKey) override
 	{
-		//identify material tree and update all materials
-		//todo: possible to handle this more precisely(only refresh changed materials) - see FMaterialObserver
+		if (FNodeTracker* NodeTracker =  GetNodeTracker(NodeKey))
+		{
+			NodeMaterialGraphModified(*NodeTracker);
+		}
+	}
 
-		if (FNodeTracker* NodeTracker = GetNodeTracker(NodeKey))
+	virtual void NodeMaterialGraphModified(INode* Node) override
+	{
+		if (FNodeTracker* NodeTracker =  GetNodeTracker(Node))
 		{
-			// todo: investigate why NodeEventNamespace::GetNodeByKey may stil return NULL
-			// testcase - add XRef Material - this will immediately have this
-			// even though NOTIFY_SCENE_ADDED_NODE was called for node and NOTIFY_SCENE_PRE_DELETED_NODE wasn't!
-			INode* Node = NodeEventNamespace::GetNodeByKey(NodeKey);
-			if (Node)
-			{
-				if (Mtl* Material = Node->GetMtl())
-				{
-					MaterialsCollectionTracker.InvalidateMaterial(Material);
-				}
-			}
+			NodeMaterialGraphModified(*NodeTracker);
 		}
+	}
 
-		InvalidateNode(NodeKey); // Invalidate node that has this material assigned. This is needed to trigger rebuild - exported geometry might change(e.g. multimaterial changed to slots will change on static mesh)
+	void NodeMaterialGraphModified(FNodeTracker& NodeTracker)
+	{
+		if (Mtl* Material = NodeTracker.Node->GetMtl())
+		{
+			MaterialsCollectionTracker.InvalidateMaterial(Material);
+		}
+		InvalidateNode(NodeTracker); // Invalidate node that has this material assigned. This is needed to trigger rebuild - exported geometry might change(e.g. multimaterial changed to slots will change on static mesh)
 	}
 
-	virtual void NodeGeometryChanged(FNodeKey NodeKey) override
+	virtual void NodeGeometryChanged(INode* Node) override
 	{
 		// GeometryChanged is executed to handle:
 		// - actual geometry modification(in any way)
 		// - change of baseObject
 
-		InvalidateNode(NodeKey);
+		if (FNodeTracker* NodeTracker =  GetNodeTracker(Node))
+		{
+			InvalidateNode(*NodeTracker);
+		}
 	}
 
-	virtual void NodeHideChanged(FNodeKey NodeKey) override
+	virtual void NodeHideChanged(INode* Node) override
 	{
 		// todo: invalidate visibility only - note to handle this not enought add/remove
 		// actor. make sure to invalidate instances(in case geometry usage changed - like hidden node with multimat), materials
 
-		InvalidateNode(NodeKey);
+		if (FNodeTracker* NodeTracker =  GetNodeTracker(Node))
+		{
+			InvalidateNode(*NodeTracker);
+		}
 	}
 
 	virtual void NodeNameChanged(FNodeKey NodeKey) override
@@ -2743,12 +2768,15 @@ public:
 		}
 	}
 
-	virtual void NodePropertiesChanged(FNodeKey NodeKey) override
+	virtual void NodePropertiesChanged(INode* Node) override
 	{
 		// todo: invalidate visibility only - note to handle this not enought add/remove
 		// actor. make sure to invalidate instances(in case geometry usage changed - like hidden node with multimat), materials
 
-		InvalidateNode(NodeKey);
+		if (FNodeTracker* NodeTracker =  GetNodeTracker(Node))
+		{
+			InvalidateNode(*NodeTracker);
+		}
 	}
 
 	virtual void NodeLinkChanged(FNodeKey NodeKey) override
@@ -2965,7 +2993,7 @@ public:
 	{
 		FUpdateProgress ProgressManager(!bQuiet, 1);
 
-		bool bResult = SceneTracker.Update(ProgressManager.MainStage, false);
+		bool bResult = SceneTracker.Update(ProgressManager.MainStage, true);
 
 		ProgressManager.Finished();
 
@@ -2981,7 +3009,7 @@ public:
 		FUpdateProgress ProgressManager(!bQuiet, 1);
 		FUpdateProgress::FStage& MainStage = ProgressManager.MainStage;
 
-		if (SceneTracker.Update(MainStage, false))
+		if (SceneTracker.Update(MainStage, true))
 		{
 			PROGRESS_STAGE("Sync With DirectLink")
 			UpdateDirectLinkScene();
@@ -3006,7 +3034,7 @@ public:
 		FUpdateProgress ProgressManager(!bQuiet, 1);
 		FUpdateProgress::FStage& MainStage = ProgressManager.MainStage;
 
-		if (SceneTracker.Update(MainStage, false))  // Don't sent redundant update if scene change wasn't detected
+		if (SceneTracker.Update(MainStage, true))  // Don't sent redundant update if scene change wasn't detected
 		{
 			PROGRESS_STAGE("Sync With DirectLink")
 			UpdateDirectLinkScene();
diff --git a/Engine/Source/Programs/Enterprise/Datasmith/DatasmithMaxExporter/Public/DatasmithMaxDirectLink.h b/Engine/Source/Programs/Enterprise/Datasmith/DatasmithMaxExporter/Public/DatasmithMaxDirectLink.h
--- a/Engine/Source/Programs/Enterprise/Datasmith/DatasmithMaxExporter/Public/DatasmithMaxDirectLink.h
+++ b/Engine/Source/Programs/Enterprise/Datasmith/DatasmithMaxExporter/Public/DatasmithMaxDirectLink.h
@@ -365,14 +365,18 @@ public:
 	// Change notifications
 	virtual void NodeAdded(INode* Node) = 0;
 	virtual void NodeDeleted(INode* Node) = 0;
-	virtual void NodeGeometryChanged(FNodeKey NodeKey) = 0;
-	virtual void NodeHideChanged(FNodeKey NodeKey) = 0;
+	virtual void NodeGeometryChanged(INode* Node) = 0;
+	virtual void NodeHideChanged(INode* Node) = 0;
 	virtual void NodeNameChanged(FNodeKey NodeKey) = 0;
-	virtual void NodePropertiesChanged(FNodeKey NodeKey) = 0;
+	virtual void NodePropertiesChanged(INode* Node) = 0;
 	virtual void NodeLinkChanged(FNodeKey NodeKey) = 0;
-	virtual void NodeTransformChanged(FNodeKey NodeKey) = 0;
+	virtual void NodeTransformChanged(INode* Node) = 0;
 	virtual void NodeMaterialAssignmentChanged(FNodeKey NodeKey) = 0;
+	virtual void NodeMaterialAssignmentChanged(INode* Node) = 0;
 	virtual void NodeMaterialGraphModified(FNodeKey NodeKey) = 0;
+	virtual void NodeMaterialGraphModified(INode* NodeKey) = 0;
+
+	virtual bool IsUpdateInProgress() = 0;
 
 	// Scene modification
 	virtual void AddMeshElement(TSharedPtr<IDatasmithMeshElement>& Mesh, FDatasmithMesh& DatasmithMesh, FDatasmithMesh* CollisionMesh) = 0;
@@ -614,35 +618,41 @@ public:
 
 //----
 
-#ifdef LOG_DEBUG_HEAVY_ENABLE
-	#define LOG_DEBUG_HEAVY(message) LogDebug(message)
-#else
-	#define LOG_DEBUG_HEAVY(message) 
-#endif
-
-void LogDebugNode(const FString& Name, class INode* Node);
-void LogNodeEvent(const MCHAR* Name, INodeEventCallback::NodeKeyTab& nodes);
 void LogFlush();
 
 void LogError(const TCHAR* Msg);
 void LogWarning(const TCHAR* Msg);
 void LogCompletion(const TCHAR* Msg);
 void LogInfo(const TCHAR* Msg);
-void LogDebug(const TCHAR* Msg);
 
 void LogError(const FString& Msg);
 void LogWarning(const FString& Msg);
 void LogCompletion(const FString& Msg);
 void LogInfo(const FString& Msg);
-void LogDebug(const FString& Msg);
 
 void LogErrorDialog(const FString& Msg);
 void LogWarningDialog(const FString& Msg);
 void LogCompletionDialog(const FString& Msg);
 void LogInfoDialog(const FString& Msg);
+
+// Debug logging
+void LogDebugImpl(const FString& Msg);
 void LogDebugDialog(const FString& Msg);
 
-}
+#ifdef LOG_DEBUG_HEAVY_ENABLE
+	#define LOG_DEBUG_HEAVY(message) LogDebug(message)
+#else
+	#define LOG_DEBUG_HEAVY(message) 
+#endif
 
+void LogDebugNode(const FString& Name, class INode* Node);
+void LogNodeEvent(const MCHAR* Name, INodeEventCallback::NodeKeyTab& Nodes);
+#ifdef LOG_DEBUG_ENABLE
+	#define LogDebug(Msg) LogDebugImpl(Msg)
+#else
+	#define LogDebug(Msg)
+#endif 
+
+}
 
 #include "Windows/HideWindowsPlatformTypes.h"
