DatasmithRhinoExporter - Fixing Light object tracking in DirectLink.

Rhino does not call the same set of events on light objects modification. And so, those objects were not properly dirtied up. This could also cause an "operation leakage" (where the plugin thought a move operation never completed) which could affect the rest of the tracking.

diff --git a/Engine/Source/Programs/Enterprise/Datasmith/DatasmithRhinoExporter/Private/DirectLink/DatasmithRhinoChangeListener.cs b/Engine/Source/Programs/Enterprise/Datasmith/DatasmithRhinoExporter/Private/DirectLink/DatasmithRhinoChangeListener.cs
--- a/Engine/Source/Programs/Enterprise/Datasmith/DatasmithRhinoExporter/Private/DirectLink/DatasmithRhinoChangeListener.cs
+++ b/Engine/Source/Programs/Enterprise/Datasmith/DatasmithRhinoExporter/Private/DirectLink/DatasmithRhinoChangeListener.cs
@@ -97,11 +97,24 @@ namespace DatasmithRhino.DirectLink
 			//Copied object will call their own creation event.
 			if (!RhinoEventArgs.ObjectsWillBeCopied)
 			{
-				System.Diagnostics.Debug.Assert(OngoingEvent == RhinoOngoingEventTypes.None, "Did not complete previous Object transform before starting a new one");
+				if (OngoingEvent != RhinoOngoingEventTypes.None)
+				{
+					System.Diagnostics.Debug.Fail("Did not complete previous Object transform before starting a new one");
+					EventStack.Clear();
+				}
 
 				for (int ObjectIndex = 0; ObjectIndex < RhinoEventArgs.ObjectCount; ++ObjectIndex)
 				{
-					EventStack.Push(RhinoOngoingEventTypes.MovingActor);
+					RhinoObject CurrentObject = RhinoEventArgs.Objects[ObjectIndex];
+					bool bIsALight = CurrentObject.ObjectType == ObjectType.Light;
+
+					if (!bIsALight)
+					{
+						// All the types that are not lights have OnReplaceRhinoObject(), OnDeleteRhinoObject() and OnAddRhinoObject() events called on them after a move.
+						// We consider that chain of event to be part of the same operation as the move.
+						EventStack.Push(RhinoOngoingEventTypes.MovingActor);
+					}
+
 					TryCatchExecute(() => ExportContext.MoveActor(RhinoEventArgs.Objects[ObjectIndex], RhinoEventArgs.Transform));
 				}
 			}
diff --git a/Engine/Source/Programs/Enterprise/Datasmith/DatasmithRhinoExporter/Private/ExportContext/DatasmithRhinoExportContext.cs b/Engine/Source/Programs/Enterprise/Datasmith/DatasmithRhinoExporter/Private/ExportContext/DatasmithRhinoExportContext.cs
--- a/Engine/Source/Programs/Enterprise/Datasmith/DatasmithRhinoExporter/Private/ExportContext/DatasmithRhinoExportContext.cs
+++ b/Engine/Source/Programs/Enterprise/Datasmith/DatasmithRhinoExporter/Private/ExportContext/DatasmithRhinoExportContext.cs
@@ -401,11 +401,13 @@ namespace DatasmithRhino.ExportContext
 				bHasMesh = true;
 			}
 
-			bool bIsInstance = InRhinoObject is InstanceObject;
-			if ((bHasMesh || bIsInstance)
+			// Block instances and lights don't have a mesh, so we must update their transform directly.
+			bool bNeedsTransformUpdate = InRhinoObject.ObjectType == ObjectType.Light || InRhinoObject.ObjectType == ObjectType.InstanceReference;
+
+			if ((bHasMesh || bNeedsTransformUpdate)
 				&& ObjectIdToHierarchyActorNodeDictionary.TryGetValue(InRhinoObject.Id, out DatasmithActorInfo ActorInfo))
 			{
-				if (bIsInstance)
+				if (bNeedsTransformUpdate)
 				{
 					ActorInfo.ApplyTransform(InTransform);
 				}
