[3ds Max] DirectLink: engine path from registry

- plus small script refactor

diff --git a/Engine/Source/Programs/Enterprise/Datasmith/DatasmithMaxExporter/Plugin/PackageContents.xml b/Engine/Source/Programs/Enterprise/Datasmith/DatasmithMaxExporter/Plugin/PackageContents.xml
--- a/Engine/Source/Programs/Enterprise/Datasmith/DatasmithMaxExporter/Plugin/PackageContents.xml
+++ b/Engine/Source/Programs/Enterprise/Datasmith/DatasmithMaxExporter/Plugin/PackageContents.xml
@@ -55,7 +55,7 @@
 
     <Components Description="pre-start-up scripts parts">
         <RuntimeRequirements OS="Win64" Platform="3ds Max" SeriesMin="2017" SeriesMax="2023" />
-        <ComponentEntry ModuleName="./UI/pre-start-up.ms" />
+        <ComponentEntry ModuleName="./UI/DatasmithPreStartUp.ms" />
         <ComponentEntry ModuleName="./UI/UnrealStudio_Datasmith_Attributes_Holder_Modifier.ms" />
     </Components>
 
@@ -66,7 +66,7 @@
 
     <Components Description="post-start-up scripts parts">
         <RuntimeRequirements OS="Win64" Platform="3ds Max" SeriesMin="2017" SeriesMax="2023" />
-        <ComponentEntry ModuleName="./UI/DatasmithDialog.ms" />
+        <ComponentEntry ModuleName="./UI/DatasmithPostStartUp.ms" />
     </Components>
 
 </ApplicationPackage>
\ No newline at end of file
