Datasmith Revit: Save last view from selector with Revit file

diff --git a/Engine/Source/Programs/Enterprise/Datasmith/DatasmithRevitExporter/DatasmithRevit2018/DatasmithRevit2018.csproj b/Engine/Source/Programs/Enterprise/Datasmith/DatasmithRevitExporter/DatasmithRevit2018/DatasmithRevit2018.csproj
--- a/Engine/Source/Programs/Enterprise/Datasmith/DatasmithRevitExporter/DatasmithRevit2018/DatasmithRevit2018.csproj
+++ b/Engine/Source/Programs/Enterprise/Datasmith/DatasmithRevitExporter/DatasmithRevit2018/DatasmithRevit2018.csproj
@@ -58,6 +58,9 @@
   </ItemGroup>
   <ItemGroup>
     <Compile Include="Properties\AssemblyInfo.cs" />
+    <Compile Include="..\Private\DatasmithRevitDocument.cs">
+      <Link>DatasmithRevitExporter\DatasmithRevitDocument.cs</Link>
+    </Compile>
     <Compile Include="..\Private\DatasmithRevitDocumentData.cs">
       <Link>DatasmithRevitExporter\DatasmithRevitDocumentData.cs</Link>
     </Compile>
diff --git a/Engine/Source/Programs/Enterprise/Datasmith/DatasmithRevitExporter/DatasmithRevit2019/DatasmithRevit2019.csproj b/Engine/Source/Programs/Enterprise/Datasmith/DatasmithRevitExporter/DatasmithRevit2019/DatasmithRevit2019.csproj
--- a/Engine/Source/Programs/Enterprise/Datasmith/DatasmithRevitExporter/DatasmithRevit2019/DatasmithRevit2019.csproj
+++ b/Engine/Source/Programs/Enterprise/Datasmith/DatasmithRevitExporter/DatasmithRevit2019/DatasmithRevit2019.csproj
@@ -58,6 +58,9 @@
   </ItemGroup>
   <ItemGroup>
     <Compile Include="Properties\AssemblyInfo.cs" />
+    <Compile Include="..\Private\DatasmithRevitDocument.cs">
+      <Link>DatasmithRevitExporter\DatasmithRevitDocument.cs</Link>
+    </Compile>
     <Compile Include="..\Private\DatasmithRevitDocumentData.cs">
       <Link>DatasmithRevitExporter\DatasmithRevitDocumentData.cs</Link>
     </Compile>
diff --git a/Engine/Source/Programs/Enterprise/Datasmith/DatasmithRevitExporter/DatasmithRevit2020/DatasmithRevit2020.csproj b/Engine/Source/Programs/Enterprise/Datasmith/DatasmithRevitExporter/DatasmithRevit2020/DatasmithRevit2020.csproj
--- a/Engine/Source/Programs/Enterprise/Datasmith/DatasmithRevitExporter/DatasmithRevit2020/DatasmithRevit2020.csproj
+++ b/Engine/Source/Programs/Enterprise/Datasmith/DatasmithRevitExporter/DatasmithRevit2020/DatasmithRevit2020.csproj
@@ -58,6 +58,9 @@
   </ItemGroup>
   <ItemGroup>
     <Compile Include="Properties\AssemblyInfo.cs" />
+    <Compile Include="..\Private\DatasmithRevitDocument.cs">
+      <Link>DatasmithRevitExporter\DatasmithRevitDocument.cs</Link>
+    </Compile>
     <Compile Include="..\Private\DatasmithRevitApplication.cs">
       <Link>DatasmithRevitExporter\DatasmithRevitApplication.cs</Link>
     </Compile>
diff --git a/Engine/Source/Programs/Enterprise/Datasmith/DatasmithRevitExporter/DatasmithRevit2021/DatasmithRevit2021.csproj b/Engine/Source/Programs/Enterprise/Datasmith/DatasmithRevitExporter/DatasmithRevit2021/DatasmithRevit2021.csproj
--- a/Engine/Source/Programs/Enterprise/Datasmith/DatasmithRevitExporter/DatasmithRevit2021/DatasmithRevit2021.csproj
+++ b/Engine/Source/Programs/Enterprise/Datasmith/DatasmithRevitExporter/DatasmithRevit2021/DatasmithRevit2021.csproj
@@ -58,6 +58,9 @@
   </ItemGroup>
   <ItemGroup>
     <Compile Include="Properties\AssemblyInfo.cs" />
+    <Compile Include="..\Private\DatasmithRevitDocument.cs">
+      <Link>DatasmithRevitExporter\DatasmithRevitDocument.cs</Link>
+    </Compile>
     <Compile Include="..\Private\DatasmithRevitDocumentData.cs">
       <Link>DatasmithRevitExporter\DatasmithRevitDocumentData.cs</Link>
     </Compile>
diff --git a/Engine/Source/Programs/Enterprise/Datasmith/DatasmithRevitExporter/DatasmithRevit2022/DatasmithRevit2022.csproj b/Engine/Source/Programs/Enterprise/Datasmith/DatasmithRevitExporter/DatasmithRevit2022/DatasmithRevit2022.csproj
--- a/Engine/Source/Programs/Enterprise/Datasmith/DatasmithRevitExporter/DatasmithRevit2022/DatasmithRevit2022.csproj
+++ b/Engine/Source/Programs/Enterprise/Datasmith/DatasmithRevitExporter/DatasmithRevit2022/DatasmithRevit2022.csproj
@@ -58,6 +58,9 @@
   </ItemGroup>
   <ItemGroup>
     <Compile Include="Properties\AssemblyInfo.cs" />
+    <Compile Include="..\Private\DatasmithRevitDocument.cs">
+      <Link>DatasmithRevitExporter\DatasmithRevitDocument.cs</Link>
+    </Compile>
     <Compile Include="..\Private\DatasmithRevitDocumentData.cs">
       <Link>DatasmithRevitExporter\DatasmithRevitDocumentData.cs</Link>
     </Compile>
diff --git a/Engine/Source/Programs/Enterprise/Datasmith/DatasmithRevitExporter/DatasmithRevit2023/DatasmithRevit2023.csproj b/Engine/Source/Programs/Enterprise/Datasmith/DatasmithRevitExporter/DatasmithRevit2023/DatasmithRevit2023.csproj
--- a/Engine/Source/Programs/Enterprise/Datasmith/DatasmithRevitExporter/DatasmithRevit2023/DatasmithRevit2023.csproj
+++ b/Engine/Source/Programs/Enterprise/Datasmith/DatasmithRevitExporter/DatasmithRevit2023/DatasmithRevit2023.csproj
@@ -58,6 +58,9 @@
   </ItemGroup>
   <ItemGroup>
     <Compile Include="Properties\AssemblyInfo.cs" />
+    <Compile Include="..\Private\DatasmithRevitDocument.cs">
+      <Link>DatasmithRevitExporter\DatasmithRevitDocument.cs</Link>
+    </Compile>
     <Compile Include="..\Private\DatasmithRevitDocumentData.cs">
       <Link>DatasmithRevitExporter\DatasmithRevitDocumentData.cs</Link>
     </Compile>
diff --git a/Engine/Source/Programs/Enterprise/Datasmith/DatasmithRevitExporter/Private/DatasmithDynamoNode.cs b/Engine/Source/Programs/Enterprise/Datasmith/DatasmithRevitExporter/Private/DatasmithDynamoNode.cs
--- a/Engine/Source/Programs/Enterprise/Datasmith/DatasmithRevitExporter/Private/DatasmithDynamoNode.cs
+++ b/Engine/Source/Programs/Enterprise/Datasmith/DatasmithRevitExporter/Private/DatasmithDynamoNode.cs
@@ -70,6 +70,7 @@ namespace DSNode
 			FDatasmithRevitExportContext ExportContext = new FDatasmithRevitExportContext(
 				App,
 				InDocument,
+				FDocument.ActiveDocument.Settings,
 				FilePaths,
 				ExportOptions,
 				null);
diff --git a/Engine/Source/Programs/Enterprise/Datasmith/DatasmithRevitExporter/Private/DatasmithRevitApplication.cs b/Engine/Source/Programs/Enterprise/Datasmith/DatasmithRevitExporter/Private/DatasmithRevitApplication.cs
--- a/Engine/Source/Programs/Enterprise/Datasmith/DatasmithRevitExporter/Private/DatasmithRevitApplication.cs
+++ b/Engine/Source/Programs/Enterprise/Datasmith/DatasmithRevitExporter/Private/DatasmithRevitApplication.cs
@@ -7,7 +7,6 @@ using System.IO;
 using System.Linq;
 using System.Reflection;
 using System.Windows.Media.Imaging;
-using Autodesk.Revit.ApplicationServices;
 using Autodesk.Revit.DB;
 using Autodesk.Revit.DB.Events;
 using Autodesk.Revit.UI;
@@ -35,7 +34,7 @@ namespace DatasmithRevitExporter
 			{
 				if (InData.GetAddedElementIds().Count > 0 || InData.GetDeletedElementIds().Count > 0 || InData.GetModifiedElementIds().Count > 0)
 				{
-					DatasmithRevitApplication.Instance.SetViewList();
+					Instance.SetViewList(ElementId.InvalidElementId);
 				}
 			}
 
@@ -67,12 +66,12 @@ namespace DatasmithRevitExporter
 		private EventHandler<DocumentClosingEventArgs> DocumentClosingHandler;
 		private EventHandler<IdlingEventArgs> IdlingEventHandler;
 		private EventHandler<DocumentOpenedEventArgs> DocumentOpenedHandler;
+		private EventHandler<DocumentCreatedEventArgs> DocumentCreatedHandler;
+		private EventHandler<ViewActivatedEventArgs> ViewActivatedHandler;
 
 		private PushButton AutoSyncPushButton;
 		private PushButton SyncPushButton;
 
-		private Document ActiveDocument = null;
-
 		private View3DUpdater ViewsUpdater;
 		private ElementId SelectedViewId = ElementId.InvalidElementId;
 		private ComboBox ComboViews;
@@ -111,7 +110,7 @@ namespace DatasmithRevitExporter
 				if (SelectedView != null)
 				{
 					SelectedViewId = SelectedView.Id;
-					FDirectLink.ActivateInstance(SelectedView);
+					FDocument.ActiveDocument?.SetActiveDirectLinkInstance(SelectedView);
 				}
 			}
 		}
@@ -128,15 +127,15 @@ namespace DatasmithRevitExporter
 			SelectedViewId = ElementId.InvalidElementId;
 		}
 
-		private void SetViewList()
+		private void SetViewList(ElementId InActiveViewId)
 		{
-			if (ActiveDocument == null)
+			if (FDocument.ActiveDocument == null)
 			{
 				return;
 			}
 
 			// Cache prev selected view
-			All3DViews = new FilteredElementCollector(ActiveDocument).OfClass(typeof(View3D)).Cast<View3D>().ToList();
+			All3DViews = new FilteredElementCollector(FDocument.ActiveDocument.RevitDoc).OfClass(typeof(View3D)).Cast<View3D>().ToList();
 			All3DViews.RemoveAll(view => (view.IsTemplate || !view.CanBePrinted));
 
 			// Combo box does not allow removal of items :). This is a workaround to always only add items and 
@@ -164,6 +163,11 @@ namespace DatasmithRevitExporter
 				ComboViews.GetItems()[ViewIndex].Visible = false;
 			}
 
+			if (InActiveViewId != ElementId.InvalidElementId)
+			{
+				SelectedViewId = InActiveViewId;
+			}
+
 			// Set the active view
 			View3D ComboActiveView = null;
 
@@ -185,7 +189,7 @@ namespace DatasmithRevitExporter
 					{
 						ComboViews.Current = Item;
 						SelectedViewId = ComboActiveView.Id;
-						FDirectLink.ActivateInstance(ComboActiveView);
+						FDocument.ActiveDocument?.SetActiveDirectLinkInstance(ComboActiveView);
 						break;
 					}
 				}
@@ -288,12 +292,18 @@ namespace DatasmithRevitExporter
 			DocumentOpenedHandler = new EventHandler<DocumentOpenedEventArgs>(OnDocumentOpened);
 			InApplication.ControlledApplication.DocumentOpened += DocumentOpenedHandler;
 
+			DocumentCreatedHandler = new EventHandler<DocumentCreatedEventArgs>(OnDocumentCreated);
+			InApplication.ControlledApplication.DocumentCreated += DocumentCreatedHandler;
+
 			DocumentClosingHandler = new EventHandler<DocumentClosingEventArgs>(OnDocumentClosing);
 			InApplication.ControlledApplication.DocumentClosing += DocumentClosingHandler;
 
 			IdlingEventHandler = new EventHandler<IdlingEventArgs>(OnIdling);
 			InApplication.Idling += IdlingEventHandler;
 
+			ViewActivatedHandler = new EventHandler<ViewActivatedEventArgs>(OnViewActivated);
+			InApplication.ViewActivated += ViewActivatedHandler;
+
 			// Setup Direct Link
 
 			string RevitEngineDir = null;
@@ -318,8 +328,6 @@ namespace DatasmithRevitExporter
 
 			Debug.Assert(bDirectLinkInitOk);
 
-			FSettingsManager.Init(InApplication);
-
 			// Register updater to react to view modification
 
 			ViewsUpdater = new View3DUpdater(InApplication.ControlledApplication.ActiveAddInId);
@@ -340,14 +348,32 @@ namespace DatasmithRevitExporter
 
 		static void OnDocumentOpened(object sender, DocumentOpenedEventArgs e)
 		{
-			Instance.ActiveDocument = e.Document;
-			Instance.SetViewList();
+			FDocument.SetActiveDocument(e.Document);
+			Instance.SetViewList(FDocument.ActiveDocument?.Settings.SyncViewId ?? ElementId.InvalidElementId);
+		}
+
+		static void OnDocumentCreated(object sender, DocumentCreatedEventArgs e)
+		{
+			FDocument.SetActiveDocument(e.Document);
+			Instance.SetViewList(ElementId.InvalidElementId);
 		}
 
 		static void OnDocumentClosing(object sender, DocumentClosingEventArgs e)
 		{
 			Instance.ClearViews();
-			FDirectLink.DestroyInstancesForDocument(e.Document, e.Document.Application);
+			FDocument.Destroy(e.Document);
+		}
+
+		static void OnViewActivated(object sender, ViewActivatedEventArgs e)
+		{
+			View Previous = e.PreviousActiveView;
+			View Current = e.CurrentActiveView;
+
+			if (Previous != null && !Previous.Document.Equals(Current.Document))
+			{
+				FDocument.SetActiveDocument(e.Document);
+				Instance.SetViewList(FDocument.ActiveDocument?.Settings.SyncViewId ?? ElementId.InvalidElementId);
+			}
 		}
 
 		// Implement the interface to execute some tasks when Revit shuts down.
@@ -355,11 +381,17 @@ namespace DatasmithRevitExporter
 			UIControlledApplication InApplication // handle to the application being shut down
 		)
 		{
+			FDocument.DestroyAll();
+
 			InApplication.ControlledApplication.DocumentClosing -= DocumentClosingHandler;
 			InApplication.ControlledApplication.DocumentOpened -= DocumentOpenedHandler;
+			InApplication.ControlledApplication.DocumentCreated -= DocumentCreatedHandler;
+			InApplication.ViewActivated -= ViewActivatedHandler;
 
 			DocumentClosingHandler = null;
 			DocumentOpenedHandler = null;
+			DocumentCreatedHandler = null;
+			ViewActivatedHandler = null;
 
 			UpdaterRegistry.UnregisterUpdater(ViewsUpdater.GetUpdaterId());
 			ViewsUpdater = null;
@@ -368,7 +400,6 @@ namespace DatasmithRevitExporter
 			{
 				ExportMessagesDialog.Close();
 			}
-			FSettingsManager.Destroy(InApplication);
 			FDatasmithFacadeDirectLink.Shutdown();
 			return Result.Succeeded;
 		}
diff --git a/Engine/Source/Programs/Enterprise/Datasmith/DatasmithRevitExporter/Private/DatasmithRevitCommand.cs b/Engine/Source/Programs/Enterprise/Datasmith/DatasmithRevitExporter/Private/DatasmithRevitCommand.cs
--- a/Engine/Source/Programs/Enterprise/Datasmith/DatasmithRevitExporter/Private/DatasmithRevitCommand.cs
+++ b/Engine/Source/Programs/Enterprise/Datasmith/DatasmithRevitExporter/Private/DatasmithRevitCommand.cs
@@ -62,7 +62,7 @@ namespace DatasmithRevitExporter
 		{
 			UIDocument UIDoc = InCommandData.Application.ActiveUIDocument;
 			Document Doc = UIDoc.Document;
-			View3D ActiveView = FDirectLink.GetSyncView();
+			View3D ActiveView = FDocument.ActiveDocument?.ActiveDirectLinkInstance?.SyncView;
 
 			if (ActiveView == null)
 			{
@@ -78,14 +78,15 @@ namespace DatasmithRevitExporter
 				return Result.Cancelled;
 			}
 
-			Debug.Assert(FDirectLink.Get() != null);
+			Debug.Assert(FDocument.ActiveDocument?.ActiveDirectLinkInstance != null);
 
 			FDatasmithRevitExportContext ExportContext = new FDatasmithRevitExportContext(
 				InCommandData.Application.Application,
 				Doc,
+				FDocument.ActiveDocument.Settings,
 				null,
 				new DatasmithRevitExportOptions(Doc),
-				FDirectLink.Get());
+				FDocument.ActiveDocument.ActiveDirectLinkInstance);
 
 			// Export the active 3D View to the given Unreal Datasmith file.
 			using (CustomExporter Exporter = new CustomExporter(Doc, ExportContext))
@@ -178,7 +179,7 @@ namespace DatasmithRevitExporter
 
 			if (ExportActiveViewOnly)
 			{
-				View3D ActiveView = FDirectLink.GetSyncView();
+				View3D ActiveView = FDocument.ActiveDocument?.ActiveDirectLinkInstance?.SyncView;
 
 				if (ActiveView == null)
 				{
@@ -288,6 +289,7 @@ namespace DatasmithRevitExporter
 			FDatasmithRevitExportContext ExportContext = new FDatasmithRevitExportContext(
 				InCommandData.Application.Application,
 				Doc,
+				FDocument.ActiveDocument.Settings,
 				FilePaths,
 				ExportOptions,
 				null);
@@ -422,7 +424,7 @@ namespace DatasmithRevitExporter
 	{
 		public Result Execute(ExternalCommandData InCommandData, ref string OutCommandMessage, ElementSet OutElements)
 		{
-			DatasmithRevitSettingsDialog ExportOptions = new DatasmithRevitSettingsDialog(InCommandData.Application.ActiveUIDocument.Document);
+			DatasmithRevitSettingsDialog ExportOptions = new DatasmithRevitSettingsDialog(InCommandData.Application.ActiveUIDocument.Document, FDocument.ActiveDocument?.Settings);
 			ExportOptions.ShowDialog();
 			return Result.Succeeded;
 		}
diff --git a/Engine/Source/Programs/Enterprise/Datasmith/DatasmithRevitExporter/Private/DatasmithRevitDirectLink.cs b/Engine/Source/Programs/Enterprise/Datasmith/DatasmithRevitExporter/Private/DatasmithRevitDirectLink.cs
--- a/Engine/Source/Programs/Enterprise/Datasmith/DatasmithRevitExporter/Private/DatasmithRevitDirectLink.cs
+++ b/Engine/Source/Programs/Enterprise/Datasmith/DatasmithRevitExporter/Private/DatasmithRevitDirectLink.cs
@@ -144,7 +144,7 @@ namespace DatasmithRevitExporter
 		private FCachedDocumentData										RootCache = null;
 		private FCachedDocumentData										CurrentCache = null;
 		
-		private View3D													SyncView = null;
+		public View3D													SyncView { get; private set; } = null;
 
 		private HashSet<Document>										ModifiedLinkedDocuments = new HashSet<Document>();
 		private HashSet<ElementId>										ExportedLinkedDocuments = new HashSet<ElementId>();
@@ -168,17 +168,14 @@ namespace DatasmithRevitExporter
 
 		private EventHandler<DocumentChangedEventArgs>					DocumentChangedHandler;
 
-		private static FDirectLink										ActiveInstance = null;
-		private static List<FDirectLink> Instances = new List<FDirectLink>();
+		private FSettings												Settings = null;
 
 		private static UIApplication UIApp { get; set; } = null;
 
 		private bool bHasChanges = false;
 		private bool bSyncInProgress = false;
 
-		private EventHandler SettingsChangedHandler;
-
-		private bool bSettingsDirty = false;
+		public bool bSettingsDirty = false;
 
 		private static bool _bAutoSync = false;
 		public static bool bAutoSync
@@ -192,117 +189,15 @@ namespace DatasmithRevitExporter
 				_bAutoSync = value;
 				if (_bAutoSync)
 				{
-					Get()?.RunAutoSync();
-				}
-			}
-		}
-
-		public static FDirectLink Get()
-		{
-			return ActiveInstance;
-		}
-
-		public static View3D GetSyncView()
-		{
-			return ActiveInstance?.SyncView ?? null;
-		}
-
-		public static void ActivateInstance(View3D InView)
-		{
-			if (UIApp == null)
-			{
-				UIApp = new UIApplication(InView.Document.Application);
-			}
-
-			// Disable existing instance, if there's active one.
-			ActiveInstance?.MakeActive(false);
-			ActiveInstance = null;
-
-			// Find out if we already have instance for this document and 
-			// activate it if we do. Otherwise, create new one.
-
-			FDirectLink InstanceToActivate = null;
-
-			foreach (FDirectLink DL in Instances)
-			{
-				if (DL.SyncView == null || !DL.SyncView.IsValidObject)
-				{
-					continue;
-				}
-
-				if (DL.RootCache.SourceDocument.Equals(InView.Document) && DL.SyncView.Id == InView.Id)
-				{
-					InstanceToActivate = DL;
-					break;
-				}
-			}
-
-			if (InstanceToActivate == null)
-			{
-				InstanceToActivate = new FDirectLink(InView);
-				Instances.Add(InstanceToActivate);
-			}
-
-			InstanceToActivate.MakeActive(true);
-			ActiveInstance = InstanceToActivate;
-		}
-
-		public static FDirectLink FindInstance(View3D InView)
-		{
-			foreach (var Inst in Instances)
-			{
-				if (Inst.SyncView.Id == InView.Id)
-				{
-					return Inst;
+					FDocument.ActiveDocument?.ActiveDirectLinkInstance?.RunAutoSync();
 				}
 			}
-			return null;
-		}
-
-		public static void DestroyInstance(FDirectLink Instance, Application InApp)
-		{
-			if (ActiveInstance == Instance)
-			{
-				ActiveInstance = null;
-			}
-			FSettingsManager.SettingsUpdated -= Instance.SettingsChangedHandler;
-			Instances.Remove(Instance);
-			Instance?.Destroy(InApp);
-		}
-
-		public static void DestroyInstancesForDocument(Document InDoc, Application InApp)
-		{
-			for (int InstanceIndex = Instances.Count - 1; InstanceIndex >= 0; --InstanceIndex)
-			{
-				FDirectLink Instance = Instances[InstanceIndex];
-
-				if (!Instance.RootCache.SourceDocument.Equals(InDoc))
-				{
-					continue;
-				}
-
-				if (ActiveInstance == Instance)
-				{
-					ActiveInstance = null;
-				}
-
-				Instances.RemoveAt(InstanceIndex);
-				Instance?.Destroy(InApp);
-			}
-		}
-
-		public static void DestroyAllInstances(Application InApp) 
-		{
-			foreach (FDirectLink DL in Instances)
-			{
-				DestroyInstance(DL, InApp);
-			}
-
-			Instances.Clear();
 		}
 
 		public static void OnApplicationIdle()
 		{
+			FDirectLink ActiveInstance = FDocument.ActiveDocument?.ActiveDirectLinkInstance ?? null;
+
 			if (ActiveInstance == null || ActiveInstance.SyncCount == 0)
 			{
 				return;
@@ -327,32 +222,32 @@ namespace DatasmithRevitExporter
 		  object InSender,
 		  DocumentChangedEventArgs InArgs) 
 		{
-			FDirectLink DirectLink = FDirectLink.Get();
+			FDirectLink ActiveInstance = FDocument.ActiveDocument?.ActiveDirectLinkInstance ?? null;
 
-			Debug.Assert(DirectLink != null);
+			Debug.Assert(ActiveInstance != null);
 
 			// Handle modified elements
 			foreach (ElementId ElemId in InArgs.GetModifiedElementIds())
 			{
-				Element ModifiedElement = DirectLink.RootCache.SourceDocument.GetElement(ElemId);
-			
-				DirectLink.bHasChanges = true;
+				Element ModifiedElement = ActiveInstance.RootCache.SourceDocument.GetElement(ElemId);
+
+				ActiveInstance.bHasChanges = true;
 
 				if (ModifiedElement != null)
 				{
 					if (ModifiedElement.GetType() == typeof(RevitLinkInstance))
 					{
-						DirectLink.ModifiedLinkedDocuments.Add((ModifiedElement as RevitLinkInstance).GetLinkDocument());
+						ActiveInstance.ModifiedLinkedDocuments.Add((ModifiedElement as RevitLinkInstance).GetLinkDocument());
 					}
 					else if (ModifiedElement.GetType() == typeof(RevitLinkType))
 					{
-						foreach (KeyValuePair<ElementId, FCachedDocumentData> Link in DirectLink.RootCache.LinkedDocumentsCache)
+						foreach (KeyValuePair<ElementId, FCachedDocumentData> Link in ActiveInstance.RootCache.LinkedDocumentsCache)
 						{
-							RevitLinkInstance LinkInstance = DirectLink.RootCache.SourceDocument.GetElement(Link.Key) as RevitLinkInstance;
+							RevitLinkInstance LinkInstance = ActiveInstance.RootCache.SourceDocument.GetElement(Link.Key) as RevitLinkInstance;
 							if (LinkInstance != null)
 							{
-								RevitLinkType LinkType = DirectLink.RootCache.SourceDocument.GetElement(LinkInstance.GetTypeId()) as RevitLinkType;
-								if (LinkType != null && LinkType.Id == ModifiedElement.Id && RevitLinkType.IsLoaded(DirectLink.RootCache.SourceDocument, LinkType.Id))
+								RevitLinkType LinkType = ActiveInstance.RootCache.SourceDocument.GetElement(LinkInstance.GetTypeId()) as RevitLinkType;
+								if (LinkType != null && LinkType.Id == ModifiedElement.Id && RevitLinkType.IsLoaded(ActiveInstance.RootCache.SourceDocument, LinkType.Id))
 								{
 									Link.Value.SetAllElementsModified();
 								}
@@ -368,30 +263,29 @@ namespace DatasmithRevitExporter
 						{
 							foreach (ElementId DepElemId in DependentElements)
 							{
-								DirectLink.RootCache.SetElementModified(true, DepElemId);
+								ActiveInstance.RootCache.SetElementModified(true, DepElemId);
 							}
 						}
 					}
 
-					DirectLink.RootCache.SetElementModified(true, ElemId);
+					ActiveInstance.RootCache.SetElementModified(true, ElemId);
 				}
 			}
 
-			DirectLink.bHasChanges = DirectLink.bHasChanges || InArgs.GetDeletedElementIds().Any();
-			DirectLink.bHasChanges = DirectLink.bHasChanges || InArgs.GetAddedElementIds().Any();
+			ActiveInstance.bHasChanges = ActiveInstance.bHasChanges || InArgs.GetDeletedElementIds().Any();
+			ActiveInstance.bHasChanges = ActiveInstance.bHasChanges || InArgs.GetAddedElementIds().Any();
 
-			if (DirectLink.bHasChanges && bAutoSync)
+			if (ActiveInstance.bHasChanges && bAutoSync)
 			{
-				DirectLink.RunAutoSync();
+				ActiveInstance.RunAutoSync();
 			}
 		}
 
-		private FDirectLink(View3D InView)
+		public FDirectLink(View3D InView, FSettings InSettings)
 		{
+			Settings = InSettings;
 			RootCache = new FCachedDocumentData(InView.Document);
-			
 			CurrentCache = RootCache;
-
 			SyncView = InView;
 
 			DatasmithScene = new FDatasmithFacadeScene(
@@ -407,12 +301,11 @@ namespace DatasmithRevitExporter
 
 			DocumentChangedHandler = new EventHandler<DocumentChangedEventArgs>(OnDocumentChanged);
 			InView.Document.Application.DocumentChanged += DocumentChangedHandler;
+		}
 
-			SettingsChangedHandler = new EventHandler((object Sender, EventArgs Args) => 
-			{
-				bSettingsDirty = true;
-			});
-			FSettingsManager.SettingsUpdated += SettingsChangedHandler;
+		public Document GetRootDocument()
+		{
+			return RootCache.SourceDocument;
 		}
 
 		private void RunAutoSync()
@@ -450,7 +343,7 @@ namespace DatasmithRevitExporter
 			}
 		}
 
-		private void MakeActive(bool bInActive)
+		public void MakeActive(bool bInActive)
 		{
 			if (!bInActive)
 			{
@@ -467,7 +360,7 @@ namespace DatasmithRevitExporter
 			}
 		}
 
-		private void Destroy(Application InApp)
+		public void Destroy(Application InApp)
 		{
 			InApp.DocumentChanged -= DocumentChangedHandler;
 			DocumentChangedHandler = null;
@@ -878,7 +771,7 @@ namespace DatasmithRevitExporter
 					ElementData.ElementMetaData.SetLabel(Actor.GetLabel());
 					ElementData.ElementMetaData.SetAssociatedElement(Actor);
 
-					FUtils.AddActorMetadata(RevitElement, ElementData.ElementMetaData);
+					FUtils.AddActorMetadata(RevitElement, ElementData.ElementMetaData, Settings);
 
 					DatasmithScene.AddMetaData(ElementData.ElementMetaData);
 
diff --git a/Engine/Source/Programs/Enterprise/Datasmith/DatasmithRevitExporter/Private/DatasmithRevitDocument.cs b/Engine/Source/Programs/Enterprise/Datasmith/DatasmithRevitExporter/Private/DatasmithRevitDocument.cs
new file mode 100644
--- /dev/null
+++ b/Engine/Source/Programs/Enterprise/Datasmith/DatasmithRevitExporter/Private/DatasmithRevitDocument.cs
@@ -0,0 +1,155 @@
+// Copyright Epic Games, Inc. All Rights Reserved.
+
+using System;
+using System.Collections.Generic;
+using Autodesk.Revit.DB;
+using Autodesk.Revit.UI;
+
+namespace DatasmithRevitExporter
+{
+	public class FDocument
+	{
+		public static FDocument ActiveDocument { get; private set; } = null;
+		public static List<FDocument> AllDocuments {  get; private set; } = new List<FDocument>();
+		public List<FDirectLink> DirectLinkInstances { get; private set; }  = new List<FDirectLink>();
+		public FDirectLink ActiveDirectLinkInstance { get; private set; } = null;
+		public Document RevitDoc { get; private set; } = null;
+		public FSettings Settings { get; private set; } = null;
+		private static UIApplication UIApp { get; set; } = null;
+
+		private EventHandler SettingsChangedHandler;
+
+		public static void SetActiveDocument(Document InRevitDoc)
+		{
+			FDocument NewActiveDocument = null;
+
+			foreach (FDocument Doc in AllDocuments)
+			{
+				if (Doc.RevitDoc.Equals(InRevitDoc))
+				{
+					NewActiveDocument = Doc;
+					break;
+				}
+			}
+
+			if (NewActiveDocument == null)
+			{
+				NewActiveDocument = new FDocument();
+				NewActiveDocument.RevitDoc = InRevitDoc;
+				NewActiveDocument.Settings = FSettingsManager.ReadSettings(InRevitDoc);
+
+				if (UIApp == null)
+				{
+					UIApp = new UIApplication(InRevitDoc.Application);
+				}
+
+				NewActiveDocument.SettingsChangedHandler = new EventHandler((object Sender, EventArgs Args) =>
+				{
+					NewActiveDocument.ActiveDirectLinkInstance.bSettingsDirty = true;
+				});
+
+				FSettingsManager.SettingsUpdated += NewActiveDocument.SettingsChangedHandler;
+
+				AllDocuments.Add(NewActiveDocument);
+			}
+
+			ActiveDocument = NewActiveDocument;
+		}
+
+		public static void Destroy(Document InDocument)
+		{
+			FDocument DocumentToDestroy = null;
+
+			// Find document instance
+			foreach (FDocument Doc in AllDocuments)
+			{
+				if (Doc.RevitDoc.Equals(InDocument))
+				{
+					DocumentToDestroy = Doc;
+					break;
+				}
+			}
+
+			if (DocumentToDestroy != null)
+			{
+				DestroyInstance(DocumentToDestroy);
+			}
+		}
+
+		private static void DestroyInstance(FDocument InDoc)
+		{
+			FSettingsManager.SettingsUpdated -= InDoc.SettingsChangedHandler;
+			InDoc.SettingsChangedHandler = null;
+
+			// Destroy active DirectLink instances
+			for (int InstanceIndex = InDoc.DirectLinkInstances.Count - 1; InstanceIndex >= 0; --InstanceIndex)
+			{
+				FDirectLink Instance = InDoc.DirectLinkInstances[InstanceIndex];
+
+				if (InDoc.ActiveDirectLinkInstance == Instance)
+				{
+					InDoc.ActiveDirectLinkInstance = null;
+				}
+
+				InDoc.DirectLinkInstances.RemoveAt(InstanceIndex);
+				Instance?.Destroy(UIApp.Application);
+			}
+
+			if (ActiveDocument == InDoc)
+			{
+				ActiveDocument = null;
+			}
+
+			AllDocuments.Remove(InDoc);
+		}
+
+		public static void DestroyAll()
+		{
+			foreach (FDocument Doc in AllDocuments)
+			{
+				DestroyInstance(Doc);
+			}
+		}
+
+		public void SetActiveDirectLinkInstance(View3D InView)
+		{
+			// Disable existing instance, if there's active one.
+			ActiveDirectLinkInstance?.MakeActive(false);
+			ActiveDirectLinkInstance = null;
+
+			// Find out if we already have instance for this document and 
+			// activate it if we do. Otherwise, create new one.
+
+			FDirectLink InstanceToActivate = null;
+
+			foreach (FDirectLink DL in DirectLinkInstances)
+			{
+				if (DL.SyncView == null || !DL.SyncView.IsValidObject)
+				{
+					continue;
+				}
+
+				if (DL.GetRootDocument().Equals(InView.Document) && DL.SyncView.Id == InView.Id)
+				{
+					InstanceToActivate = DL;
+					break;
+				}
+			}
+
+			if (InstanceToActivate == null)
+			{
+				InstanceToActivate = new FDirectLink(InView, Settings);
+				DirectLinkInstances.Add(InstanceToActivate);
+			}
+
+			InstanceToActivate.MakeActive(true);
+			ActiveDirectLinkInstance = InstanceToActivate;
+
+			if (Settings.SyncViewId.IntegerValue != InView.Id.IntegerValue)
+			{
+				Settings.SyncViewId = InView.Id;
+				FSettingsManager.WriteSettings(RevitDoc, Settings);
+			}
+		}
+	}
+}
diff --git a/Engine/Source/Programs/Enterprise/Datasmith/DatasmithRevitExporter/Private/DatasmithRevitDocumentData.cs b/Engine/Source/Programs/Enterprise/Datasmith/DatasmithRevitExporter/Private/DatasmithRevitDocumentData.cs
--- a/Engine/Source/Programs/Enterprise/Datasmith/DatasmithRevitExporter/Private/DatasmithRevitDocumentData.cs
+++ b/Engine/Source/Programs/Enterprise/Datasmith/DatasmithRevitExporter/Private/DatasmithRevitDocumentData.cs
@@ -1052,12 +1052,12 @@ namespace DatasmithRevitExporter
 				}
 
 				// Add Revit element metadata to the Datasmith actor.
-				FUtils.AddActorMetadata(CurrentElement, "Element*", ElementMetaData);
+				FUtils.AddActorMetadata(CurrentElement, "Element*", ElementMetaData, DocumentData.CurrentSettings);
 
 				if (BaseElementType != null)
 				{
 					// Add Revit element type metadata to the Datasmith actor.
-					FUtils.AddActorMetadata(BaseElementType, "Type*", ElementMetaData);
+					FUtils.AddActorMetadata(BaseElementType, "Type*", ElementMetaData, DocumentData.CurrentSettings);
 				}
 
 				return ElementMetaData;
@@ -1077,6 +1077,8 @@ namespace DatasmithRevitExporter
 		private string									CurrentMaterialName = null;
 		private List<string>							MessageList = null;
 
+		private FSettings								CurrentSettings = null;
+
 		// Apply world offset to elements
 		public	FSettings.EInsertionPoint				InsertionPoint { get; set; } = FSettings.EInsertionPoint.Default;
 
@@ -1093,13 +1095,15 @@ namespace DatasmithRevitExporter
 
 		public FDocumentData(
 			Document InDocument,
+			FSettings InSettings,
 			ref List<string> InMessageList,
 			FDirectLink InDirectLink,
 			string InLinkedDocumentId
 		)
 		{
-			InsertionPoint = FSettingsManager.CurrentSettings?.InsertionPoint ?? FSettings.EInsertionPoint.Default;
+			InsertionPoint = InSettings?.InsertionPoint ?? FSettings.EInsertionPoint.Default;
 
+			CurrentSettings = InSettings;
 			DirectLink = InDirectLink;
 			CurrentDocument = InDocument;
 			MessageList = InMessageList;
@@ -1910,6 +1914,7 @@ namespace DatasmithRevitExporter
 					// Don't apply offset since basepoints aren't yet initialized
 					SetActorTransform(TranslationMatrix.Multiply(InWorldTransform), BasePointActor, false);
 
+
 					// Set the Datasmith placeholder actor layer to the base point category name.
 					BasePointActor.SetLayer(BasePointLocation.Category.Name);
 
@@ -1931,7 +1936,7 @@ namespace DatasmithRevitExporter
 
 					if (!bSkipMetadataExport)
 					{
-						FUtils.AddActorMetadata(BasePointLocation, MetadataPrefix, BasePointMetaData);
+						FUtils.AddActorMetadata(BasePointLocation, MetadataPrefix, BasePointMetaData, CurrentSettings);
 					}
 
 					if (BasePointElement == null)
@@ -2310,7 +2315,8 @@ namespace DatasmithRevitExporter
 		private void SetActorTransform(
 			Transform InWorldTransform,
 			FDatasmithFacadeActor IOActor,
-			bool bInApplyOffset = true)
+			bool bInApplyOffset = true
+		)
 		{
 			XYZ transformBasisX = InWorldTransform.BasisX;
 			XYZ transformBasisY = InWorldTransform.BasisY;
@@ -2322,21 +2328,23 @@ namespace DatasmithRevitExporter
 			{
 				switch (InsertionPoint)
 				{
-					case FSettings.EInsertionPoint.BasePoint: 
+					case FSettings.EInsertionPoint.BasePoint:
 					{
 						if (ProjectBasePoint != null)
 						{
 							transformOrigin -= ProjectBasePoint;
 						}
-					} break;
+					}
+					break;
 
-					case FSettings.EInsertionPoint.SurveyPoint: 
+					case FSettings.EInsertionPoint.SurveyPoint:
 					{
 						if (ProjectSurveyPoint != null)
 						{
 							transformOrigin -= ProjectSurveyPoint;
 						}
-					} break;
+					}
+					break;
 				}
 			}
 
diff --git a/Engine/Source/Programs/Enterprise/Datasmith/DatasmithRevitExporter/Private/DatasmithRevitExportContext.cs b/Engine/Source/Programs/Enterprise/Datasmith/DatasmithRevitExporter/Private/DatasmithRevitExportContext.cs
--- a/Engine/Source/Programs/Enterprise/Datasmith/DatasmithRevitExporter/Private/DatasmithRevitExportContext.cs
+++ b/Engine/Source/Programs/Enterprise/Datasmith/DatasmithRevitExporter/Private/DatasmithRevitExportContext.cs
@@ -64,6 +64,8 @@ namespace DatasmithRevitExporter
 
 		private FDirectLink DirectLink;
 
+		private FSettings DocumentSettings;
+
 		private FDatasmithRevitExportContext() {}
 
 		private bool CurrentElementSkipped = false;
@@ -85,19 +87,21 @@ namespace DatasmithRevitExporter
 		public FDatasmithRevitExportContext(
 			Application						InApplication,        // running Revit application
 			Document						InDocument,           // active Revit document
+			FSettings						InSettings,
 			Dictionary<ElementId, string>	InDatasmithFilePaths, // Datasmith output file path
-			DatasmithRevitExportOptions		InExportOptions,      // Unreal Datasmith export options
-			FDirectLink						InDirectLink		  // DirectLink manager
+			DatasmithRevitExportOptions		InExportOptions,       // Unreal Datasmith export options
+			FDirectLink						InDirectLink
 		)
 		{
 			ProductVersion = InApplication.VersionNumber;
+			DocumentSettings = InSettings;
 			DatasmithFilePaths = InDatasmithFilePaths;
 			RevitDocument = InDocument;
 			DirectLink = InDirectLink;
 
 			// Get the Unreal Datasmith export options.
 			DebugLog = InExportOptions.GetWriteLogFile() ? new FDatasmithFacadeLog() : null;
-			LevelOfTessellation = FSettingsManager.CurrentSettings?.LevelOfTesselation ?? 8;
+			LevelOfTessellation = InSettings.LevelOfTesselation;
 		}
 
 		//========================================================================================================================
@@ -605,7 +609,7 @@ namespace DatasmithRevitExporter
 		}
 
 		private void PushDocument(
-			Document InDocument,
+			Document InRevitDocument,
 			string InLinkedDocumentId
 		)
 		{
@@ -617,7 +621,7 @@ namespace DatasmithRevitExporter
 			}
 
 			// Check if we have cache for this document.
-			FDocumentData DocumentData = new FDocumentData(InDocument, ref MessageList, DirectLink, InLinkedDocumentId);
+			FDocumentData DocumentData = new FDocumentData(InRevitDocument, DocumentSettings, ref MessageList, DirectLink, InLinkedDocumentId);
 
 			DocumentDataStack.Push(DocumentData);
 
diff --git a/Engine/Source/Programs/Enterprise/Datasmith/DatasmithRevitExporter/Private/DatasmithRevitSettingsDialog.cs b/Engine/Source/Programs/Enterprise/Datasmith/DatasmithRevitExporter/Private/DatasmithRevitSettingsDialog.cs
--- a/Engine/Source/Programs/Enterprise/Datasmith/DatasmithRevitExporter/Private/DatasmithRevitSettingsDialog.cs
+++ b/Engine/Source/Programs/Enterprise/Datasmith/DatasmithRevitExporter/Private/DatasmithRevitSettingsDialog.cs
@@ -95,7 +95,7 @@ namespace DatasmithRevitExporter
 		private List<string> AddedParamNames = new List<string>();
 		private bool bSettingsChanged = false;
 
-		public DatasmithRevitSettingsDialog(Autodesk.Revit.DB.Document InDocument)
+		public DatasmithRevitSettingsDialog(Autodesk.Revit.DB.Document InDocument, FSettings InSettings)
 		{
 			string FormatTooltip(string InText)
 			{
@@ -106,7 +106,7 @@ namespace DatasmithRevitExporter
 			Document = InDocument;
 			FormClosing += OnClosing;
 
-			Settings = FSettingsManager.CurrentSettings;
+			Settings = InSettings;
 
 			ToolTip OptionToolTip = new ToolTip();
 			OptionToolTip.AutoPopDelay = 10000; // milliseconds
@@ -174,10 +174,6 @@ namespace DatasmithRevitExporter
 				LevelOfTessellation.AutoSize = true;
 				LevelOfTessellation.TabIndex = 3;
 				LevelOfTessellation.TextAlign = HorizontalAlignment.Right;
-				LevelOfTessellation.ValueChanged += (S, E) =>
-				{
-					bSettingsChanged = true;
-				};
 
 				OptionToolTip.SetToolTip(LevelOfTessellationLabel, FormatTooltip(DatasmithRevitResources.Strings.SettingsDialog_LevelOfTesselationTooltip));
 
@@ -333,11 +329,7 @@ namespace DatasmithRevitExporter
 
 			Controls.Add(TopLayout);
 
-			if (Settings == null)
-			{
-				Settings = new FSettings();
-			}
-			else
+			if (Settings != null)
 			{
 				InsertionPointCombo.SelectedIndex = (int)Settings.InsertionPoint;
 				LevelOfTessellation.Value = Settings.LevelOfTesselation;
@@ -350,6 +342,11 @@ namespace DatasmithRevitExporter
 
 				ReloadAddedGroupsList();
 			}
+
+			LevelOfTessellation.ValueChanged += (S, E) =>
+			{
+				bSettingsChanged = true;
+			};
 		}
 
 		private void ReloadAddedGroupsList()
diff --git a/Engine/Source/Programs/Enterprise/Datasmith/DatasmithRevitExporter/Private/DatasmithRevitSettingsManager.cs b/Engine/Source/Programs/Enterprise/Datasmith/DatasmithRevitExporter/Private/DatasmithRevitSettingsManager.cs
--- a/Engine/Source/Programs/Enterprise/Datasmith/DatasmithRevitExporter/Private/DatasmithRevitSettingsManager.cs
+++ b/Engine/Source/Programs/Enterprise/Datasmith/DatasmithRevitExporter/Private/DatasmithRevitSettingsManager.cs
@@ -3,9 +3,7 @@
 using System;
 using System.Collections.Generic;
 using Autodesk.Revit.DB;
-using Autodesk.Revit.DB.Events;
 using Autodesk.Revit.DB.ExtensibleStorage;
-using Autodesk.Revit.UI;
 
 namespace DatasmithRevitExporter
 {
@@ -27,6 +25,7 @@ namespace DatasmithRevitExporter
 
 		public int LevelOfTesselation { get; set; } = 8;
 		public EInsertionPoint InsertionPoint { get; set; } = EInsertionPoint.Default;
+		public ElementId SyncViewId { get; set; } = ElementId.InvalidElementId;
 
 		public IList<string> MetadataParamNamesFilter 
 		{
@@ -85,12 +84,13 @@ namespace DatasmithRevitExporter
 	{
 		static class FSettingsSchema
 		{
-			readonly static Guid SchemaGuid = new Guid("{F910B32B-32C6-4BD1-BFD1-7DBF81FFC41B}");
+			readonly static Guid SchemaGuid = new Guid("{67D3D706-4881-4BC3-9AA1-75768DE40415}");
 			
 			public readonly static string MetadataParamNamesField = "MetadataParamNames";
 			public readonly static string MetadataParamGroupsField = "MetadataParamGroups";
 			public readonly static string LevelOfTesselationField = "LevelOfTesselation";
 			public readonly static string InsertionPointField = "InsertionPoint";
+			public readonly static string SyncViewField = "SyncView";
 
 			public static Schema GetSchema()
 			{
@@ -107,6 +107,7 @@ namespace DatasmithRevitExporter
 				SchemaBuilder.AddArrayField(MetadataParamGroupsField, typeof(int));
 				SchemaBuilder.AddSimpleField(LevelOfTesselationField, typeof(int));
 				SchemaBuilder.AddSimpleField(InsertionPointField, typeof(int));
+				SchemaBuilder.AddSimpleField(SyncViewField, typeof(int));
 
 				return SchemaBuilder.Finish();
 			}
@@ -114,7 +115,7 @@ namespace DatasmithRevitExporter
 
 		static class DataStorageUniqueIdSchema
 		{
-			static readonly Guid SchemaGuid = new Guid("{2E363F82-4589-43BE-8927-C596AA02396F}");
+			static readonly Guid SchemaGuid = new Guid("{B03CF2DB-3617-4C67-90D4-A34C734A54E4}");
 
 			public static Schema GetSchema()
 			{
@@ -132,31 +133,11 @@ namespace DatasmithRevitExporter
 			}
 		}
 
-		static EventHandler<DocumentOpenedEventArgs> DocumentOpenedHandler;
-		static EventHandler<DocumentClosingEventArgs> DocumentClosingHandler;
-		static readonly Guid SettingDataStorageId = new Guid("{6CD2BB42-4A04-4F12-8E7F-5C6723C33791}");
-
-		public static FSettings CurrentSettings = null;
+		static readonly Guid SettingDataStorageId = new Guid("{5933427B-3E44-44AC-85CF-74C396662E05}");
 
 		public static event EventHandler SettingsUpdated;
 
-		public static void Init(UIControlledApplication InApplication)
-		{
-			DocumentOpenedHandler = new EventHandler<DocumentOpenedEventArgs>(OnDocumentOpened);
-			DocumentClosingHandler = new EventHandler<DocumentClosingEventArgs>(OnDocumentClosing);
-			InApplication.ControlledApplication.DocumentOpened += DocumentOpenedHandler;
-			InApplication.ControlledApplication.DocumentClosing += DocumentClosingHandler;
-		}
-
-		public static void Destroy(UIControlledApplication InApplication)
-		{
-			InApplication.ControlledApplication.DocumentClosing -= DocumentClosingHandler;
-			InApplication.ControlledApplication.DocumentOpened -= DocumentOpenedHandler;
-			DocumentClosingHandler = null;
-			DocumentOpenedHandler = null;
-		}
-
-		private static FSettings ReadSettings(Document Doc)
+		public static FSettings ReadSettings(Document Doc)
 		{
 			FSettings Settings = new FSettings();
 
@@ -165,6 +146,7 @@ namespace DatasmithRevitExporter
 				// No settings in document, create defaults
 				Settings.LevelOfTesselation = 8;
 				Settings.InsertionPoint = 0;
+				Settings.SyncViewId = ElementId.InvalidElementId;
 				Settings.MetadataParamGroupsFilter.Add((int)Autodesk.Revit.DB.BuiltInParameterGroup.PG_GEOMETRY);
 				Settings.MetadataParamGroupsFilter.Add((int)Autodesk.Revit.DB.BuiltInParameterGroup.PG_IDENTITY_DATA);
 				Settings.MetadataParamGroupsFilter.Add((int)Autodesk.Revit.DB.BuiltInParameterGroup.PG_MATERIALS);
@@ -188,6 +170,7 @@ namespace DatasmithRevitExporter
 					Settings.MetadataParamGroupsFilter = SettingsEntity.Get<IList<int>>(FSettingsSchema.MetadataParamGroupsField);
 					Settings.LevelOfTesselation = SettingsEntity.Get<int>(FSettingsSchema.LevelOfTesselationField);
 					Settings.InsertionPoint = (FSettings.EInsertionPoint)SettingsEntity.Get<int>(FSettingsSchema.InsertionPointField);
+					Settings.SyncViewId = new ElementId(SettingsEntity.Get<int>(FSettingsSchema.SyncViewField));
 				}
 				catch
 				{
@@ -217,6 +200,7 @@ namespace DatasmithRevitExporter
 					SettingsEntity.Set(FSettingsSchema.MetadataParamGroupsField, Settings.MetadataParamGroupsFilter);
 					SettingsEntity.Set(FSettingsSchema.LevelOfTesselationField, Settings.LevelOfTesselation);
 					SettingsEntity.Set(FSettingsSchema.InsertionPointField, (int)Settings.InsertionPoint);
+					SettingsEntity.Set(FSettingsSchema.SyncViewField, Settings.SyncViewId.IntegerValue);
 
 					// Identify settings data storage
 					Entity IdEntity = new Entity(DataStorageUniqueIdSchema.GetSchema());
@@ -225,26 +209,14 @@ namespace DatasmithRevitExporter
 					SettingDataStorage.SetEntity(IdEntity);
 					SettingDataStorage.SetEntity(SettingsEntity);
 
-					CurrentSettings = Settings;
-
 					T.Commit();
 
-					SettingsUpdated?.Invoke(CurrentSettings, null);
+					SettingsUpdated?.Invoke(Settings, null);
 				}
 				catch { }
 			}
 		}
 
-		static void OnDocumentOpened(object sender, DocumentOpenedEventArgs e)
-		{
-			CurrentSettings = ReadSettings(e.Document);
-		}
-
-		static void OnDocumentClosing(object sender, DocumentClosingEventArgs e)
-		{
-			CurrentSettings = null;
-		}
-
 		private static DataStorage GetSettingsDataStorage(Document Doc)
 		{
 			// Retrieve all data storages from project
diff --git a/Engine/Source/Programs/Enterprise/Datasmith/DatasmithRevitExporter/Private/DatasmithRevitUtils.cs b/Engine/Source/Programs/Enterprise/Datasmith/DatasmithRevitExporter/Private/DatasmithRevitUtils.cs
--- a/Engine/Source/Programs/Enterprise/Datasmith/DatasmithRevitExporter/Private/DatasmithRevitUtils.cs
+++ b/Engine/Source/Programs/Enterprise/Datasmith/DatasmithRevitExporter/Private/DatasmithRevitUtils.cs
@@ -18,7 +18,7 @@ namespace DatasmithRevitExporter
 			return InElement.Document.GetElement(InElement.GetTypeId()) as ElementType;
 		}
 
-		public static void AddActorMetadata(Element InElement, FDatasmithFacadeMetaData ActorMetadata)
+		public static void AddActorMetadata(Element InElement, FDatasmithFacadeMetaData ActorMetadata, FSettings InSettings)
 		{
 			// Add the Revit element category name metadata to the Datasmith actor.
 			string CategoryName = GetCategoryName(InElement);
@@ -43,23 +43,22 @@ namespace DatasmithRevitExporter
 			}
 
 			// Add Revit element metadata to the Datasmith actor.
-			AddActorMetadata(InElement, "Element*", ActorMetadata);
+			AddActorMetadata(InElement, "Element*", ActorMetadata, InSettings);
 
 			if (ElemType != null)
 			{
 				// Add Revit element type metadata to the Datasmith actor.
-				AddActorMetadata(ElemType, "Type*", ActorMetadata);
+				AddActorMetadata(ElemType, "Type*", ActorMetadata, InSettings);
 			}
 		}
 
 		public static void AddActorMetadata(
 			Element InSourceElement,
 			string InMetadataPrefix,
-			FDatasmithFacadeMetaData ElementMetaData
+			FDatasmithFacadeMetaData ElementMetaData,
+			FSettings InSettings
 		)
 		{
-			FSettings Settings = FSettingsManager.CurrentSettings;
-
 			IList<Parameter> Parameters = InSourceElement.GetOrderedParameters();
 
 			if (Parameters != null)
@@ -70,7 +69,7 @@ namespace DatasmithRevitExporter
 					{
 						if (Parameter.HasValue)
 						{
-							if (Settings != null && !Settings.MatchParameterByMetadata(Parameter))
+							if (InSettings != null && !InSettings.MatchParameterByMetadata(Parameter))
 							{
 								continue; // Skip export of this param
 							}
