[3ds Max] DirectLink: Limit Texture Resolution option

- baked texmap size limited to specified resolution

diff --git a/Engine/Source/Programs/Enterprise/Datasmith/DatasmithMaxExporter/Extras/3dsMaxScripts/UI/DatasmithUI.xaml b/Engine/Source/Programs/Enterprise/Datasmith/DatasmithMaxExporter/Extras/3dsMaxScripts/UI/DatasmithUI.xaml
--- a/Engine/Source/Programs/Enterprise/Datasmith/DatasmithMaxExporter/Extras/3dsMaxScripts/UI/DatasmithUI.xaml
+++ b/Engine/Source/Programs/Enterprise/Datasmith/DatasmithMaxExporter/Extras/3dsMaxScripts/UI/DatasmithUI.xaml
@@ -131,6 +131,7 @@
             IsSplit="false"
             IsSynchronizedWithCurrentItem="true"
             IsToolTipEnabled="true"
+            LargeImage="pack://application:,,,/AdWindows;component/Resources/Images/MissingImage.ico"
             ListImageSize="Standard"
             MaxscriptIndexCallback="Datasmith_ExportOptions_IncludeIndex()"
             MinWidth="0"
@@ -144,6 +145,7 @@
               Image="pack://application:,,,/AdWindows;component/Resources/Images/MissingImage.ico"
               IsCheckable="true"
               IsToolTipEnabled="true"
+              LargeImage="pack://application:,,,/AdWindows;component/Resources/Images/MissingImage.ico"
               MinWidth="100"
               Name="Visible Objects"
               ShowImage="false"
@@ -231,6 +233,156 @@
         </ads:RibbonRowPanel>
       </wpfcust:MaxRibbonPanelSource>
     </wpfcust:MaxModelingRibbonPanel>
+    <wpfcust:MaxModelingRibbonPanel
+      CollapsedPanelImage="Icons/MissingImage_32.png"
+      FloatingGroup="0"
+      FloatingOrder="0"
+      UID="42ca9067-e3a4-4d3c-8d83-af9bb10b1124">
+      <wpfcust:MaxRibbonPanelSource
+        Title="Settings"
+        UID="cc748308-a562-4516-8c94-4b2dae7a3a8b">
+        <ads:RibbonRowPanel
+          AllowInToolBar="false"
+          Image="Icons/MissingImage_16.png"
+          IsToolTipEnabled="true"
+          LargeImage="Icons/MissingImage_32.png"
+          MinWidth="100"
+          ShowText="false"
+          UID="939d8499-112d-4b8a-8b95-375ca7f8e029">
+          <wpfcust:MaxRibbonLabel
+            AllowInToolBar="false"
+            IsToolTipEnabled="false"
+            MinWidth="0"
+            ResizeStyle="NoResize"
+            ShowText="true"
+            Text="Limit Texture Resolution"
+            UID="4519fb6c-d5d3-45ee-908a-6464f8639ba1" />
+          <ads:RibbonSeparator
+            AllowInToolBar="false"
+            Image="Icons/MissingImage_16.png"
+            IsToolTipEnabled="true"
+            LargeImage="Icons/MissingImage_32.png"
+            MinWidth="0"
+            ResizeStyle="NoResize"
+            ShowText="false"
+            UID="1c3c4b93-a99e-4246-ac1f-7173c87aec67" />
+          <wpfcust:MaxRibbonSplitButton
+            Image="pack://application:,,,/AdWindows;component/Resources/Images/MissingImage.ico"
+            IsCheckable="false"
+            IsSplit="false"
+            IsSynchronizedWithCurrentItem="true"
+            IsToolTipEnabled="true"
+            LargeImage="pack://application:,,,/AdWindows;component/Resources/Images/MissingImage.ico"
+            ListImageSize="Standard"
+            MaxscriptIndexCallback="Datasmith_ExportOptions_TextureResolutionIndex()"
+            MinWidth="0"
+            ResizeStyle="HideText"
+            ShowImage="false"
+            ShowText="true"
+            UID="a721d8e7-992f-464a-b2b4-c080a4f39fd9">
+            <ads:RibbonCommandItem
+              AllowInToolBar="false"
+              CommandHandler="{wpfcust:ActionItem MacroName=Datasmith_ExportOptions_TextureResolution64, MacroCategory=Datasmith}"
+              Image="pack://application:,,,/AdWindows;component/Resources/Images/MissingImage.ico"
+              IsCheckable="true"
+              IsToolTipEnabled="true"
+              LargeImage="pack://application:,,,/AdWindows;component/Resources/Images/MissingImage.ico"
+              MinWidth="100"
+              Name="64"
+              ShowImage="false"
+              ShowText="false"
+              Text="64"
+              UID="28d398da-505c-4f91-9cf3-14316cd2893b" />
+            <ads:RibbonCommandItem
+              AllowInToolBar="false"
+              CommandHandler="{wpfcust:ActionItem MacroName=Datasmith_ExportOptions_TextureResolution128, MacroCategory=Datasmith}"
+              Image="pack://application:,,,/AdWindows;component/Resources/Images/MissingImage.ico"
+              IsCheckable="true"
+              IsToolTipEnabled="true"
+              LargeImage="pack://application:,,,/AdWindows;component/Resources/Images/MissingImage.ico"
+              MinWidth="100"
+              Name="128"
+              ShowImage="false"
+              ShowText="false"
+              Text="128"
+              UID="89a922df-7286-40b9-833e-c51419fe3c1d" />
+            <ads:RibbonCommandItem
+              AllowInToolBar="false"
+              CommandHandler="{wpfcust:ActionItem MacroName=Datasmith_ExportOptions_TextureResolution256, MacroCategory=Datasmith}"
+              Image="pack://application:,,,/AdWindows;component/Resources/Images/MissingImage.ico"
+              IsCheckable="true"
+              IsToolTipEnabled="true"
+              LargeImage="pack://application:,,,/AdWindows;component/Resources/Images/MissingImage.ico"
+              MinWidth="100"
+              Name="256"
+              ShowImage="false"
+              ShowText="false"
+              Text="256"
+              UID="177f6116-6ee0-4b6f-8d06-72d73bcc89e4" />
+            <ads:RibbonCommandItem
+              AllowInToolBar="false"
+              CommandHandler="{wpfcust:ActionItem MacroName=Datasmith_ExportOptions_TextureResolution512, MacroCategory=Datasmith}"
+              Image="pack://application:,,,/AdWindows;component/Resources/Images/MissingImage.ico"
+              IsCheckable="true"
+              IsToolTipEnabled="true"
+              LargeImage="pack://application:,,,/AdWindows;component/Resources/Images/MissingImage.ico"
+              MinWidth="100"
+              Name="512"
+              ShowImage="false"
+              ShowText="false"
+              Text="512"
+              UID="a5dd7f25-278a-4f8e-a9c2-7d73164df5f0" />
+            <ads:RibbonCommandItem
+              AllowInToolBar="false"
+              CommandHandler="{wpfcust:ActionItem MacroName=Datasmith_ExportOptions_TextureResolution1024, MacroCategory=Datasmith}"
+              Image="pack://application:,,,/AdWindows;component/Resources/Images/MissingImage.ico"
+              IsCheckable="true"
+              IsToolTipEnabled="true"
+              LargeImage="pack://application:,,,/AdWindows;component/Resources/Images/MissingImage.ico"
+              MinWidth="100"
+              Name="1024"
+              ShowImage="false"
+              ShowText="false"
+              Text="1024"
+              UID="09fec31c-f9e9-4e55-b33c-f4fb2218a4a0" />
+            <ads:RibbonCommandItem
+              AllowInToolBar="false"
+              CommandHandler="{wpfcust:ActionItem MacroName=Datasmith_ExportOptions_TextureResolution2048, MacroCategory=Datasmith}"
+              Image="pack://application:,,,/AdWindows;component/Resources/Images/MissingImage.ico"
+              IsCheckable="true"
+              IsToolTipEnabled="true"
+              LargeImage="pack://application:,,,/AdWindows;component/Resources/Images/MissingImage.ico"
+              MinWidth="100"
+              Name="2048"
+              ShowImage="false"
+              ShowText="false"
+              Text="2048"
+              UID="ef8585b8-1012-47a7-9af0-46371cf70df2" />
+            <ads:RibbonCommandItem
+              AllowInToolBar="false"
+              CommandHandler="{wpfcust:ActionItem MacroName=Datasmith_ExportOptions_TextureResolution4096, MacroCategory=Datasmith}"
+              Image="pack://application:,,,/AdWindows;component/Resources/Images/MissingImage.ico"
+              IsCheckable="true"
+              IsToolTipEnabled="true"
+              LargeImage="pack://application:,,,/AdWindows;component/Resources/Images/MissingImage.ico"
+              MinWidth="100"
+              Name="4096"
+              ShowImage="false"
+              ShowText="false"
+              Text="4096"
+              UID="c5c4c61f-7d64-4267-a387-cd6df2704161" />
+          </wpfcust:MaxRibbonSplitButton>
+          <ads:RibbonRowBreak
+            AllowInToolBar="false"
+            Image="Icons/MissingImage_16.png"
+            IsToolTipEnabled="true"
+            LargeImage="Icons/MissingImage_32.png"
+            MinWidth="100"
+            ShowText="false"
+            UID="189ae82b-5e0f-4d3b-aa93-81e928234b27" />
+        </ads:RibbonRowPanel>
+      </wpfcust:MaxRibbonPanelSource>
+    </wpfcust:MaxModelingRibbonPanel>
     <wpfcust:MaxModelingRibbonPanel
       CollapsedPanelImage="Icons/Datasmith/DatasmithSettingsIcon32.png"
       FloatingGroup="0"
diff --git a/Engine/Source/Programs/Enterprise/Datasmith/DatasmithMaxExporter/Private/DatasmithMaxDirectLinkMaxScript.cpp b/Engine/Source/Programs/Enterprise/Datasmith/DatasmithMaxExporter/Private/DatasmithMaxDirectLinkMaxScript.cpp
--- a/Engine/Source/Programs/Enterprise/Datasmith/DatasmithMaxExporter/Private/DatasmithMaxDirectLinkMaxScript.cpp
+++ b/Engine/Source/Programs/Enterprise/Datasmith/DatasmithMaxExporter/Private/DatasmithMaxDirectLinkMaxScript.cpp
@@ -293,14 +293,36 @@ Value* SetExportOption_##name##_cf(Value** arg_list, int count) \
 } \
 Primitive SetExportOption_##name##_pf(_M("Datasmith_SetExportOption_" #name), SetExportOption_##name##_cf);\
 
+#define DefinePersistentExportOptionInt(name) \
+Value* GetExportOption_##name##_cf(Value** arg_list, int count) \
+{ \
+	check_arg_count(GetExportOption_##name, 0, count); \
+ \
+	one_value_local(result); \
+	vl.result = Integer::intern(GetPersistentExportOptions().Get##name##()); \
+	return_value(vl.result); \
+} \
+Primitive GetExportOption_##name##_pf(_M("Datasmith_GetExportOption_" #name), GetExportOption_##name##_cf);\
+\
+Value* SetExportOption_##name##_cf(Value** arg_list, int count) \
+{ \
+	check_arg_count(SetExportOption_##name, 1, count); \
+	Value* pValue= arg_list[0]; \
+	GetPersistentExportOptions().Set##name##(pValue->to_int()); \
+	return &true_value; \
+} \
+Primitive SetExportOption_##name##_pf(_M("Datasmith_SetExportOption_" #name), SetExportOption_##name##_cf);\
+
 //////////////////////////////////////////
 DefinePersistentExportOption(SelectedOnly)
 DefinePersistentExportOption(AnimatedTransforms)
 
 DefinePersistentExportOption(StatSync)
+DefinePersistentExportOptionInt(TextureResolution)
 //////////////////////////////////////////
 
 #undef DefinePersistentExportOption
+#undef DefinePersistentExportOptionInt
 
 Value* GetDirectlinkCacheDirectory_cf(Value** arg_list, int count)
 {
diff --git a/Engine/Source/Programs/Enterprise/Datasmith/DatasmithMaxExporter/Private/DatasmithMaxExporterWithDirectLink.cpp b/Engine/Source/Programs/Enterprise/Datasmith/DatasmithMaxExporter/Private/DatasmithMaxExporterWithDirectLink.cpp
--- a/Engine/Source/Programs/Enterprise/Datasmith/DatasmithMaxExporter/Private/DatasmithMaxExporterWithDirectLink.cpp
+++ b/Engine/Source/Programs/Enterprise/Datasmith/DatasmithMaxExporter/Private/DatasmithMaxExporterWithDirectLink.cpp
@@ -548,11 +548,13 @@ public:
 
 struct FExportOptions
 {
-	// Default options for DirectLink change-tracking
 	bool bSelectedOnly = false;
 	bool bAnimatedTransforms = false;
 
 	bool bStatSync = false;
+	int32 TextureResolution = 4;
+
+	static const int32 TextureResolutionMax = 6; // 
 };
 
 // Global export options, stored in preferences
@@ -567,48 +569,49 @@ public:
 		}
 		GetBool(TEXT("SelectedOnly"), Options.bSelectedOnly);
 		GetBool(TEXT("AnimatedTransforms"), Options.bAnimatedTransforms);
+		GetInt(TEXT("TextureResolution"), Options.TextureResolution);
 		bLoaded = true;
 	}
 
-	virtual void SetSelectedOnly(bool bValue) override
-	{
-		Options.bSelectedOnly = bValue;
-		SetBool(TEXT("SelectedOnly"), bValue);
-	}
-	virtual bool GetSelectedOnly() override
-	{
-		return Options.bSelectedOnly;
-	}
-
-	virtual void SetAnimatedTransforms(bool bValue) override
+	void GetBool(const TCHAR* Name, bool& bValue)
 	{
-		Options.bAnimatedTransforms = bValue;
-		SetBool(TEXT("AnimatedTransforms"), bValue);
+		if (!GConfig)
+		{
+			return;
+		}
+		FString ConfigPath = GetConfigPath();
+		GConfig->GetBool(TEXT("Export"), Name, bValue, ConfigPath);
 	}
 
-	virtual bool GetAnimatedTransforms() override
+	void SetBool(const TCHAR* Name, bool bValue)
 	{
-		return Options.bAnimatedTransforms;
+		if (!GConfig)
+		{
+			return;
+		}
+		FString ConfigPath = GetConfigPath();
+		GConfig->SetBool(TEXT("Export"), Name, bValue, ConfigPath);
+		GConfig->Flush(false, ConfigPath);
 	}
 
-	void GetBool(const TCHAR* Name, bool& bValue)
+	void GetInt(const TCHAR* Name, int32& bValue)
 	{
 		if (!GConfig)
 		{
 			return;
 		}
 		FString ConfigPath = GetConfigPath();
-		GConfig->GetBool(TEXT("Export"), Name, bValue, ConfigPath);
+		GConfig->GetInt(TEXT("Export"), Name, bValue, ConfigPath);
 	}
 
-	void SetBool(const TCHAR* Name, bool bValue)
+	void SetInt(const TCHAR* Name, int32 bValue)
 	{
 		if (!GConfig)
 		{
 			return;
 		}
 		FString ConfigPath = GetConfigPath();
-		GConfig->SetBool(TEXT("Export"), Name, bValue, ConfigPath);
+		GConfig->SetInt(TEXT("Export"), Name, bValue, ConfigPath);
 		GConfig->Flush(false, ConfigPath);
 	}
 
@@ -618,6 +621,27 @@ public:
 		return FPaths::Combine(PlugCfgPath, TEXT("UnrealDatasmithMax.ini"));
 	}
 
+	virtual void SetSelectedOnly(bool bValue) override
+	{
+		Options.bSelectedOnly = bValue;
+		SetBool(TEXT("SelectedOnly"), bValue);
+	}
+	virtual bool GetSelectedOnly() override
+	{
+		return Options.bSelectedOnly;
+	}
+
+	virtual void SetAnimatedTransforms(bool bValue) override
+	{
+		Options.bAnimatedTransforms = bValue;
+		SetBool(TEXT("AnimatedTransforms"), bValue);
+	}
+
+	virtual bool GetAnimatedTransforms() override
+	{
+		return Options.bAnimatedTransforms;
+	}
+
 	virtual void SetStatSync(bool bValue) override
 	{
 		Options.bStatSync = bValue;
@@ -629,6 +653,17 @@ public:
 		return Options.bStatSync;
 	}
 
+	virtual void SetTextureResolution(int32 Value) override
+	{
+		
+		Options.TextureResolution = FMath::Clamp(Value, 0, FExportOptions::TextureResolutionMax);
+		SetInt(TEXT("TextureResolution"), Value);
+	}
+
+	virtual int32 GetTextureResolution() override
+	{
+		return Options.TextureResolution;
+	}
 
 	FExportOptions Options;
 	bool bLoaded = false;
@@ -1167,6 +1202,7 @@ public:
 
 		CurrentSyncPoint.Time = GetCOREInterface()->GetTime();
 
+		// Changes present only when there are modified layers(changes checked manually), nodes(notified by Max) or materials(notified by Max with all changes in dependencies)
 		bool bChangeEncountered = false;
 
 		Stats.Reset();
@@ -1182,8 +1218,6 @@ public:
 			return false;	
 		}
 
-		// Changes present only when there are modified layers(changes checked manually), nodes(notified by Max) or materials(notified by Max with all changes in dependencies)
-		bChangeEncountered |= !MaterialsCollectionTracker.GetInvalidatedMaterials().IsEmpty();
 
 		{
 			PROGRESS_STAGE("Remove deleted nodes")
@@ -1374,6 +1408,21 @@ public:
 			InvalidatedNodeTrackers.Finish();
 		}
 
+		// 0 -> 64; 4 -> 1024; 6 -> 4096
+		int32 MaxBakedTextureSize = 1 << (Options.TextureResolution + 6);
+
+		// When requested max texture size changed  invalidate all materials(material update emits textures to bake)
+		if (MaxBakedTextureSize != FDatasmithExportOptions::MaxTextureSize)
+		{
+			for (TTuple<MtlBase*, FMaterialTrackerHandle> Material : MaterialsCollectionTracker.MaterialTrackers)
+			{
+				MaterialsCollectionTracker.InvalidateMaterial(Material.Value.GetMaterialTracker()->Material);
+			}
+			FDatasmithExportOptions::MaxTextureSize = MaxBakedTextureSize;
+		}
+
+		bChangeEncountered |= !MaterialsCollectionTracker.GetInvalidatedMaterials().IsEmpty();
+
 		// Each tracked(i.e. assigned to a visible node) Max material can result in multiple actual materials
 		// e.g. an assigned MultiSubObj material is aclually a set of material not directly assigned to a node
 		TSet<Mtl*> ActualMaterialToUpdate; 
diff --git a/Engine/Source/Programs/Enterprise/Datasmith/DatasmithMaxExporter/Public/DatasmithMaxDirectLink.h b/Engine/Source/Programs/Enterprise/Datasmith/DatasmithMaxExporter/Public/DatasmithMaxDirectLink.h
--- a/Engine/Source/Programs/Enterprise/Datasmith/DatasmithMaxExporter/Public/DatasmithMaxDirectLink.h
+++ b/Engine/Source/Programs/Enterprise/Datasmith/DatasmithMaxExporter/Public/DatasmithMaxDirectLink.h
@@ -59,6 +59,9 @@ public:
 	// Whether to output export statistics to listener/log
 	virtual void SetStatSync(bool) = 0;
 	virtual bool GetStatSync() = 0;
+
+	virtual void SetTextureResolution(int32 Bool) = 0;
+	virtual int32 GetTextureResolution() = 0;
 };
 
 //---- Main class for export/change tracking
