[3ds Max] DirectLink: materials notifications rework

- main improvement is to handle situations when material change event is sent during update - vray materials also fire events when RenderBegin/RenderEnd is called on their nodes

diff --git a/Engine/Source/Programs/Enterprise/Datasmith/DatasmithMaxExporter/Private/DatasmithMaxDirectLinkNotifications.cpp b/Engine/Source/Programs/Enterprise/Datasmith/DatasmithMaxExporter/Private/DatasmithMaxDirectLinkNotifications.cpp
--- a/Engine/Source/Programs/Enterprise/Datasmith/DatasmithMaxExporter/Private/DatasmithMaxDirectLinkNotifications.cpp
+++ b/Engine/Source/Programs/Enterprise/Datasmith/DatasmithMaxExporter/Private/DatasmithMaxDirectLinkNotifications.cpp
@@ -32,16 +32,6 @@ public:
 
 	virtual BOOL VerboseDeleted() { return TRUE; }
 
-	// Tracks node's material parameter change(even if it's a submaterial of multimat that is assigned)
-	virtual void MaterialOtherEvent(NodeKeyTab& nodes) override
-	{
-		LogNodeEvent(L"MaterialOtherEvent", nodes);
-		for (int NodeIndex = 0; NodeIndex < nodes.Count(); ++NodeIndex)
-		{
-			SceneTracker.NodeMaterialGraphModified(nodes[NodeIndex]);
-		}
-	}
-
 	virtual void NameChanged(NodeKeyTab& nodes) override
 	{
 		LogNodeEvent(L"NameChanged", nodes);
@@ -111,6 +101,12 @@ public:
 		LogNodeEvent(L"MaterialStructured", nodes);
 	}
 
+	// Tracks node's material parameter change(even if it's a submaterial of multimat that is assigned)
+	virtual void MaterialOtherEvent(NodeKeyTab& nodes) override
+	{
+		LogNodeEvent(L"MaterialOtherEvent", nodes);
+	}
+
 	virtual void Added(NodeKeyTab& nodes) override
 	{
 		LogNodeEvent(L"Added", nodes);
@@ -200,128 +196,234 @@ public:
 	ISceneEventManager::CallbackKey CallbackKey;
 };
 
-// Material changes more precise tracking can be done with ReferenceMaker
-// INodeEventCallback's MaterialOtherEvent tracks that any change is done to the material assigned to node
-// When a submaterial of a multimat is changed MaterialOtherEvent is called /omitting/ details which submaterial is modified
-// ReferenceMaker on the other hand tracks individual (sub)material changes
-// todo: 
-// - stop observing material when not needed(i.e. it's not assigned, used as submaterial or something else(?)
-// - remove when deleted
-class FMaterialObserver: public ReferenceMaker
-{
-	typedef int FMaterialIndex;
 
+// Watches all notifications sent to a single ReferenceTarget(INode, Mtl etc)
+template<typename ItemNotificationsHandler>
+class TBaseItemNotifications: public ReferenceMaker
+{
 public:
-	~FMaterialObserver()
+
+	explicit TBaseItemNotifications(ItemNotificationsHandler& InHandler): Handler(InHandler)
 	{
-		DeleteAllRefs(); // Required to be called in destructor
+		LogDebug(FString::Printf(TEXT("TBaseItemNotifications(%s - %p)::TBaseItemNotifications"), *GetName(), Handler.GetNotificationTargetHandle()));
+
+		// Don't record setting reference for Undo/Redo  - we control this ourselves(e.g. on ADDED_NODE event for addition)
+		HoldSuspend SuspendUndoRedoGuard; 
+		ReplaceReference(0, Handler.GetNotificationTargetHandle());
 	}
 
-	void Reset()
+	void Delete()
 	{
-		IndexToReferencedMaterial.Reset();
-		ReferencedMaterialToIndex.Reset();
+		LogDebug(FString::Printf(TEXT("TBaseItemNotifications(%p)::Delete"), Handler.GetNotificationTargetHandle()));
+		
+		HoldSuspend SuspendUndoRedoGuard; 
+		DeleteMe(); //Instead of calling DeleteAllRefs() in destructor, calling DeleteMe under HoldSuspend to disable undo for this object
 	}
 
-	RefResult NotifyRefChanged(const Interval& ChangeInterval, RefTargetHandle TargetHandle, PartID& PartId, RefMessage Message, BOOL propagate) override
+	// Same pointer could be added to the scene - handle this in case notifier for this pointer still kept
+	void ReAdded()
 	{
-		// todo: remove material handling???
-		ensure(ReferencedMaterialToIndex.Contains(TargetHandle));
+		ReplaceReference(0, Handler.GetNotificationTargetHandle());
+	}
 
+protected:
 
-		LogDebug(FString::Printf(TEXT("NotifyRefChanged: %s: %x"), dynamic_cast<Mtl*>(TargetHandle)->GetName().data(), Message));
+	virtual RefResult NotifyRefChanged(const Interval& ChangeInterval, RefTargetHandle TargetHandle, PartID& PartId, RefMessage Message, BOOL propagate) override
+	{
+		if (!bReferenceSet)
+		{
+			LOG_DEBUG_HEAVY(FString::Printf(TEXT("TBaseItemNotifications(UNSET - %p)::NotifyRefChanged:  Target - %p, Message - %x"), Handler.GetNotificationTargetHandle(), TargetHandle, Message));
+			return REF_DONTCARE;
+		}
+		
+		LOG_DEBUG_HEAVY(FString::Printf(TEXT("TBaseItemNotifications(%s - %p)::NotifyRefChanged:  Target - %p, Message - %x"), *GetName(), Handler.GetNotificationTargetHandle(), TargetHandle, Message));
 
-		return REF_SUCCEED;
+		return Handler.HandleEvent(ChangeInterval, TargetHandle, PartId, Message);
 	}
 
-	void AddMaterial(Mtl* Material)
+	virtual int NumRefs() override
+	{
+		return 1;
+	}
+
+	FString GetName()
+	{
+		return Handler.GetName();
+	}
+
+	virtual RefTargetHandle GetReference(int ReferenceIndex) override
 	{
-		if (!ReferencedMaterialToIndex.Contains(Material))
+		if (bReferenceSet)
+		{
+			LOG_DEBUG_HEAVY(FString::Printf(TEXT("TBaseItemNotifications(%s - %p)::GetReference: %d"), *GetName(), Handler.GetNotificationTargetHandle(), ReferenceIndex));
+			return Handler.GetNotificationTargetHandle();
+		}
+		else
 		{
-			ReplaceReference(NumRefs(), Material);
+			LOG_DEBUG_HEAVY(FString::Printf(TEXT("TBaseItemNotifications(UNSET - %p)::GetReference: %d"), Handler.GetNotificationTargetHandle(), ReferenceIndex));
+			return nullptr;
 		}
 	}
 
-	// todo: unused
-	// RECONSIDER: when this method is used - removed material reduces NumRefs result so adding new material will overwrite already existing reference 
-	// e.g. was two materials added, with index 0 and 1, material 0 removed, NumRefs becomes 1 so next call ReplaceReference(NumRefs(), Material) will replace material 1 in the map
-	void RemoveMaterial(Mtl* Material)
+	virtual void SetReference(int ReferenceIndex, RefTargetHandle Handle) override
 	{
-		FMaterialIndex MaterialIndex;
-		if (ReferencedMaterialToIndex.RemoveAndCopyValue(Material, MaterialIndex))
+		LOG_DEBUG_HEAVY(FString::Printf(TEXT("TBaseItemNotifications(%s - %p)::SetReference: %d, %p"), *GetName(), Handler.GetNotificationTargetHandle(), ReferenceIndex, Handle));
+
+		// Either clears reference or passes the same handler
+		check((Handle == nullptr ) || (Handle == Handler.GetNotificationTargetHandle()));
+
+		bool bReferenceSetNew = Handle != nullptr;
+
+		if (bReferenceSet && !bReferenceSetNew)
 		{
-			IndexToReferencedMaterial.Remove(MaterialIndex);
+			Handler.HandleDeleted();
 		}
+
+		bReferenceSet = bReferenceSetNew;
 	}
 
-	int NumRefs() override
+private:
+	ItemNotificationsHandler& Handler; // Handler for specific item type
+	bool bReferenceSet = false;
+
+protected:
+	virtual ~TBaseItemNotifications() override {} // Should call Delete
+
+};
+
+class FMaterialNotifications: public FNoncopyable
+{
+public:
+
+	FMaterialNotifications(ISceneTracker& InSceneTracker, Mtl& Material)
+		: SceneTracker(InSceneTracker)
+		, Material(Material)
+		, ItemNotifications(new TBaseItemNotifications{*this})
 	{
-		return IndexToReferencedMaterial.Num();
 	}
 
-	RefTargetHandle GetReference(int ReferenceIndex) override
+	~FMaterialNotifications()
 	{
-		return IndexToReferencedMaterial[ReferenceIndex];
+		ItemNotifications->Delete();
 	}
 
-	void SetReference(int ReferenceIndex, RefTargetHandle TargetHandle) override
+	FString GetName()
 	{
-		IndexToReferencedMaterial.Add(ReferenceIndex, TargetHandle);
-		ReferencedMaterialToIndex.Add(TargetHandle, ReferenceIndex);
+		return Material.GetName().data();
 	}
-private:
-	TMap<FMaterialIndex, RefTargetHandle> IndexToReferencedMaterial;
-	TMap<RefTargetHandle, FMaterialIndex> ReferencedMaterialToIndex;
 
+	RefTargetHandle GetNotificationTargetHandle()
+	{
+		return &Material;
+	}
+
+	RefResult HandleEvent(const Interval& ChangeInterval, RefTargetHandle TargetHandle, PartID& PartId, RefMessage Message)
+	{
+		if (SceneTracker.IsUpdateInProgress())
+		{
+			LogDebug(TEXT("     Event Skipped during Update"));
+			// Don't propagate ReferenceMaker events when Update is in progress - these events are result of accessing render geometry on some objects, including
+			// Smooth/TurboSmooth modifiers, VRay proxy and similar. Plugin retrieves render mesh from them and those objects fire geometry change events
+			// Plugin reverts settings back to previously set(i.e. viewport for smooth modifiers or whatever vray proxy had for display)
+			return REF_DONTCARE;
+		}
+
+		switch (Message)
+		{
+		case REFMSG_REF_DELETED:
+			// 
+		break;
+	    case REFMSG_CHANGE:
+			LOG_DEBUG_HEAVY(FString::Printf(TEXT("   REFMSG_CHANGE.PartId:  %llx"), PartId));
+
+			SceneTracker.MaterialGraphModified(&Material);
+		break;
+
+		case REFMSG_SUBANIM_STRUCTURE_CHANGED:
+			SceneTracker.MaterialGraphModified(&Material);
+		break;
+
+		default: 
+			break;
+		};
+
+		return REF_SUCCEED;
+	}
+
+	void HandleDeleted()
+	{
+	}
+
+	void ReAdded()
+	{
+		ItemNotifications->ReAdded();
+	}
+
+private:
+	ISceneTracker& SceneTracker;
+	Mtl& Material;
+	TBaseItemNotifications<FMaterialNotifications>* ItemNotifications;
 };
 
 
-// Every instance watches single node
-class FNodeNotifications: public ReferenceMaker
+class FMaterialObserver
 {
 public:
 	ISceneTracker& SceneTracker;
 
-	FNodeNotifications(ISceneTracker& InSceneTracker, INode& Node)
-		: SceneTracker(InSceneTracker)
-		, Node(Node)
+	FMaterialObserver(ISceneTracker& InSceneTracker) : SceneTracker(InSceneTracker)
 	{
-		LogDebug(FString::Printf(TEXT("FNodeNotifications(%s - %p)::FNodeNotifications"), Node.GetName(), &Node));
+	}
 
-		// Don't record setting reference for Undo/Redo  - we control this ourselves(e.g. on ADDED_NODE event for addition)
-		HoldSuspend SuspendUndoRedoGuard; 
-		ReplaceReference(0, &Node);
+	~FMaterialObserver()
+	{
 	}
 
-	void Delete()
+	void AddItem(Mtl* Material)
 	{
-		LogDebug(FString::Printf(TEXT("FNodeNotifications(%p)::Delete"), &Node));
-		
-		HoldSuspend SuspendUndoRedoGuard; 
-		DeleteMe(); //Instead of calling DeleteAllRefs() in destructor, calling DeleteMe under HoldSuspend to disable undo for this object
+		if (TUniquePtr<FMaterialNotifications>* NotifierPtr = Items.Find(Material))
+		{
+			(*NotifierPtr)->ReAdded();
+		}
+		else
+		{
+			Items.Emplace(Material, new FMaterialNotifications(SceneTracker, *Material));
+		}
 	}
 
-	// Same INode pointer could be added to the scene - handle this in case notifier for this pointer still kept
-	void ReAdded()
+	void Reset()
 	{
-		ReplaceReference(0, &Node);
+		Items.Reset();
 	}
 
-protected:
+private:
+	
+	TMap<Mtl*, TUniquePtr<FMaterialNotifications>> Items;
+};
 
-	virtual RefResult NotifyRefChanged(const Interval& ChangeInterval, RefTargetHandle TargetHandle, PartID& PartId, RefMessage Message, BOOL propagate) override
+class FNodeNotifications: public FNoncopyable
+{
+public:
+
+	FNodeNotifications(ISceneTracker& InSceneTracker, INode& Node)
+		: SceneTracker(InSceneTracker)
+		, Node(Node)
+		, ItemNotifications(new TBaseItemNotifications{*this})
+	{
+	}
+
+	~FNodeNotifications()
+	{
+		ItemNotifications->Delete();
+	}
+
+	RefResult HandleEvent(const Interval& ChangeInterval, RefTargetHandle TargetHandle, PartID& PartId, RefMessage Message)
 	{
-		if (!bReferenceSet)
-		{
-			LOG_DEBUG_HEAVY(FString::Printf(TEXT("FNodeNotifications(UNSET - %p)::NotifyRefChanged:  Target - %p, Message - %x"), &Node, TargetHandle, Message));
-			return REF_DONTCARE;
-		}
-		
-		LOG_DEBUG_HEAVY(FString::Printf(TEXT("FNodeNotifications(%s - %p)::NotifyRefChanged:  Target - %p, Message - %x"), Node.GetName(), &Node, TargetHandle, Message));
 		LogDebugNode(TEXT("     Node:"), &Node);
 
 		if (SceneTracker.IsUpdateInProgress())
 		{
+			LogDebug(TEXT("     Event Skipped during Update"));
 			// Don't propagate ReferenceMaker events when Update is in progress - these events are result of accessing render geometry on some objects, including
 			// Smooth/TurboSmooth modifiers, VRay proxy and similar. Plugin retrieves render mesh from them and those objects fire geometry change events
 			// Plugin reverts settings back to previously set(i.e. viewport for smooth modifiers or whatever vray proxy had for display)
@@ -377,55 +479,33 @@ protected:
 			break;
 		};
 
-
 		return REF_SUCCEED;
 	}
 
-	virtual int NumRefs() override
+	FString GetName()
 	{
-		return 1;
+		return Node.GetName();
 	}
 
-	virtual RefTargetHandle GetReference(int ReferenceIndex) override
+	RefTargetHandle GetNotificationTargetHandle()
 	{
-		if (bReferenceSet)
-		{
-			LOG_DEBUG_HEAVY(FString::Printf(TEXT("FNodeNotifications(%s - %p)::GetReference: %d"), Node.GetName(), &Node, ReferenceIndex));
-			return &Node;
-		}
-		else
-		{
-			LOG_DEBUG_HEAVY(FString::Printf(TEXT("FNodeNotifications(UNSET - %p)::GetReference: %d"), &Node, ReferenceIndex));
-			return nullptr;
-		}
+		return &Node;
 	}
 
-	virtual void SetReference(int ReferenceIndex, RefTargetHandle Handle) override
+	void HandleDeleted()
 	{
-		LOG_DEBUG_HEAVY(FString::Printf(TEXT("FNodeNotifications(%s - %p)::SetReference: %d, %p"), Node.GetName(), &Node, ReferenceIndex, Handle));
-
-		// Either clears reference or passes the same node
-		check((Handle == nullptr ) || (dynamic_cast<INode*>(Handle) == &Node));
-
-		bool bReferenceSetNew = Handle != nullptr;
-
-		if (bReferenceSet && !bReferenceSetNew)
-		{
-			// Sometimes node deleted without another notification event but with removing it from ReferenceMakers
-			// This is visible with Containers opening/closing them
-			SceneTracker.NodeDeleted(&Node);
-		}
+		SceneTracker.NodeDeleted(&Node);
+	}
 
-		bReferenceSet = bReferenceSetNew;
+	void ReAdded()
+	{
+		ItemNotifications->ReAdded();
 	}
 
 private:
+	ISceneTracker& SceneTracker;
 	INode& Node;
-	bool bReferenceSet = false;
-
-
-protected:
-	~FNodeNotifications(){} // Should call Delete
+	TBaseItemNotifications<FNodeNotifications>* ItemNotifications;
 };
 
 class FNodeObserver
@@ -443,34 +523,30 @@ public:
 
 	void AddItem(INode* Node)
 	{
-		if (FNodeNotifications** NotifierPtr = Nodes.Find(Node))
+		if (TUniquePtr<FNodeNotifications>* NotifierPtr = Items.Find(Node))
 		{
 			(*NotifierPtr)->ReAdded();
 		}
 		else
 		{
-			Nodes.Add(Node, new FNodeNotifications(SceneTracker, *Node));
+			Items.Emplace(Node, new FNodeNotifications(SceneTracker, *Node));
 		}
 	}
 
 	void Reset()
 	{
-		for (TPair<INode*, FNodeNotifications*> Node : Nodes)
-		{
-			Node.Value->Delete();
-		}
-		Nodes.Reset();
+		Items.Reset();
 	}
 
 private:
 	
-	TMap<INode*, FNodeNotifications*> Nodes;
+	TMap<INode*, TUniquePtr<FNodeNotifications>> Items;
 };
 
 FNotifications::FNotifications(IExporter& InExporter)
 	: Exporter(InExporter)
 	, NodeObserver(MakeUnique<FNodeObserver>(InExporter.GetSceneTracker()))
-	, MaterialObserver(MakeUnique<FMaterialObserver>())
+	, MaterialObserver(MakeUnique<FMaterialObserver>(InExporter.GetSceneTracker()))
 {
 	RegisterForSystemNotifications(); // Always follow system events like file reset, new, open and others. But some system notifications also track changes within scene.
 }
@@ -491,6 +567,11 @@ void FNotifications::AddNode(INode* Node)
 	NodeObserver->AddItem(Node);
 }
 
+void FNotifications::AddMaterial(Mtl* Node)
+{
+	MaterialObserver->AddItem(Node);
+}
+
 void FNotifications::RegisterForSystemNotifications()
 {
 	// Build todo: remove strings, for debug/logging
diff --git a/Engine/Source/Programs/Enterprise/Datasmith/DatasmithMaxExporter/Private/DatasmithMaxExporterWithDirectLink.cpp b/Engine/Source/Programs/Enterprise/Datasmith/DatasmithMaxExporter/Private/DatasmithMaxExporterWithDirectLink.cpp
--- a/Engine/Source/Programs/Enterprise/Datasmith/DatasmithMaxExporter/Private/DatasmithMaxExporterWithDirectLink.cpp
+++ b/Engine/Source/Programs/Enterprise/Datasmith/DatasmithMaxExporter/Private/DatasmithMaxExporterWithDirectLink.cpp
@@ -624,7 +624,7 @@ public:
 	virtual void SetSelectedOnly(bool bValue) override
 	{
 		Options.bSelectedOnly = bValue;
-		SetBool(TEXT("SelectedOnly"), bValue);
+		SetBool(TEXT("SelectedOnly"), Options.bSelectedOnly);
 	}
 	virtual bool GetSelectedOnly() override
 	{
@@ -634,7 +634,7 @@ public:
 	virtual void SetAnimatedTransforms(bool bValue) override
 	{
 		Options.bAnimatedTransforms = bValue;
-		SetBool(TEXT("AnimatedTransforms"), bValue);
+		SetBool(TEXT("AnimatedTransforms"), Options.bAnimatedTransforms);
 	}
 
 	virtual bool GetAnimatedTransforms() override
@@ -645,7 +645,7 @@ public:
 	virtual void SetStatSync(bool bValue) override
 	{
 		Options.bStatSync = bValue;
-		SetBool(TEXT("StatExport"), bValue);
+		SetBool(TEXT("StatExport"), Options.bStatSync);
 	}
 
 	virtual bool GetStatSync() override
@@ -657,7 +657,7 @@ public:
 	{
 		
 		Options.TextureResolution = FMath::Clamp(Value, 0, FExportOptions::TextureResolutionMax);
-		SetInt(TEXT("TextureResolution"), Value);
+		SetInt(TEXT("TextureResolution"), Options.TextureResolution);
 	}
 
 	virtual int32 GetTextureResolution() override
@@ -1680,23 +1680,18 @@ public:
 	// Node invalidated:
 	//   - on initial parsing
 	//   - change event received for node itself
-	//   - node up hierachy invalidated - previous or current
+	//   - node up hierarchy invalidated - previous or current
 	//   - node's validity interval is invalid for current time
-	// Possible improvements:
-	//   - do we really need to invalidate hierachy immediately? Might be long task(even though it's just setting flags) 
-	//      - probably do this on Update
-	//          'MarkedForUpdate' => recursively Invalidate(i.e. rebuild)
-	//          separate invalidation to 'Changed'(received event, interval, etc) and 'NeedRebuild'?
-	//            'NeedRebuild' are caused by 'Changed', not just 'invalidated' all.
-	//         - e.g. on initial parsing - only root nodes are 'changed'
-	//         - mark 'changed' as 'need rebuild' and descend to children, if a child has no 'need rebuild' set it, else skip its hierarchy
 
 	void InvalidateNode(FNodeTracker& NodeTracker, bool bCheckCalledInProgress = true)
 	{
-		// We don't expect node chances while Update inprogress(unless Invalidate called explicitly)
 		if (bCheckCalledInProgress)
 		{
-			ensure(!bUpdateInProgress);
+			// We don't expect node changes while Update inprogress(unless Invalidate called explicitly in certain places)
+			if (!ensure(!bUpdateInProgress))
+			{
+				return;
+			}
 		}
 
 		if (NodeTracker.bDeleted)
@@ -2326,9 +2321,19 @@ public:
 
 	void RegisterNodeForMaterial(FNodeTracker& NodeTracker, Mtl* Material)
 	{
+		if (!Material)
+		{
+			return;
+		}
+
 		FMaterialTracker* MaterialTracker = MaterialsCollectionTracker.AddMaterial(Material);
 		NodeTracker.MaterialTrackers.Add(MaterialTracker);
 		MaterialsAssignedToNodes.FindOrAdd(MaterialTracker).Add(&NodeTracker);
+
+		if (NotificationsHandler)
+		{
+			NotificationsHandler->AddMaterial(Material);
+		}
 	}
 
 	virtual void UnregisterNodeForMaterial(FNodeTracker& NodeTracker) override 
@@ -2785,6 +2790,21 @@ public:
 		InvalidateNode(NodeTracker); // Invalidate node that has this material assigned. This is needed to trigger rebuild - exported geometry might change(e.g. multimaterial changed to slots will change on static mesh)
 	}
 
+	virtual void MaterialGraphModified(Mtl* Material) override
+	{
+		MaterialsCollectionTracker.InvalidateMaterial(Material);
+
+		FMaterialTracker* MaterialTracker = MaterialsCollectionTracker.AddMaterial(Material);
+		
+		if (TSet<FNodeTracker*>* NodeTrackersPtr = MaterialsAssignedToNodes.Find(MaterialTracker))
+		{
+			for (FNodeTracker* NodeTracker : *NodeTrackersPtr)
+			{
+				InvalidateNode(*NodeTracker);
+			}
+		}
+	}
+
 	virtual void NodeGeometryChanged(INode* Node) override
 	{
 		// GeometryChanged is executed to handle:
diff --git a/Engine/Source/Programs/Enterprise/Datasmith/DatasmithMaxExporter/Public/DatasmithMaxDirectLink.h b/Engine/Source/Programs/Enterprise/Datasmith/DatasmithMaxExporter/Public/DatasmithMaxDirectLink.h
--- a/Engine/Source/Programs/Enterprise/Datasmith/DatasmithMaxExporter/Public/DatasmithMaxDirectLink.h
+++ b/Engine/Source/Programs/Enterprise/Datasmith/DatasmithMaxExporter/Public/DatasmithMaxDirectLink.h
@@ -32,6 +32,7 @@ MAX_INCLUDES_END
 // #define CANCEL_DEBUG_ENABLE 1
 
 
+
 void AssignMeshMaterials(TSharedPtr<IDatasmithMeshElement>&MeshElement, Mtl * Material, const TSet<uint16>&SupportedChannels);
 
 class FDatasmithMaxStaticMeshAttributes;
@@ -379,6 +380,8 @@ public:
 	virtual void NodeMaterialGraphModified(FNodeKey NodeKey) = 0;
 	virtual void NodeMaterialGraphModified(INode* NodeKey) = 0;
 
+	virtual void MaterialGraphModified(Mtl* Material) = 0;
+
 	virtual bool IsUpdateInProgress() = 0;
 
 	// Scene modification
@@ -442,6 +445,7 @@ public:
 	void StopSceneChangeTracking();
 
 	void AddNode(INode*);
+	void AddMaterial(Mtl*);
 
 	FString ConvertNotificationCodeToString(int code);
 	void PrepareForUpdate();
