[3ds Max] DirectLink: bake pivot when pivots are same

- enables conversion to Datasmith of non-uniformly scaled objects done along rotated axes(pivot) by baking pivots into geometry(when appropriate)

diff --git a/Engine/Source/Programs/Enterprise/Datasmith/DatasmithMaxExporter/Private/DatasmithMaxDirectLinkMesh.cpp b/Engine/Source/Programs/Enterprise/Datasmith/DatasmithMaxExporter/Private/DatasmithMaxDirectLinkMesh.cpp
--- a/Engine/Source/Programs/Enterprise/Datasmith/DatasmithMaxExporter/Private/DatasmithMaxDirectLinkMesh.cpp
+++ b/Engine/Source/Programs/Enterprise/Datasmith/DatasmithMaxExporter/Private/DatasmithMaxDirectLinkMesh.cpp
@@ -45,7 +45,7 @@
 	}
 };
 
-FRenderMeshForConversion GetMeshForGeomObject(TimeValue CurrentTime, INode* Node)
+FRenderMeshForConversion GetMeshForGeomObject(TimeValue CurrentTime, INode* Node, const FTransform& Pivot)
 {
 	// todo: baseline exporter uses GetBaseObject which takes result of EvalWorldState
 	// and searched down DerivedObject pipeline(by taking GetObjRef) 
@@ -58,6 +58,7 @@
 
 	FRenderMeshForConversion Result(Node, RenderMesh, bNeedsDelete);
 	Result.SetValidityInterval(ValidityInterval);
+	Result.SetPivot(Pivot);
 	return MoveTemp(Result);
 }
 
@@ -102,7 +103,7 @@
 }
 
 
-FRenderMeshForConversion GetMeshForCollision(TimeValue CurrentTime, ISceneTracker& SceneTracker, INode* Node)
+FRenderMeshForConversion GetMeshForCollision(TimeValue CurrentTime, ISceneTracker& SceneTracker, INode* Node, bool bBakePivot)
 {
 	// source: FDatasmithMaxMeshExporter::ExportMesh
 	FDatasmithConverter Converter;
@@ -113,8 +114,6 @@
 	if (CollisionNode)
 	{
 
-		bool bBakePivot = false; // todo: bake collision pivot if render mesh pivot is baked
-
 		FTransform ColliderPivot = FDatasmithMaxSceneExporter::GetPivotTransform(CollisionNode, Converter.UnitToCentimeter);
 
 		if (bIsCollisionFromDatasmithAttributes)
@@ -155,9 +154,8 @@
 	{
 		return FRenderMeshForConversion();
 	}
-	FRenderMeshForConversion RenderMesh = GetMeshForGeomObject(CurrentTime, CollisionNode);
-	RenderMesh.SetPivot(CollisionPivot);
-	return MoveTemp(RenderMesh);
+
+	return GetMeshForGeomObject(CurrentTime, CollisionNode, CollisionPivot);
 }
 
 
@@ -310,16 +308,14 @@
 	DatasmithMesh.SetNormal(Index * 3 + 2, NormalVector.X, NormalVector.Y, NormalVector.Z);
 }
 
-bool FillDatasmithMeshFromBoundingBox(TimeValue CurrentTime, FDatasmithMesh& DatasmithMesh, INode* ExportNode, FTransform Pivot)
+bool FillDatasmithMeshFromBoundingBox(TimeValue CurrentTime, FDatasmithMesh& DatasmithMesh, const GeomUtils::FRenderMeshForConversion& MaxMesh)
 {
-	if (!ExportNode)
+	if (!MaxMesh.GetNode())
 	{
 		return false;
 	}
 	FDatasmithConverter Converter;
 
-	FRenderMeshForConversion MaxMesh = GetMeshForGeomObject(CurrentTime, ExportNode);
-
 	if (!MaxMesh.IsValid())
 	{
 		return false;
@@ -343,7 +339,7 @@
 	for (int32 i = 0; i < 8; i++)
 	{
 		FVector Vertex = Converter.toDatasmithVector(BoundingBox[i]);
-		Vertex = Pivot.TransformPosition(Vertex);
+		Vertex = MaxMesh.GetPivot().TransformPosition(Vertex);
 		DatasmithMesh.SetVertex(i, Vertex.X, Vertex.Y, Vertex.Z);
 	}
 
@@ -351,7 +347,7 @@
 
 	Matrix3 RotationMatrix;
 	RotationMatrix.IdentityMatrix();
-	Quat ObjectOffsetRotation = ExportNode->GetObjOffsetRot();
+	Quat ObjectOffsetRotation = MaxMesh.GetNode()->GetObjOffsetRot();
 	RotateMatrix(RotationMatrix, ObjectOffsetRotation);
 	Point3 Point;
 	FVector NormalVector;
@@ -557,7 +553,7 @@
 	if (DatasmithAttributes && DatasmithAttributes->GetExportMode() == EStaticMeshExportMode::BoundingBox)
 	{
 		FDatasmithMesh DatasmithMesh;
-		if (GeomUtils::FillDatasmithMeshFromBoundingBox(CurrentTime, DatasmithMesh, MeshSource.Node, FTransform::Identity))
+		if (GeomUtils::FillDatasmithMeshFromBoundingBox(CurrentTime, DatasmithMesh, MeshSource.RenderMesh))
 		{
 			Scene.AddMeshElement(MeshConverted.DatasmithMeshElement, DatasmithMesh, nullptr);
 			return true;
@@ -574,7 +570,6 @@
 
 	MeshConversionParams RenderMeshParams = {
 		MeshSource.RenderMesh.GetNode(),
-		*MeshSource.MeshName,
 		MeshSource.RenderMesh,
 		MeshSource.bConsolidateMaterialIds
 	};
@@ -614,9 +609,8 @@
 
 		MeshConversionParams CollisionParams = {
 			MeshSource.CollisionMesh.GetNode(),
-			nullptr,
 			MeshSource.CollisionMesh,
-			true
+			true // Consolidate material ids into single mesh for collision
 		};
 
 		if (GeomUtils::CreateDatasmithMeshFromMaxMesh(DatasmithCollisionMesh, CollisionParams, SupportedChannelsDummy, UVChannelsMapDummy))
diff --git a/Engine/Source/Programs/Enterprise/Datasmith/DatasmithMaxExporter/Private/DatasmithMaxExporterWithDirectLink.cpp b/Engine/Source/Programs/Enterprise/Datasmith/DatasmithMaxExporter/Private/DatasmithMaxExporterWithDirectLink.cpp
--- a/Engine/Source/Programs/Enterprise/Datasmith/DatasmithMaxExporter/Private/DatasmithMaxExporterWithDirectLink.cpp
+++ b/Engine/Source/Programs/Enterprise/Datasmith/DatasmithMaxExporter/Private/DatasmithMaxExporterWithDirectLink.cpp
@@ -166,6 +166,8 @@
 	// Datasmith mesh conversion results
 	FMeshConverted Converted;
 	Interval ValidityInterval;
+	bool bShouldBakePivot; // Pivot(when it's common for all instances) should be baked into vertices in order to avoid creating extra actors
+	FTransform BakePivot; // The pivot transform to bake vertices with
 
 	bool HasMesh()
 	{
@@ -2205,16 +2207,59 @@
 			return;
 		}
 
+		bool bInitial = true;
+		bool bShouldBakePivot = true;
+		FTransform BakePivot = FTransform::Identity;
+
 		for (FNodeTracker* NodeTrackerPtr : Instances.NodeTrackers)
 		{
 			FNodeTracker& NodeTracker = *NodeTrackerPtr;
 			ClearNodeFromDatasmithScene(NodeTracker);
 			if (ensure(NodeTracker.GetConverter().ConverterType == FNodeConverter::MeshNode))
 			{
-				static_cast<FMeshNodeConverter&>(NodeTracker.GetConverter()).bMaterialsAssignedToStaticMesh = false;
+				FMeshNodeConverter& Converter = static_cast<FMeshNodeConverter&>(NodeTracker.GetConverter());
+				Converter.bMaterialsAssignedToStaticMesh = false;
+
+				// Pivot considerations
+				// - Single instance
+				//     a. identity pivot => bShouldBakePivot = false
+				//     b. non-identity pivot => bShouldBakePivot = true (and set BakePivot)
+				// - Multiple nodes
+				//     a. Identity ALL pivots => bShouldBakePivot = false
+				//     b. EQUAL non-Identity ALL pivots => bShouldBakePivot = true
+				//     c. NOT equal pivots => bShouldBakePivot = false, add warning
+				FTransform Pivot;
+				GetPivotTransform(NodeTracker, Pivot);
+
+				if (bInitial)
+				{
+					if (Pivot.Equals(BakePivot))
+					{
+						bShouldBakePivot = false;
+					}
+					else
+					{
+						BakePivot = Pivot;
+						bShouldBakePivot = true;
+					}
+
+					bInitial = false;
+				}
+				else
+				{
+					if (!Pivot.Equals(BakePivot))
+					{
+						bShouldBakePivot = false;
+
+						LogWarning(FString::Printf(TEXT("Multiple different pivots on instances of object %s."), NodeTracker.Node->GetName()));
+					}
+				}
 			}
 		}
 
+		Instances.bShouldBakePivot = bShouldBakePivot;
+		Instances.BakePivot = BakePivot;
+
 		// Export static mesh using first lucky node
 		// todo: possible optimization to reuse previous node. Somehow. 
 		for (FNodeTracker* NodeTrackerPtr : Instances.NodeTrackers)
@@ -2295,8 +2340,10 @@
 		return true;
 	}
 
-	void GetNodeObjectTransform(FNodeTracker& NodeTracker, FDatasmithConverter Converter, FTransform& ObjectTransform)
+	void GetNodeObjectTransform(FNodeTracker& NodeTracker, FTransform& ObjectTransform)
 	{
+		FDatasmithConverter Converter;
+
 		FVector Translation, Scale;
 		FQuat Rotation;
 
@@ -2394,6 +2441,12 @@
 		}
 	}
 
+	void GetPivotTransform(FNodeTracker& NodeTracker, FTransform& Pivot)
+	{
+		FDatasmithConverter Converter;
+		Pivot = FDatasmithMaxSceneExporter::GetPivotTransform(NodeTracker.Node, Converter.UnitToCentimeter);
+	}
+
 	virtual void ConvertGeometryNodeToDatasmith(FNodeTracker& NodeTracker, FMeshNodeConverter& MeshConverter) override
 	{
 		FInstances* InstancesPtr = InstancesManager.GetInstancesForNodeTracker(NodeTracker);
@@ -2403,16 +2456,16 @@
 		}
 		FInstances& Instances = *InstancesPtr;
 
-		FDatasmithConverter Converter;
+		FTransform ObjectTransform;
+		GetNodeObjectTransform(NodeTracker, ObjectTransform);
 
-		FTransform ObjectTransform;
-		GetNodeObjectTransform(NodeTracker, Converter, ObjectTransform);
+		FTransform Pivot;
+		GetPivotTransform(NodeTracker, Pivot);
 
-		FTransform Pivot = FDatasmithMaxSceneExporter::GetPivotTransform(NodeTracker.Node, Converter.UnitToCentimeter);
+		bool bPivotBakedInMesh = Instances.bShouldBakePivot;
+		// Create separate actor only when there are multiple instances
+		bool bNeedPivotComponent = !bPivotBakedInMesh && !Pivot.Equals(FTransform::Identity) && (Instances.NodeTrackers.Num() > 1) && Instances.HasMesh();  
 
-		// Create separate actor only when there are multiple instances, 
-		bool bNeedPivotComponent = !Pivot.Equals(FTransform::Identity) && (Instances.NodeTrackers.Num() > 1) && Instances.HasMesh();  
-
 		TSharedPtr<IDatasmithActorElement> DatasmithActorElement;
 		TSharedPtr<IDatasmithMeshActorElement> DatasmithMeshActor;
 
@@ -2453,15 +2506,16 @@
 		}
 
 		// Set transforms
+
+		// Remove pivot from the node actor transform when pivot is baked or separate component is created for pivot
+		FTransform NodeTransform = (bPivotBakedInMesh || bNeedPivotComponent) ? (Pivot.Inverse() * ObjectTransform) : ObjectTransform;
+
+		DatasmithActorElement->SetTranslation(NodeTransform.GetTranslation());
+		DatasmithActorElement->SetScale(NodeTransform.GetScale3D());
+		DatasmithActorElement->SetRotation(NodeTransform.GetRotation());
+
 		if (bNeedPivotComponent) 
 		{
-			// Remove pivot from the node actor transform
-			FTransform NodeTransform = Pivot.Inverse() * ObjectTransform;
-
-			DatasmithActorElement->SetTranslation(NodeTransform.GetTranslation());
-			DatasmithActorElement->SetScale(NodeTransform.GetScale3D());
-			DatasmithActorElement->SetRotation(NodeTransform.GetRotation());
-
 			// Setup mesh actor with (relative)pivot transform
 			DatasmithMeshActor->SetTranslation(Pivot.GetTranslation());
 			DatasmithMeshActor->SetRotation(Pivot.GetRotation());
@@ -2470,15 +2524,7 @@
 
 			DatasmithActorElement->AddChild(DatasmithMeshActor, EDatasmithActorAttachmentRule::KeepRelativeTransform);
 		}
-		else
-		{
-			FTransform NodeTransform = ObjectTransform;
 
-			DatasmithActorElement->SetTranslation(NodeTransform.GetTranslation());
-			DatasmithActorElement->SetScale(NodeTransform.GetScale3D());
-			DatasmithActorElement->SetRotation(NodeTransform.GetRotation());
-		}
-
 		FNodeConverted& Converted = NodeTracker.CreateConverted();
 		Converted.DatasmithActorElement = DatasmithActorElement;
 		Converted.DatasmithMeshActor = DatasmithMeshActor;
@@ -2531,8 +2577,8 @@
 
 		FMeshConverterSource MeshSource = {
 			Node, MeshName,
-			GeomUtils::GetMeshForGeomObject(CurrentSyncPoint.Time, Node), false,
-			GeomUtils::GetMeshForCollision(CurrentSyncPoint.Time, *this, Node),
+			GeomUtils::GetMeshForGeomObject(CurrentSyncPoint.Time, Node, Instances.bShouldBakePivot ? Instances.BakePivot : FTransform::Identity), false,
+			GeomUtils::GetMeshForCollision(CurrentSyncPoint.Time, *this, Node, Instances.bShouldBakePivot),
 		};
 
 		Instances.ValidityInterval = MeshSource.RenderMesh.GetValidityInterval() & MeshSource.CollisionMesh.GetValidityInterval();
@@ -2588,9 +2634,8 @@
 			NodeTracker.GetConverted().DatasmithActorElement->SetLayer(*NodeTracker.Layer->Name);
 		}
 
-		FDatasmithConverter Converter;
 		FTransform ObjectTransform;
-		GetNodeObjectTransform(NodeTracker, Converter, ObjectTransform);
+		GetNodeObjectTransform(NodeTracker, ObjectTransform);
 
 		FTransform NodeTransform = ObjectTransform;
 		TSharedRef<IDatasmithActorElement> DatasmithActorElement = NodeTracker.GetConverted().DatasmithActorElement.ToSharedRef();
diff --git a/Engine/Source/Programs/Enterprise/Datasmith/DatasmithMaxExporter/Public/DatasmithMaxDirectLink.h b/Engine/Source/Programs/Enterprise/Datasmith/DatasmithMaxExporter/Public/DatasmithMaxDirectLink.h
--- a/Engine/Source/Programs/Enterprise/Datasmith/DatasmithMaxExporter/Public/DatasmithMaxDirectLink.h
+++ b/Engine/Source/Programs/Enterprise/Datasmith/DatasmithMaxExporter/Public/DatasmithMaxDirectLink.h
@@ -306,7 +306,7 @@
 	GeomUtils::FRenderMeshForConversion RenderMesh; // Extracted render mesh
 	bool bConsolidateMaterialIds; // Whether to join all materials id into single material slot for render mesh (used when a geometry doesn't have a multimaterial assigned)
 
-	GeomUtils::FRenderMeshForConversion CollisionMesh; 
+	GeomUtils::FRenderMeshForConversion CollisionMesh;
 };
 
 
@@ -415,7 +415,6 @@
 struct MeshConversionParams
 {
 	INode* Node; // Node, this geom object created from
-	const TCHAR* MeshName;
 	const GeomUtils::FRenderMeshForConversion& RenderMesh; // Extracted render mesh
 	bool bConsolidateMaterialIds; // Whether to join all materials id into single material slot (used when a geometry doesn't have a multimaterial assigned)
 };
diff --git a/Engine/Source/Programs/Enterprise/Datasmith/DatasmithMaxExporter/Public/DatasmithMaxGeomUtils.h b/Engine/Source/Programs/Enterprise/Datasmith/DatasmithMaxExporter/Public/DatasmithMaxGeomUtils.h
--- a/Engine/Source/Programs/Enterprise/Datasmith/DatasmithMaxExporter/Public/DatasmithMaxGeomUtils.h
+++ b/Engine/Source/Programs/Enterprise/Datasmith/DatasmithMaxExporter/Public/DatasmithMaxGeomUtils.h
@@ -117,10 +117,9 @@
 bool ConvertForest(DatasmithMaxDirectLink::ISceneTracker& Scene, DatasmithMaxDirectLink::FNodeTracker& NodeTracker);
 
 Mesh* GetMeshFromRenderMesh(TimeValue CurrentTime, INode* Node, Interval& ValidityInterval, BOOL& bNeedsDelete);
-FRenderMeshForConversion GetMeshForGeomObject(TimeValue CurrentTime, INode* Node); // Extract mesh using already evaluated object
-FRenderMeshForConversion GetMeshForNode(TimeValue CurrentTime, INode* Node, FTransform Pivot); // Extract mesh evaluating node object
+FRenderMeshForConversion GetMeshForGeomObject(TimeValue CurrentTime, INode* Node, const FTransform& Pivot = FTransform::Identity); // Extract mesh using already evaluated object
 
-FRenderMeshForConversion GetMeshForCollision(TimeValue CurrentTime, DatasmithMaxDirectLink::ISceneTracker& SceneTracker, INode* Node);
+FRenderMeshForConversion GetMeshForCollision(TimeValue CurrentTime, DatasmithMaxDirectLink::ISceneTracker& SceneTracker, INode* Node, bool bBakePivot);
 INode* GetCollisionNode(DatasmithMaxDirectLink::ISceneTracker& SceneTracker, INode* OriginalNode, const FDatasmithMaxStaticMeshAttributes* DatasmithAttributes, bool& bOutFromDatasmithAttribute);
 
 void FillDatasmithMeshFromMaxMesh(TimeValue CurrentTime, FDatasmithMesh& DatasmithMesh, Mesh& MaxMesh, INode* ExportedNode, bool bForceSingleMat, TSet<uint16>& SupportedChannels, TMap<int32, int32>& UVChannelsMap, FTransform Pivot);
