[3ds Max] DirectLink: Log window Copy button copies whole content when nothing is selected

diff --git a/Engine/Source/Programs/Enterprise/Datasmith/DatasmithMaxExporter/Private/UI/Datasmith3dsMaxUI.cpp b/Engine/Source/Programs/Enterprise/Datasmith/DatasmithMaxExporter/Private/UI/Datasmith3dsMaxUI.cpp
--- a/Engine/Source/Programs/Enterprise/Datasmith/DatasmithMaxExporter/Private/UI/Datasmith3dsMaxUI.cpp
+++ b/Engine/Source/Programs/Enterprise/Datasmith/DatasmithMaxExporter/Private/UI/Datasmith3dsMaxUI.cpp
@@ -209,22 +209,20 @@ public:
 		];
 	}
 
-	void CopyLog( )
+	void CopyLog() const
 	{
-		TArray<TSharedPtr<FMessageData>> SelectedItems = MessagesView->GetSelectedItems();
-		
-		if (SelectedItems.Num() > 0)
+		const TArray<TSharedPtr<FMessageData>> SelectedItems = MessagesView->GetSelectedItems();
+
+		const TArray<TSharedPtr<FMessageData>>& Items = SelectedItems.IsEmpty() ? Messages: SelectedItems;
+
+		FString Text;
+		for( int32 Index = 0; Index < Items.Num(); ++Index )
 		{
-			FString SelectedText;
-		
-			for( int32 Index = 0; Index < SelectedItems.Num(); ++Index )
-			{
-				SelectedText += SelectedItems[Index]->MessageText;
-				SelectedText += LINE_TERMINATOR;
-			}
-		
-			FPlatformApplicationMisc::ClipboardCopy( *SelectedText );
+			Text += Items[Index]->MessageText;
+			Text += LINE_TERMINATOR;
 		}
+	
+		FPlatformApplicationMisc::ClipboardCopy(*Text);
 	}
 
 	void ClearMessages()
