Revit - Direct link sends updates even when autosync isn't enabled

diff --git a/Engine/Source/Programs/Enterprise/Datasmith/DatasmithRevitExporter/Private/DatasmithRevitDocument.cs b/Engine/Source/Programs/Enterprise/Datasmith/DatasmithRevitExporter/Private/DatasmithRevitDocument.cs
--- a/Engine/Source/Programs/Enterprise/Datasmith/DatasmithRevitExporter/Private/DatasmithRevitDocument.cs
+++ b/Engine/Source/Programs/Enterprise/Datasmith/DatasmithRevitExporter/Private/DatasmithRevitDocument.cs
@@ -113,6 +113,12 @@ namespace DatasmithRevitExporter
 
 		public void SetActiveDirectLinkInstance(View3D InView)
 		{
+			if (ActiveDirectLinkInstance != null && ActiveDirectLinkInstance.SyncView.Id == InView.Id)
+			{
+				// This view is already the active one
+				return;
+			}
+
 			// Disable existing instance, if there's active one.
 			ActiveDirectLinkInstance?.MakeActive(false);
 			ActiveDirectLinkInstance = null;
