Datasmith Revit: Ability to specify a custom origin point

diff --git a/Engine/Source/Programs/Enterprise/Datasmith/DatasmithRevitExporter/Private/DatasmithRevitSettingsDialog.cs b/Engine/Source/Programs/Enterprise/Datasmith/DatasmithRevitExporter/Private/DatasmithRevitSettingsDialog.cs
--- a/Engine/Source/Programs/Enterprise/Datasmith/DatasmithRevitExporter/Private/DatasmithRevitSettingsDialog.cs
+++ b/Engine/Source/Programs/Enterprise/Datasmith/DatasmithRevitExporter/Private/DatasmithRevitSettingsDialog.cs
@@ -347,6 +347,11 @@ namespace DatasmithRevitExporter
 			{
 				bSettingsChanged = true;
 			};
+
+			InsertionPointCombo.SelectedValueChanged += (S, E) =>
+			{
+				bSettingsChanged = true;
+			};
 		}
 
 		private void ReloadAddedGroupsList()
