[3ds Max] DirectLink: Max uv channel ids are remapped to Unreal's uv index

diff --git a/Engine/Source/Programs/Enterprise/Datasmith/DatasmithMaxExporter/Private/DatasmithMaxDirectLinkMaterial.cpp b/Engine/Source/Programs/Enterprise/Datasmith/DatasmithMaxExporter/Private/DatasmithMaxDirectLinkMaterial.cpp
--- a/Engine/Source/Programs/Enterprise/Datasmith/DatasmithMaxExporter/Private/DatasmithMaxDirectLinkMaterial.cpp
+++ b/Engine/Source/Programs/Enterprise/Datasmith/DatasmithMaxExporter/Private/DatasmithMaxDirectLinkMaterial.cpp
@@ -150,6 +150,7 @@ void FMaterialsCollectionTracker::ConvertMaterial(Mtl* Material, TSharedRef<IDat
 			SCENE_UPDATE_STAT_INC(UpdateMaterials, Converted);
 			DatasmithScene->AddMaterial(DatasmithMaterial);
 			UsedMaterialToDatasmithMaterial.Add(Material, DatasmithMaterial);
+			SceneTracker.RemapConvertedMaterialUVChannels(Material, DatasmithMaterial);
 		}
 	}
 
diff --git a/Engine/Source/Programs/Enterprise/Datasmith/DatasmithMaxExporter/Private/DatasmithMaxDirectLinkMesh.cpp b/Engine/Source/Programs/Enterprise/Datasmith/DatasmithMaxExporter/Private/DatasmithMaxDirectLinkMesh.cpp
--- a/Engine/Source/Programs/Enterprise/Datasmith/DatasmithMaxExporter/Private/DatasmithMaxDirectLinkMesh.cpp
+++ b/Engine/Source/Programs/Enterprise/Datasmith/DatasmithMaxExporter/Private/DatasmithMaxDirectLinkMesh.cpp
@@ -566,7 +566,6 @@
 	}
 
 	FDatasmithMesh DatasmithMesh;
-	TMap<int32, int32> UVChannelsMap;
 
 	MeshConversionParams RenderMeshParams = {
 		MeshSource.RenderMesh.GetNode(),
@@ -574,7 +573,7 @@
 		MeshSource.bConsolidateMaterialIds
 	};
 
-	if (!GeomUtils::CreateDatasmithMeshFromMaxMesh(DatasmithMesh, RenderMeshParams, MeshConverted.SupportedChannels, UVChannelsMap))
+	if (!GeomUtils::CreateDatasmithMeshFromMaxMesh(DatasmithMesh, RenderMeshParams, MeshConverted.SupportedChannels, MeshConverted.UVChannelsMap))
 	{
 		LogWarning(FString(TEXT("Invalid object: ")) + MeshSource.Node->GetName());
 		return false;
@@ -586,7 +585,7 @@
 		constexpr int32 MaxToUnrealUVOffset = -1;
 		constexpr int32 DefaultValue = -1;
 		const int32 SelectedLightmapUVChannel = DatasmithAttributes->GetLightmapUVChannel();
-		const int32* ExportedSelectedChannel = UVChannelsMap.Find(SelectedLightmapUVChannel + MaxToUnrealUVOffset);
+		const int32* ExportedSelectedChannel = MeshConverted.UVChannelsMap.Find(SelectedLightmapUVChannel + MaxToUnrealUVOffset);
 
 		if (ExportedSelectedChannel)
 		{
diff --git a/Engine/Source/Programs/Enterprise/Datasmith/DatasmithMaxExporter/Private/DatasmithMaxExporterWithDirectLink.cpp b/Engine/Source/Programs/Enterprise/Datasmith/DatasmithMaxExporter/Private/DatasmithMaxExporterWithDirectLink.cpp
--- a/Engine/Source/Programs/Enterprise/Datasmith/DatasmithMaxExporter/Private/DatasmithMaxExporterWithDirectLink.cpp
+++ b/Engine/Source/Programs/Enterprise/Datasmith/DatasmithMaxExporter/Private/DatasmithMaxExporterWithDirectLink.cpp
@@ -1198,6 +1198,26 @@
 		return bResult;
 	}
 
+	void RemapUVChannels(const TSharedPtr<IDatasmithUEPbrMaterialElement>& MaterialElement, const FDatasmithMaxMeshExporter::FUVChannelsMap& UVChannels)
+	{
+		// Remap UV channels of all composite textures of each shaders in the material
+		for (int32 i = 0; i < MaterialElement->GetExpressionsCount(); ++i)
+		{
+			IDatasmithMaterialExpression* MaterialExpression = MaterialElement->GetExpression(i);
+
+			if ( MaterialExpression && MaterialExpression->IsSubType( EDatasmithMaterialExpressionType::TextureCoordinate ) )
+			{
+				IDatasmithMaterialExpressionTextureCoordinate& TextureCoordinate = static_cast< IDatasmithMaterialExpressionTextureCoordinate& >( *MaterialExpression );
+
+				const int32* UVChannel = UVChannels.Find( TextureCoordinate.GetCoordinateIndex() );
+				if ( UVChannel != nullptr )
+				{
+					TextureCoordinate.SetCoordinateIndex( *UVChannel );
+				}
+			}
+		}
+	}
+
 	 // returns if any change in scene was encountered and scene update completed(i.e. DirectLink Sync can be run)
 	bool UpdateInternal(FUpdateProgress::FStage& MainStage)
 	{
@@ -2425,6 +2445,57 @@
 		return nullptr;
 	}
 
+	// Backtrack material to the converted mesh it's used on to fix uv channel mapping(Max channel id -> Datasmith uv channel index)
+	// Max channel ids are a sparse set on a Max mesh and its id number is just an integer key. But Unreal(and Datasmith) uses a continuous array of uv sets and therefore Unreal uv ids are index into this array, ranging from 0 to NumUVChannels-1 (including lightmap uvs)
+	// 
+	// todo: simple remapping works only when all meshes using the material have their uv channels built in the same way(same set and order)
+	//   in general a situation when one material is used on two different meshes with different channel order could be handled(like duplicating material to use different channel mappings)
+	virtual void RemapConvertedMaterialUVChannels(Mtl* ActualMaterial, const TSharedPtr<IDatasmithBaseMaterialElement>& DatasmithMaterial) override
+	{
+		if (!ensure(DatasmithMaterial->IsA(EDatasmithElementType::UEPbrMaterial)))
+		{
+			return;
+		}
+
+		TSet<FMaterialTracker*>* MaterialTrackers = MaterialsCollectionTracker.UsedMaterialToMaterialTracker.Find(ActualMaterial);
+
+		if (!MaterialTrackers)
+		{
+			return;
+		}
+
+		for (FMaterialTracker* MaterialTracker: *MaterialTrackers)
+		{
+			if (TSet<FNodeTracker*>* NodeTrackersPtr = MaterialsAssignedToNodes.Find(MaterialTracker))
+			{
+				for (FNodeTracker* NodeTracker: *NodeTrackersPtr)
+				{
+					if (!NodeTracker->HasConverter())
+					{
+						continue;
+					}
+							
+					if (NodeTracker->GetConverter().ConverterType == FNodeConverter::MeshNode)
+					{
+						if (FInstances* Instances = InstancesManager.GetInstancesForNodeTracker(*NodeTracker))
+						{
+							RemapUVChannels(StaticCastSharedPtr<IDatasmithUEPbrMaterialElement>(DatasmithMaterial), Instances->Converted.UVChannelsMap);
+							return;
+						}
+					}
+					else if (NodeTracker->GetConverter().ConverterType == FNodeConverter::HismNode)
+					{
+						for (FMeshConverted& MeshConverted: static_cast<FHismNodeConverter&>(NodeTracker->GetConverter()).Meshes)
+						{
+							RemapUVChannels(StaticCastSharedPtr<IDatasmithUEPbrMaterialElement>(DatasmithMaterial), MeshConverted.UVChannelsMap);
+							return;
+						}
+					}
+				}
+			}
+		}
+	}
+
 	virtual void AddGeometryNodeInstance(FNodeTracker& NodeTracker, FMeshNodeConverter& MeshConverter, Object* Obj) override
 	{
 		InvalidateInstances(InstancesManager.AddNodeTracker(NodeTracker, MeshConverter, Obj));
diff --git a/Engine/Source/Programs/Enterprise/Datasmith/DatasmithMaxExporter/Public/DatasmithMaxConverters.h b/Engine/Source/Programs/Enterprise/Datasmith/DatasmithMaxExporter/Public/DatasmithMaxConverters.h
--- a/Engine/Source/Programs/Enterprise/Datasmith/DatasmithMaxExporter/Public/DatasmithMaxConverters.h
+++ b/Engine/Source/Programs/Enterprise/Datasmith/DatasmithMaxExporter/Public/DatasmithMaxConverters.h
@@ -35,10 +35,13 @@ class FMeshConverted
 public:
 	TSharedPtr<IDatasmithMeshElement> DatasmithMeshElement;
 	TSet<uint16> SupportedChannels;
+	TMap<int32, int32> UVChannelsMap;
+
 	void ReleaseMeshConverted()
 	{
 		DatasmithMeshElement.Reset();
 		SupportedChannels.Reset();
+		UVChannelsMap.Reset();
 	}
 };
 
diff --git a/Engine/Source/Programs/Enterprise/Datasmith/DatasmithMaxExporter/Public/DatasmithMaxDirectLink.h b/Engine/Source/Programs/Enterprise/Datasmith/DatasmithMaxExporter/Public/DatasmithMaxDirectLink.h
--- a/Engine/Source/Programs/Enterprise/Datasmith/DatasmithMaxExporter/Public/DatasmithMaxDirectLink.h
+++ b/Engine/Source/Programs/Enterprise/Datasmith/DatasmithMaxExporter/Public/DatasmithMaxDirectLink.h
@@ -392,9 +392,9 @@
 	virtual void RemoveMaterial(const TSharedPtr<IDatasmithBaseMaterialElement>& DatasmithMaterial) = 0;
 	virtual void RemoveTexture(const TSharedPtr<IDatasmithTextureElement>&) = 0;
 	virtual void NodeXRefMerged(INode* Node) = 0;
+	virtual void RemapConvertedMaterialUVChannels(Mtl* ActualMaterial, const TSharedPtr<IDatasmithBaseMaterialElement>& DatasmithMaterial) = 0;
 
 
-
 	// Sync/Update
 	FSyncPoint CurrentSyncPoint;
 
