DatasmithExpoter - Adding a CloseCurrentSource() function to our DirectLink simplified API.

Some DCC may close a 3d scene without closing the application, in that case they must be able to close the DirectLink source as well.
With the simplified DirectLink API (FDatasmithDirectLink) this was only achievable by creating a new scene that would override the existing one, which is not desirable in this case.
Our C# exporters can only use the simplified API at the moment and this seems like a common enough use-case to be added to FDatasmithDirectLink.

diff --git a/Engine/Source/Developer/Datasmith/DatasmithExporter/Private/DatasmithDirectLink.cpp b/Engine/Source/Developer/Datasmith/DatasmithExporter/Private/DatasmithDirectLink.cpp
--- a/Engine/Source/Developer/Datasmith/DatasmithExporter/Private/DatasmithDirectLink.cpp
+++ b/Engine/Source/Developer/Datasmith/DatasmithExporter/Private/DatasmithDirectLink.cpp
@@ -59,6 +59,12 @@ public:
 		return true;
 	}
 
+	void CloseCurrentSource()
+	{
+		Endpoint->RemoveSource(Source);
+		CurrentScene.Reset();
+	}
+
 	TSharedRef<DirectLink::FEndpoint, ESPMode::ThreadSafe> GetEnpoint() const
 	{
 		return Endpoint;
@@ -97,6 +103,11 @@ bool FDatasmithDirectLink::InitializeForScene(const TSharedRef<IDatasmithScene>&
 bool FDatasmithDirectLink::UpdateScene(const TSharedRef<IDatasmithScene>& Scene)
 { return DirectLinkImpl->UpdateScene(Scene); }
 
+void FDatasmithDirectLink::CloseCurrentSource()
+{
+	DirectLinkImpl->CloseCurrentSource();
+}
+
 TSharedRef<DirectLink::FEndpoint, ESPMode::ThreadSafe> FDatasmithDirectLink::GetEnpoint()
 {
 	if (!DirectLinkImpl)
diff --git a/Engine/Source/Developer/Datasmith/DatasmithExporter/Public/DatasmithDirectLink.h b/Engine/Source/Developer/Datasmith/DatasmithExporter/Public/DatasmithDirectLink.h
--- a/Engine/Source/Developer/Datasmith/DatasmithExporter/Public/DatasmithDirectLink.h
+++ b/Engine/Source/Developer/Datasmith/DatasmithExporter/Public/DatasmithDirectLink.h
@@ -30,5 +30,10 @@ public:
 	bool InitializeForScene(const TSharedRef<IDatasmithScene>& Scene);
 	bool UpdateScene(const TSharedRef<IDatasmithScene>& Scene);
 
+	/**
+	 * Close the initialized DirectLink source if any.
+	 */
+	void CloseCurrentSource();
+
 	static TSharedRef<DirectLink::FEndpoint, ESPMode::ThreadSafe> GetEnpoint();
 };
diff --git a/Engine/Source/Developer/Datasmith/DatasmithFacade/Private/DirectLink/DatasmithFacadeDirectLink.cpp b/Engine/Source/Developer/Datasmith/DatasmithFacade/Private/DirectLink/DatasmithFacadeDirectLink.cpp
--- a/Engine/Source/Developer/Datasmith/DatasmithFacade/Private/DirectLink/DatasmithFacadeDirectLink.cpp
+++ b/Engine/Source/Developer/Datasmith/DatasmithFacade/Private/DirectLink/DatasmithFacadeDirectLink.cpp
@@ -70,6 +70,11 @@ bool FDatasmithFacadeDirectLink::UpdateScene(FDatasmithFacadeScene* FacadeScene)
 	return false;
 }
 
+void FDatasmithFacadeDirectLink::CloseCurrentSource()
+{
+	Impl.CloseCurrentSource();
+}
+
 void FDatasmithFacadeDirectLink::AddEndpointObserver(FDatasmithFacadeEndpointObserver* Observer)
 {
 	if (Observer)
diff --git a/Engine/Source/Developer/Datasmith/DatasmithFacade/Public/DirectLink/DatasmithFacadeDirectLink.h b/Engine/Source/Developer/Datasmith/DatasmithFacade/Public/DirectLink/DatasmithFacadeDirectLink.h
--- a/Engine/Source/Developer/Datasmith/DatasmithFacade/Public/DirectLink/DatasmithFacadeDirectLink.h
+++ b/Engine/Source/Developer/Datasmith/DatasmithFacade/Public/DirectLink/DatasmithFacadeDirectLink.h
@@ -20,6 +20,11 @@ public:
 	bool InitializeForScene(FDatasmithFacadeScene* FacadeScene);
 	bool UpdateScene(FDatasmithFacadeScene* FacadeScene);
 
+	/**
+	 * Close the initialized DirectLink source if any.
+	 */
+	void CloseCurrentSource();
+
 	/**
 	 * Register an IEndpointObserver that will be notified periodically with the last state of the swarm.
 	 * @param Observer      Object that should be notified
diff --git a/Engine/Source/Programs/Enterprise/Datasmith/DatasmithFacadeCSharp/Private/DatasmithFacadeCSharp.cpp b/Engine/Source/Programs/Enterprise/Datasmith/DatasmithFacadeCSharp/Private/DatasmithFacadeCSharp.cpp
--- a/Engine/Source/Programs/Enterprise/Datasmith/DatasmithFacadeCSharp/Private/DatasmithFacadeCSharp.cpp
+++ b/Engine/Source/Programs/Enterprise/Datasmith/DatasmithFacadeCSharp/Private/DatasmithFacadeCSharp.cpp
@@ -6733,6 +6733,14 @@ SWIGEXPORT unsigned int SWIGSTDCALL CSharp_FDatasmithFacadeDirectLink_UpdateScen
 }
 
 
+SWIGEXPORT void SWIGSTDCALL CSharp_FDatasmithFacadeDirectLink_CloseCurrentSource(void * jarg1) {
+  FDatasmithFacadeDirectLink *arg1 = (FDatasmithFacadeDirectLink *) 0 ;
+  
+  arg1 = (FDatasmithFacadeDirectLink *)jarg1; 
+  (arg1)->CloseCurrentSource();
+}
+
+
 SWIGEXPORT void SWIGSTDCALL CSharp_FDatasmithFacadeDirectLink_AddEndpointObserver(void * jarg1, void * jarg2) {
   FDatasmithFacadeDirectLink *arg1 = (FDatasmithFacadeDirectLink *) 0 ;
   FDatasmithFacadeEndpointObserver *arg2 = (FDatasmithFacadeEndpointObserver *) 0 ;
diff --git a/Engine/Source/Programs/Enterprise/Datasmith/DatasmithFacadeCSharp/Public/DatasmithFacadeCSharpPINVOKE.cs b/Engine/Source/Programs/Enterprise/Datasmith/DatasmithFacadeCSharp/Public/DatasmithFacadeCSharpPINVOKE.cs
--- a/Engine/Source/Programs/Enterprise/Datasmith/DatasmithFacadeCSharp/Public/DatasmithFacadeCSharpPINVOKE.cs
+++ b/Engine/Source/Programs/Enterprise/Datasmith/DatasmithFacadeCSharp/Public/DatasmithFacadeCSharpPINVOKE.cs
@@ -1859,6 +1859,9 @@ class DatasmithFacadeCSharpPINVOKE {
   [global::System.Runtime.InteropServices.DllImport("DatasmithFacadeCSharp", EntryPoint="CSharp_FDatasmithFacadeDirectLink_UpdateScene")]
   public static extern bool FDatasmithFacadeDirectLink_UpdateScene(global::System.Runtime.InteropServices.HandleRef jarg1, global::System.Runtime.InteropServices.HandleRef jarg2);
 
+  [global::System.Runtime.InteropServices.DllImport("DatasmithFacadeCSharp", EntryPoint="CSharp_FDatasmithFacadeDirectLink_CloseCurrentSource")]
+  public static extern void FDatasmithFacadeDirectLink_CloseCurrentSource(global::System.Runtime.InteropServices.HandleRef jarg1);
+
   [global::System.Runtime.InteropServices.DllImport("DatasmithFacadeCSharp", EntryPoint="CSharp_FDatasmithFacadeDirectLink_AddEndpointObserver")]
   public static extern void FDatasmithFacadeDirectLink_AddEndpointObserver(global::System.Runtime.InteropServices.HandleRef jarg1, global::System.Runtime.InteropServices.HandleRef jarg2);
 
diff --git a/Engine/Source/Programs/Enterprise/Datasmith/DatasmithFacadeCSharp/Public/FDatasmithFacadeDirectLink.cs b/Engine/Source/Programs/Enterprise/Datasmith/DatasmithFacadeCSharp/Public/FDatasmithFacadeDirectLink.cs
--- a/Engine/Source/Programs/Enterprise/Datasmith/DatasmithFacadeCSharp/Public/FDatasmithFacadeDirectLink.cs
+++ b/Engine/Source/Programs/Enterprise/Datasmith/DatasmithFacadeCSharp/Public/FDatasmithFacadeDirectLink.cs
@@ -85,6 +85,10 @@ public class FDatasmithFacadeDirectLink : global::System.IDisposable {
     return ret;
   }
 
+  public void CloseCurrentSource() {
+    DatasmithFacadeCSharpPINVOKE.FDatasmithFacadeDirectLink_CloseCurrentSource(swigCPtr);
+  }
+
   public void AddEndpointObserver(FDatasmithFacadeEndpointObserver Observer) {
     if (!RegisteredObservers.Contains(Observer))
     {
