[3ds Max] DirectLink: Ies light syncs texture element for Ies texture

diff --git a/Engine/Source/Programs/Enterprise/Datasmith/DatasmithMaxExporter/Private/DatasmithMaxConverters.cpp b/Engine/Source/Programs/Enterprise/Datasmith/DatasmithMaxExporter/Private/DatasmithMaxConverters.cpp
--- a/Engine/Source/Programs/Enterprise/Datasmith/DatasmithMaxExporter/Private/DatasmithMaxConverters.cpp
+++ b/Engine/Source/Programs/Enterprise/Datasmith/DatasmithMaxExporter/Private/DatasmithMaxConverters.cpp
@@ -14,6 +14,9 @@
 #include "DatasmithMaxLogger.h"
 #include "DatasmithMaxHelper.h"
 
+#include "HAL/PlatformFilemanager.h"
+#include "GenericPlatform/GenericPlatformFile.h"
+
 namespace  DatasmithMaxDirectLink
 {
 
@@ -92,7 +95,7 @@ void FLightNodeConverter::ConvertToDatasmith(ISceneTracker& SceneTracker, FNodeT
 		}
 		else
 		{
-			if ( !FDatasmithMaxSceneExporter::ParseLight(NodeTracker.Node, LightElement.ToSharedRef(), SceneTracker.GetDatasmithSceneRef()) )
+			if ( !FDatasmithMaxSceneExporter::ParseLight(*this, NodeTracker.Node, LightElement.ToSharedRef(), SceneTracker.GetDatasmithSceneRef()) )
 			{
 				return;
 			}
@@ -133,10 +136,37 @@ void FLightNodeConverter::ConvertToDatasmith(ISceneTracker& SceneTracker, FNodeT
 		NodeTracker.Validity.NarrowValidityToInterval(ValidityInterval);
 	}
 
+	if (IsIesProfileValid())
+	{
+		if (const TCHAR* TexturePath = SceneTracker.AcquireIesTexture(IesFilePath))
+		{
+			LightElement->SetIesTexturePathName(TexturePath);
+		}
+		else
+		{
+			LightElement->SetUseIes(false);  // Force disable Ies if ies profile not found
+		}
+	}
 }
 
 void FLightNodeConverter::RemoveFromTracked(ISceneTracker& SceneTracker, FNodeTracker& NodeTracker)
 {
+	SceneTracker.ReleaseIesTexture(IesFilePath);
+}
+
+void FLightNodeConverter::ApplyIesProfile(const TCHAR* InIesFilePath)
+{
+	IesFilePath = InIesFilePath;
+}
+
+const TCHAR* FLightNodeConverter::GetIesProfile()
+{
+	return *IesFilePath;
+}
+
+bool FLightNodeConverter::IsIesProfileValid()
+{
+	return FPlatformFileManager::Get().GetPlatformFile().FileExists(*IesFilePath);
 }
 
 void FMeshNodeConverter::Parse(ISceneTracker& SceneTracker, FNodeTracker& NodeTracker)
diff --git a/Engine/Source/Programs/Enterprise/Datasmith/DatasmithMaxExporter/Private/DatasmithMaxExporterWithDirectLink.cpp b/Engine/Source/Programs/Enterprise/Datasmith/DatasmithMaxExporter/Private/DatasmithMaxExporterWithDirectLink.cpp
--- a/Engine/Source/Programs/Enterprise/Datasmith/DatasmithMaxExporter/Private/DatasmithMaxExporterWithDirectLink.cpp
+++ b/Engine/Source/Programs/Enterprise/Datasmith/DatasmithMaxExporter/Private/DatasmithMaxExporterWithDirectLink.cpp
@@ -20,7 +20,6 @@
 #include "DatasmithMaxExporterUtils.h"
 
 #include "Modules/ModuleManager.h"
-#include "HAL/PlatformFilemanager.h"
 #include "Misc/ConfigCacheIni.h"
 #include "Misc/Paths.h"
 
@@ -546,8 +545,6 @@ public:
 	TSet<INode*> NodesPrepared;
 };
 
-
-
 struct FExportOptions
 {
 	bool bSelectedOnly = false;
@@ -935,6 +932,91 @@ public:
 	};
 #endif
 
+class FIesTexturesCollection: FNoncopyable
+{
+public:
+
+	struct FTextureIesConverted
+	{
+		FTextureIesConverted(const TSharedPtr<IDatasmithTextureElement>& InTextureElement)
+			: TextureElement(InTextureElement)
+			, UsageCount(0)
+		{
+		}
+
+		TSharedPtr<IDatasmithTextureElement> TextureElement;
+		int32 UsageCount;
+	};
+
+	FIesTexturesCollection(ISceneTracker& InSceneTracker): SceneTracker(InSceneTracker)
+	{
+	}
+
+	// Cache Ies textures by their path's base name
+	FString GetIesTextureKey(const FString& IesFilePath)
+	{
+		return FDatasmithUtils::SanitizeObjectName(FPaths::GetBaseFilename(IesFilePath));
+	}
+
+	const TCHAR* AcquireIesTexture(const FString& IesFilePath)
+	{
+		if (IesFilePath.IsEmpty())
+		{
+			return nullptr;
+		}
+
+		FTextureIesConverted& Texture = GetOrCreateIesTexture(IesFilePath);
+		Texture.UsageCount++;
+		return Texture.TextureElement->GetName();
+	}
+
+	FTextureIesConverted& GetOrCreateIesTexture(const FString& IesFilePath)
+	{
+		FString Name = GetIesTextureKey(IesFilePath);
+
+		if (FTextureIesConverted* TexturePtr = TexturesIes.Find(Name))
+		{
+			return *TexturePtr;
+		}
+
+		TSharedRef<IDatasmithTextureElement> TextureElement = FDatasmithSceneFactory::CreateTexture(*(Name + FString("_Ies")));
+
+		TextureElement->SetTextureMode(EDatasmithTextureMode::Ies);
+		TextureElement->SetLabel(*Name);
+		TextureElement->SetFile(*IesFilePath);
+
+		SceneTracker.GetDatasmithSceneRef()->AddTexture(TextureElement);
+
+		return TexturesIes.Emplace(Name, TextureElement);
+	}
+
+	void ReleaseIesTexture(const FString& IesFilePath)
+	{
+		FString BaseName = GetIesTextureKey(IesFilePath);
+
+		if(!TexturesIes.Contains(BaseName))
+		{
+			return;
+		}
+
+		FTextureIesConverted& Texture = TexturesIes[BaseName];
+		Texture.UsageCount--;
+		if (!Texture.UsageCount)
+		{
+			SceneTracker.RemoveTexture(Texture.TextureElement);
+			TexturesIes.Remove(BaseName);
+		}
+	}
+
+	void Reset()
+	{
+		TexturesIes.Reset();
+	}
+
+	ISceneTracker& SceneTracker;
+
+	TMap<FString, FTextureIesConverted> TexturesIes;
+};
 
 
 // Holds states of entities for syncronization and handles change events
@@ -945,7 +1027,9 @@ public:
 		: Options(InOptions)
 		, ExportedScene(InExportedScene)
 		, NotificationsHandler(InNotificationsHandler)
-		, MaterialsCollectionTracker(*this) {}
+		, MaterialsCollectionTracker(*this)
+		, IesTextures(*this)
+	{}
 
 	virtual TSharedRef<IDatasmithScene> GetDatasmithSceneRef() override
 	{
@@ -2417,6 +2501,16 @@ public:
 		NodeTracker.MaterialTrackers.Reset();
 	}
 
+	virtual const TCHAR* AcquireIesTexture(const FString& IesFilePath) override
+	{
+		return IesTextures.AcquireIesTexture(IesFilePath);
+	}
+
+	virtual void ReleaseIesTexture(const FString& IesFilePath) override
+	{
+		return IesTextures.ReleaseIesTexture(IesFilePath);
+	}
+
 	static Mtl* UpdateGeometryNodeMaterial(FSceneTracker& SceneTracker, FInstances& Instances, FNodeTracker& NodeTracker)
 	{
 		if (Instances.HasMesh())
@@ -3005,6 +3099,8 @@ public:
 
 		InvalidatedNodeTrackers.Reset();
 		InvalidatedInstances.Reset();
+
+		IesTextures.Reset();
 	}
 
 	///////////////////////////////////////////////
@@ -3028,6 +3124,8 @@ public:
 	FMaterialsCollectionTracker MaterialsCollectionTracker;
 	TMap<FMaterialTracker*, TSet<FNodeTracker*>> MaterialsAssignedToNodes;
 
+	FIesTexturesCollection IesTextures;
+
 	TMap<FNodeTracker*, TSharedPtr<IDatasmithMetaDataElement>> NodeDatasmithMetadata; // All scene nodes
 
 	// Nodes/instances that need to be rebuilt
diff --git a/Engine/Source/Programs/Enterprise/Datasmith/DatasmithMaxExporter/Private/DatasmithMaxSceneExporter.cpp b/Engine/Source/Programs/Enterprise/Datasmith/DatasmithMaxExporter/Private/DatasmithMaxSceneExporter.cpp
--- a/Engine/Source/Programs/Enterprise/Datasmith/DatasmithMaxExporter/Private/DatasmithMaxSceneExporter.cpp
+++ b/Engine/Source/Programs/Enterprise/Datasmith/DatasmithMaxExporter/Private/DatasmithMaxSceneExporter.cpp
@@ -1,6 +1,9 @@
 // Copyright Epic Games, Inc. All Rights Reserved.
 #include "DatasmithMaxSceneExporter.h"
 
+#include "DatasmithMaxDirectLink.h"
+#include "DatasmithMaxConverters.h"
+
 #include "DatasmithMaxAttributes.h"
 #include "DatasmithMaxCameraExporter.h"
 #include "DatasmithMaxExporterDefines.h"
@@ -12,6 +15,7 @@
 #include "DatasmithSceneFactory.h"
 #include "VRayLights.h"
 
+
 #include "GenericPlatform/GenericPlatformFile.h"
 #include "HAL/PlatformFileManager.h"
 #include "Misc/Paths.h"
@@ -942,7 +946,7 @@ bool FDatasmithMaxSceneExporter::ParseTransformAnimation(INode* Node, TSharedRef
 	return true;
 }
 
-bool FDatasmithMaxSceneExporter::ParseLight(INode* Node, TSharedRef< IDatasmithLightActorElement > LightElement, TSharedRef< IDatasmithScene > DatasmithScene)
+bool FDatasmithMaxSceneExporter::ParseLight(DatasmithMaxDirectLink::FLightNodeConverter& Converter, INode* Node, TSharedRef< IDatasmithLightActorElement > LightElement, TSharedRef< IDatasmithScene > DatasmithScene)
 {
 	EMaxLightClass LightClass = FDatasmithMaxSceneParser::GetLightClass(Node);
 	if (LightClass == EMaxLightClass::Unknown)
@@ -976,7 +980,7 @@ bool FDatasmithMaxSceneExporter::ParseLight(INode* Node, TSharedRef< IDatasmithL
 	{
 		TSharedRef< IDatasmithPointLightElement > PointLightElement = StaticCastSharedRef< IDatasmithPointLightElement >( LightElement );
 
-		ParsePhotometricLight( *Light, PointLightElement, DatasmithScene );
+		ParsePhotometricLight(Converter, *Light, PointLightElement, DatasmithScene );
 	}
 	else if ( LightClass == EMaxLightClass::VRayLight && LightElement->IsA( EDatasmithElementType::AreaLight ) )
 	{
@@ -994,29 +998,25 @@ bool FDatasmithMaxSceneExporter::ParseLight(INode* Node, TSharedRef< IDatasmithL
 	{
 		TSharedRef< IDatasmithPointLightElement > PointLightElement = StaticCastSharedRef< IDatasmithPointLightElement >( LightElement );
 
-		ParseVRayLightIES( *Light, PointLightElement, DatasmithScene );
+		ParseVRayLightIES( Converter,  *Light, PointLightElement, DatasmithScene );
 	}
 	else if ( LightClass == EMaxLightClass::CoronaLight && LightElement->IsA( EDatasmithElementType::AreaLight ) )
 	{
 		TSharedRef< IDatasmithAreaLightElement > AreaLightElement = StaticCastSharedRef< IDatasmithAreaLightElement >( LightElement );
 
-		ParseCoronaLight( *Light, AreaLightElement, DatasmithScene );
+		ParseCoronaLight( Converter, *Light, AreaLightElement, DatasmithScene );
 	}
 	else
 	{
-		ParseLightParameters( LightClass, *Light, LightElement, DatasmithScene );
+		ParseLightParameters( Converter, LightClass, *Light, LightElement, DatasmithScene );
 	}
 
-	IPlatformFile& PlatformFile = FPlatformFileManager::Get().GetPlatformFile();
-
 	if ( LightElement->GetUseIes() )
 	{
-		if ( !PlatformFile.FileExists( LightElement->GetIesFile() ) )
+		if ( !Converter.IsIesProfileValid() )
 		{
 			LightElement->SetUseIes( false );
-
-			FString Error = FString("IES light definition \"") + FPaths::GetCleanFilename( LightElement->GetIesFile() ) + FString("\" cannot be found");
-			DatasmithMaxLogger::Get().AddMissingAssetError(*Error);
+			DatasmithMaxDirectLink::LogWarning(FString::Printf(TEXT("IES light definition \"%s\" cannot be found for light \"%s\""), *FPaths::GetCleanFilename(Converter.GetIesProfile()), LightElement->GetName()));
 		}
 	}
 
@@ -1093,7 +1093,7 @@ bool FDatasmithMaxSceneExporter::ParseLightObject(LightObject& Light, TSharedRef
 	return true;
 }
 
-bool FDatasmithMaxSceneExporter::ParseCoronaLight(LightObject& Light, TSharedRef< IDatasmithAreaLightElement > AreaLightElement, TSharedRef< IDatasmithScene > DatasmithScene)
+bool FDatasmithMaxSceneExporter::ParseCoronaLight(DatasmithMaxDirectLink::FLightNodeConverter& Converter, LightObject& Light, TSharedRef< IDatasmithAreaLightElement > AreaLightElement, TSharedRef< IDatasmithScene > DatasmithScene)
 {
 	enum class ECoronaIntensityUnits
 	{
@@ -1168,7 +1168,7 @@ bool FDatasmithMaxSceneExporter::ParseCoronaLight(LightObject& Light, TSharedRef
 			}
 			else if ( FCString::Stricmp(ParamDefinition.int_name, TEXT("iesfile")) == 0 )
 			{
-				AreaLightElement->SetIesFile( *GetActualPath( ParamBlock2->GetStr( ParamDefinition.ID, GetCOREInterface()->GetTime() ) ) );
+				Converter.ApplyIesProfile(*GetActualPath( ParamBlock2->GetStr( ParamDefinition.ID, GetCOREInterface()->GetTime() ) ));
 			}
 			else if ( FCString::Stricmp(ParamDefinition.int_name, TEXT("colorMode")) == 0 )
 			{
@@ -1239,15 +1239,17 @@ bool FDatasmithMaxSceneExporter::ParseCoronaLight(LightObject& Light, TSharedRef
 
 	if ( AreaLightElement->GetUseIes() )
 	{
+
 		IPlatformFile& PlatformFile = FPlatformFileManager::Get().GetPlatformFile();
 
-		if (PlatformFile.FileExists(AreaLightElement->GetIesFile()))
+		if (Converter.IsIesProfileValid())
 		{
 			AreaLightElement->SetLightType(EDatasmithAreaLightType::Point);
 		}
 		else
 		{
 			AreaLightElement->SetUseIes( false );
+			DatasmithMaxDirectLink::LogWarning(FString::Printf(TEXT("IES light definition \"%s\" cannot be found for light \"%s\""), *FPaths::GetCleanFilename(Converter.GetIesProfile()), AreaLightElement->GetName()));
 		}
 	}
 
@@ -1359,7 +1361,7 @@ bool FDatasmithMaxSceneExporter::ParseCoronaLight(LightObject& Light, TSharedRef
 	return true;
 }
 
-bool FDatasmithMaxSceneExporter::ParsePhotometricLight(LightObject& Light, TSharedRef< IDatasmithPointLightElement > PointLightElement, TSharedRef< IDatasmithScene > DatasmithScene)
+bool FDatasmithMaxSceneExporter::ParsePhotometricLight(DatasmithMaxDirectLink::FLightNodeConverter& Converter, LightObject& Light, TSharedRef< IDatasmithPointLightElement > PointLightElement, TSharedRef< IDatasmithScene > DatasmithScene)
 {
 	LightscapeLight& PhotometricLight = static_cast< LightscapeLight& >( Light );
 
@@ -1415,8 +1417,13 @@ bool FDatasmithMaxSceneExporter::ParsePhotometricLight(LightObject& Light, TShar
 		PointLightElement->SetColor( FLinearColor( RGBFilter.x, RGBFilter.y, RGBFilter.z ) );
 	}
 
-	PointLightElement->SetIesFile( PhotometricLight.GetFullWebFileName() );
-	PointLightElement->SetUseIes( PhotometricLight.GetDistribution() == LightscapeLight::WEB_DIST );
+	bool bUseIes = PhotometricLight.GetDistribution() == LightscapeLight::WEB_DIST;
+
+	PointLightElement->SetUseIes( bUseIes );
+	if (bUseIes)
+	{
+		Converter.ApplyIesProfile(PhotometricLight.GetFullWebFileName());
+	}
 	PointLightElement->SetUseIesBrightness( false );
 
 	FQuat IesRotateX = FQuat( FVector::RightVector, FMath::DegreesToRadians( -PhotometricLight.GetWebRotateX() ) );
@@ -1758,7 +1765,7 @@ bool FDatasmithMaxSceneExporter::ParseVRayLightPortal(LightObject& Light, TShare
 	return true;
 }
 
-bool FDatasmithMaxSceneExporter::ParseVRayLightIES(LightObject& Light, TSharedRef< IDatasmithPointLightElement > PointLightElement, TSharedRef< IDatasmithScene > DatasmithScene)
+bool FDatasmithMaxSceneExporter::ParseVRayLightIES(DatasmithMaxDirectLink::FLightNodeConverter& Converter, LightObject& Light, TSharedRef< IDatasmithPointLightElement > PointLightElement, TSharedRef< IDatasmithScene > DatasmithScene)
 {
 	PointLightElement->SetUseIes( true );
 	PointLightElement->SetUseIesBrightness( false ); // Vray IES lights have their own intensity
@@ -1802,7 +1809,7 @@ bool FDatasmithMaxSceneExporter::ParseVRayLightIES(LightObject& Light, TSharedRe
 			}
 			else if ( FCString::Stricmp(ParamDefinition.int_name, TEXT("ies_file")) == 0 )
 			{
-				PointLightElement->SetIesFile( *GetActualPath( ParamBlock2->GetStr( ParamDefinition.ID, GetCOREInterface()->GetTime() ) ) );
+				Converter.ApplyIesProfile( *GetActualPath( ParamBlock2->GetStr( ParamDefinition.ID, GetCOREInterface()->GetTime() ) ) );
 			}
 			else if ( FCString::Stricmp(ParamDefinition.int_name, TEXT("rotation_X")) == 0 )
 			{
@@ -1915,7 +1922,7 @@ bool FDatasmithMaxSceneExporter::ParseVRayLightIES(LightObject& Light, TSharedRe
 	return true;
 }
 
-bool FDatasmithMaxSceneExporter::ParseLightParameters(EMaxLightClass LightClass, LightObject& Light, TSharedRef< IDatasmithLightActorElement > LightElement, TSharedRef< IDatasmithScene > DatasmithScene)
+bool FDatasmithMaxSceneExporter::ParseLightParameters(DatasmithMaxDirectLink::FLightNodeConverter& Converter, EMaxLightClass LightClass, LightObject& Light, TSharedRef< IDatasmithLightActorElement > LightElement, TSharedRef< IDatasmithScene > DatasmithScene)
 {
 	const int NumParamBlocks = Light.NumParamBlocks();
 
@@ -2006,7 +2013,7 @@ bool FDatasmithMaxSceneExporter::ParseLightParameters(EMaxLightClass LightClass,
 				if (FCString::Stricmp(ParamDefinition.int_name, TEXT("webfile")) == 0 || FCString::Stricmp(ParamDefinition.int_name, TEXT("IesFile")) == 0 || 
 					FCString::Stricmp(ParamDefinition.int_name, TEXT("ies_file")) == 0 || FCString::Stricmp(ParamDefinition.int_name, TEXT("filename")) == 0)
 				{
-					LightElement->SetIesFile( *GetActualPath(ParamBlock2->GetStr(ParamDefinition.ID, GetCOREInterface()->GetTime())) );
+					Converter.ApplyIesProfile( *GetActualPath(ParamBlock2->GetStr(ParamDefinition.ID, GetCOREInterface()->GetTime())) );
 				}
 			}
 
@@ -2124,47 +2131,6 @@ bool FDatasmithMaxSceneExporter::ExportActor(TSharedRef< IDatasmithScene > Datas
 	return ParseActor(Node, ActorElement, UnitMultiplier, DatasmithScene);
 }
 
-bool FDatasmithMaxSceneExporter::WriteXMLLightActor(TSharedRef< IDatasmithScene > DatasmithScene, INode* Parent, INode* Node, const TCHAR* Name, float UnitMultiplier)
-{
-	if ( !Node )
-	{
-		return false;
-	}
-
-	// weird behavior of 3dsmax it returns the head assembly as Instance of the sun
-	ObjectState ObjState = Node->EvalWorldState(0);
-
-	if (ObjState.obj == nullptr || ObjState.obj->SuperClassID() != LIGHT_CLASS_ID)
-	{
-		return false;
-	}
-
-	TSharedPtr< IDatasmithLightActorElement > LightElement = CreateLightElementForNode( Node, Name );
-
-	if ( !LightElement.IsValid() )
-	{
-		if (FDatasmithMaxSceneParser::GetLightClass(Node) == EMaxLightClass::SkyEquivalent)
-		{
-			DatasmithScene->SetUsePhysicalSky(true);
-		}
-		else
-		{
-			DatasmithMaxLogger::Get().AddUnsupportedLight(Node);
-		}
-		return false;
-	}
-
-	ParseActor(Node, LightElement.ToSharedRef(), UnitMultiplier, DatasmithScene);
-
-	if ( !ParseLight(Parent, LightElement.ToSharedRef(), DatasmithScene) )
-	{
-		return false;
-	}
-	
-	DatasmithScene->AddActor( LightElement );
-	return true;
-}
-
 bool FDatasmithMaxSceneExporter::ExportCameraActor(TSharedRef< IDatasmithScene > DatasmithScene, INode* Parent, INodeTab Instances, int InstanceIndex, const TCHAR* Name, float UnitMultiplier)
 {
 	// weird behavior of 3dsmax it returns the head assembly as Instance of the sun
diff --git a/Engine/Source/Programs/Enterprise/Datasmith/DatasmithMaxExporter/Public/DatasmithMaxConverters.h b/Engine/Source/Programs/Enterprise/Datasmith/DatasmithMaxExporter/Public/DatasmithMaxConverters.h
--- a/Engine/Source/Programs/Enterprise/Datasmith/DatasmithMaxExporter/Public/DatasmithMaxConverters.h
+++ b/Engine/Source/Programs/Enterprise/Datasmith/DatasmithMaxExporter/Public/DatasmithMaxConverters.h
@@ -132,6 +132,14 @@ public:
 	virtual void Parse(ISceneTracker& SceneTracker, FNodeTracker& NodeTracker) override;
 	virtual void ConvertToDatasmith(ISceneTracker& SceneTracker, FNodeTracker& NodeTracker) override;
 	virtual void RemoveFromTracked(ISceneTracker& SceneTracker, FNodeTracker& NodeTracker) override;
+
+	void ApplyIesProfile(const TCHAR* InIesFilePath);
+	const TCHAR* GetIesProfile();
+	bool IsIesProfileValid();
+
+private:
+	FString IesFilePath;
+
 };
 
 class FCameraNodeConverter: public FNodeConverter
diff --git a/Engine/Source/Programs/Enterprise/Datasmith/DatasmithMaxExporter/Public/DatasmithMaxDirectLink.h b/Engine/Source/Programs/Enterprise/Datasmith/DatasmithMaxExporter/Public/DatasmithMaxDirectLink.h
--- a/Engine/Source/Programs/Enterprise/Datasmith/DatasmithMaxExporter/Public/DatasmithMaxDirectLink.h
+++ b/Engine/Source/Programs/Enterprise/Datasmith/DatasmithMaxExporter/Public/DatasmithMaxDirectLink.h
@@ -394,7 +394,6 @@ public:
 	virtual void NodeXRefMerged(INode* Node) = 0;
 	virtual void RemapConvertedMaterialUVChannels(Mtl* ActualMaterial, const TSharedPtr<IDatasmithBaseMaterialElement>& DatasmithMaterial) = 0;
 
-
 	// Sync/Update
 	FSyncPoint CurrentSyncPoint;
 
@@ -404,6 +403,9 @@ public:
 
 	virtual void UnregisterNodeForMaterial(FNodeTracker& NodeTracker) = 0;
 
+	virtual const TCHAR* AcquireIesTexture(const FString& IesFilePath) = 0;
+	virtual void ReleaseIesTexture(const FString& IesFilePath) = 0;
+
 	virtual TSharedRef<IDatasmithScene> GetDatasmithSceneRef() = 0;
 
 	// Utility
diff --git a/Engine/Source/Programs/Enterprise/Datasmith/DatasmithMaxExporter/Public/DatasmithMaxSceneExporter.h b/Engine/Source/Programs/Enterprise/Datasmith/DatasmithMaxExporter/Public/DatasmithMaxSceneExporter.h
--- a/Engine/Source/Programs/Enterprise/Datasmith/DatasmithMaxExporter/Public/DatasmithMaxSceneExporter.h
+++ b/Engine/Source/Programs/Enterprise/Datasmith/DatasmithMaxExporter/Public/DatasmithMaxSceneExporter.h
@@ -34,6 +34,12 @@ class ToneOperator;
 
 enum EStaticMeshExportMode : uint8;
 
+
+namespace DatasmithMaxDirectLink
+{
+	class FLightNodeConverter;
+}
+
 /**
  * Lights may have multiple default orientation depending on their shape and type. 
  * This is a structure used when converting actor's coordinate in MaxToUnrealCoordinates to make sure to apply right correction to the light objects.
@@ -57,7 +63,6 @@ public:
 		float UnitMultiplier, bool bPivotIsBakedInGeometry, Mtl* StaticMeshMtl, const EStaticMeshExportMode& ExportMode);
 	static TSharedRef< IDatasmithActorElement > ExportHierarchicalInstanceStaticMeshActor(INode* Node, INode* CustomMeshNode, const TCHAR* Label, TSet<uint16>& SupportedChannels, Mtl* StaticMeshMtl, const TArray<Matrix3>* Instances,
 		const TCHAR* MeshName, float UnitMultiplier, const EStaticMeshExportMode& ExportMode, TSharedPtr< IDatasmithActorElement >& OutInversedHISMActor );
-	static bool WriteXMLLightActor(TSharedRef< IDatasmithScene > DatasmithScene, INode* Parent, INode* Node, const TCHAR* Name, float UnitMultiplier);
 	static bool ExportCameraActor(TSharedRef< IDatasmithScene > DatasmithScene, INode* Parent, INodeTab Instances, int InstanceIndex, const TCHAR* Name, float UnitMultiplier);
 	static void WriteEnvironment(TSharedRef< IDatasmithScene > DatasmithScene, bool bOnlySelection);
 	static void ExportToneOperator(TSharedRef< IDatasmithScene > DatasmithScene);
@@ -81,7 +86,7 @@ public:
 	static void ParseMaterialForMeshActor(Mtl* Material, TSharedRef< IDatasmithMeshActorElement >& MeshActor, TSet<uint16>& SupportedChannels, const FVector3f& RandomSeed);
 	
 	static TSharedPtr< IDatasmithLightActorElement > CreateLightElementForNode(INode* Node, const TCHAR* Name);
-	static bool ParseLight(INode* Node, TSharedRef< IDatasmithLightActorElement > LightElement, TSharedRef< IDatasmithScene > DatasmithScene);
+	static bool ParseLight(DatasmithMaxDirectLink::FLightNodeConverter&, INode* Node, TSharedRef< IDatasmithLightActorElement > LightElement, TSharedRef< IDatasmithScene > DatasmithScene);
 
 	static TSharedPtr<IDatasmithMetaDataElement> ParseUserProperties(INode* Node, TSharedRef< IDatasmithActorElement > ActorElement, TSharedRef< IDatasmithScene > DatasmithScene);
 
@@ -93,12 +98,12 @@ private:
 
 	static bool ParseActor(INode* Node, TSharedRef< IDatasmithActorElement > ActorElement, float UnitMultiplier, TSharedRef< IDatasmithScene > DatasmithScene);
 	static bool ParseLightObject(LightObject& Light, TSharedRef< IDatasmithLightActorElement > LightElement, TSharedRef< IDatasmithScene > DatasmithScene);
-	static bool ParseCoronaLight(LightObject& Light, TSharedRef< IDatasmithAreaLightElement > AreaLightElement, TSharedRef< IDatasmithScene > DatasmithScene);
-	static bool ParsePhotometricLight(LightObject& Light, TSharedRef< IDatasmithPointLightElement > PointLightElement, TSharedRef< IDatasmithScene > DatasmithScene);
+	static bool ParseCoronaLight(DatasmithMaxDirectLink::FLightNodeConverter&, LightObject& Light, TSharedRef< IDatasmithAreaLightElement > AreaLightElement, TSharedRef< IDatasmithScene > DatasmithScene);
+	static bool ParsePhotometricLight(DatasmithMaxDirectLink::FLightNodeConverter& Converter, LightObject& Light, TSharedRef< IDatasmithPointLightElement > PointLightElement, TSharedRef< IDatasmithScene > DatasmithScene);
 	static bool ParseVRayLight(LightObject& Light, TSharedRef< IDatasmithAreaLightElement > AreaLightElement, TSharedRef< IDatasmithScene > DatasmithScene);
 	static bool ParseVRayLightPortal(LightObject& Light, TSharedRef< IDatasmithLightmassPortalElement > LightPortalElement, TSharedRef< IDatasmithScene > DatasmithScene);
-	static bool ParseVRayLightIES(LightObject& Light, TSharedRef< IDatasmithPointLightElement > PointLightElement, TSharedRef< IDatasmithScene > DatasmithScene);
-	static bool ParseLightParameters(EMaxLightClass LightClass, LightObject& Light, TSharedRef< IDatasmithLightActorElement > LightElement, TSharedRef< IDatasmithScene > DatasmithScene);
+	static bool ParseVRayLightIES(DatasmithMaxDirectLink::FLightNodeConverter&, LightObject& Light, TSharedRef< IDatasmithPointLightElement > PointLightElement, TSharedRef< IDatasmithScene > DatasmithScene);
+	static bool ParseLightParameters(DatasmithMaxDirectLink::FLightNodeConverter&, EMaxLightClass LightClass, LightObject& Light, TSharedRef< IDatasmithLightActorElement > LightElement, TSharedRef< IDatasmithScene > DatasmithScene);
 	static bool ProcessLightTexture(TSharedRef< IDatasmithLightActorElement > LightElement, Texmap* LightTexture, TSharedRef< IDatasmithScene > DatasmithScene);
 	static bool ParseTransformAnimation(INode* Node, TSharedRef< IDatasmithTransformAnimationElement > AnimationElement, float UnitMultiplier, const FMaxLightCoordinateConversionParams& LightParams);
 	static void ParseSun(INode* Node, TSharedRef<IDatasmithLightActorElement> LightElement);
