[3ds Max] DirectLink: fixed export of animated transform when parent of animated node is hidded

Animation for transforms exported as relative to parent actors. Hidden nodes are not exported by the plugin so relative transform need to be taken to closest exported ancestor node.

diff --git a/Engine/Source/Programs/Enterprise/Datasmith/DatasmithMaxExporter/Private/DatasmithMaxExporterWithDirectLink.cpp b/Engine/Source/Programs/Enterprise/Datasmith/DatasmithMaxExporter/Private/DatasmithMaxExporterWithDirectLink.cpp
--- a/Engine/Source/Programs/Enterprise/Datasmith/DatasmithMaxExporter/Private/DatasmithMaxExporterWithDirectLink.cpp
+++ b/Engine/Source/Programs/Enterprise/Datasmith/DatasmithMaxExporter/Private/DatasmithMaxExporterWithDirectLink.cpp
@@ -1718,17 +1718,21 @@ public:
 
 			if (NodeTracker->HasConverted())
 			{
+				// Not all nodes are converted to Datasmith actors(e.g. hidden nodes are omitted from Datasmith scene)
+				// Transform animation is exported relative to parent actor's so we need node with actual datasmith actor to compute relative transform.
+				FNodeTracker* ParentNodeTracker = GetAncestorNodeTrackerWithDatasmithActor(*NodeTracker);
+				INode* ParentNode = ParentNodeTracker ? ParentNodeTracker->Node : nullptr;
 
 				if (NodeTracker->GetConverterType() == FNodeConverter::LightNode)
 				{
 					const TSharedPtr<IDatasmithLightActorElement> LightElement = StaticCastSharedPtr< IDatasmithLightActorElement >(NodeTracker->GetConverted().DatasmithActorElement);
 					const FMaxLightCoordinateConversionParams LightParams(NodeTracker->Node,
 						LightElement->IsA(EDatasmithElementType::AreaLight) ? StaticCastSharedPtr<IDatasmithAreaLightElement>(LightElement)->GetLightShape() : EDatasmithLightShape::None);
-					FDatasmithMaxSceneExporter::ExportAnimation(LevelSequence, NodeTracker->Node, NodeTracker->GetConverted().DatasmithActorElement->GetName(), Converter.UnitToCentimeter, LightParams);
+					FDatasmithMaxSceneExporter::ExportAnimation(LevelSequence, ParentNode, NodeTracker->Node, NodeTracker->GetConverted().DatasmithActorElement->GetName(), Converter.UnitToCentimeter, LightParams);
 				}
 				else
 				{
-					FDatasmithMaxSceneExporter::ExportAnimation(LevelSequence, NodeTracker->Node, NodeTracker->GetConverted().DatasmithActorElement->GetName(), Converter.UnitToCentimeter);
+					FDatasmithMaxSceneExporter::ExportAnimation(LevelSequence, ParentNode, NodeTracker->Node, NodeTracker->GetConverted().DatasmithActorElement->GetName(), Converter.UnitToCentimeter);
 				}
 			}
 		}
diff --git a/Engine/Source/Programs/Enterprise/Datasmith/DatasmithMaxExporter/Private/DatasmithMaxSceneExporter.cpp b/Engine/Source/Programs/Enterprise/Datasmith/DatasmithMaxExporter/Private/DatasmithMaxSceneExporter.cpp
--- a/Engine/Source/Programs/Enterprise/Datasmith/DatasmithMaxExporter/Private/DatasmithMaxSceneExporter.cpp
+++ b/Engine/Source/Programs/Enterprise/Datasmith/DatasmithMaxExporter/Private/DatasmithMaxSceneExporter.cpp
@@ -267,11 +267,11 @@ FTransform FDatasmithMaxSceneExporter::GetPivotTransform( INode* Node, float Uni
 	return Pivot;
 }
 
-void FDatasmithMaxSceneExporter::ExportAnimation( TSharedRef< IDatasmithLevelSequenceElement > LevelSequence, INode* Node, const TCHAR* Name, float UnitMultiplier, const FMaxLightCoordinateConversionParams& LightParams)
+void FDatasmithMaxSceneExporter::ExportAnimation( TSharedRef< IDatasmithLevelSequenceElement > LevelSequence, INode* ParentNode, INode* Node, const TCHAR* Name, float UnitMultiplier, const FMaxLightCoordinateConversionParams& LightParams)
 {
 	TSharedRef< IDatasmithTransformAnimationElement > Animation = FDatasmithSceneFactory::CreateTransformAnimation( Name );
 
-	if ( ParseTransformAnimation( Node, Animation, UnitMultiplier, LightParams) )
+	if ( ParseTransformAnimation( ParentNode, Node, Animation, UnitMultiplier, LightParams) )
 	{
 		LevelSequence->AddAnimation( Animation );
 	}
@@ -821,7 +821,7 @@ bool FDatasmithMaxSceneExporter::ParseActor(INode* Node, TSharedRef< IDatasmithA
 	return true;
 }
 
-bool FDatasmithMaxSceneExporter::ParseTransformAnimation(INode* Node, TSharedRef< IDatasmithTransformAnimationElement > AnimationElement, float UnitMultiplier, const FMaxLightCoordinateConversionParams& LightParams)
+bool FDatasmithMaxSceneExporter::ParseTransformAnimation(INode* ParentNode, INode* Node, TSharedRef< IDatasmithTransformAnimationElement > AnimationElement, float UnitMultiplier, const FMaxLightCoordinateConversionParams& LightParams)
 {
 	if (Node == nullptr)
 	{
@@ -838,13 +838,12 @@ bool FDatasmithMaxSceneExporter::ParseTransformAnimation(INode* Node, TSharedRef
 	// Query the node transform in local space at -infinity to determine if it has any animation
 	Interval ValidInterval = FOREVER;
 	Matrix3 NodeTransform = Node->GetNodeTM(TIME_NegInfinity, &ValidInterval);
-	INode* Parent = Node->GetParentNode();
 
 // Matrix3::Matrix3(BOOL) is deprecated in 3ds max 2022 SDK
 #if MAX_PRODUCT_YEAR_NUMBER < 2022
-	Matrix3 ParentTransform = Parent ? Parent->GetNodeTM(TIME_NegInfinity, &ValidInterval) : Matrix3(true);
+	Matrix3 ParentTransform = ParentNode ? ParentNode->GetNodeTM(TIME_NegInfinity, &ValidInterval) : Matrix3(true);
 #else
-	Matrix3 ParentTransform = Parent ? Parent->GetNodeTM(TIME_NegInfinity, &ValidInterval) : Matrix3();
+	Matrix3 ParentTransform = ParentNode ? ParentNode->GetNodeTM(TIME_NegInfinity, &ValidInterval) : Matrix3();
 #endif
 	Matrix3 LocalTransform;
 
@@ -866,12 +865,12 @@ bool FDatasmithMaxSceneExporter::ParseTransformAnimation(INode* Node, TSharedRef
 
 		// The parent node could change at each frame because of the Link Constraint
 		NodeTransform = Node->GetNodeTM(CurrentTime, &ValidInterval);
-		Parent = Node->GetParentNode();
+		
 // Matrix3::Matrix3(BOOL) is deprecated in 3ds max 2022 SDK
 #if MAX_PRODUCT_YEAR_NUMBER < 2022
-		ParentTransform = Parent ? Parent->GetNodeTM(CurrentTime, &ValidInterval) : Matrix3(true);
+		ParentTransform = ParentNode ? ParentNode->GetNodeTM(CurrentTime, &ValidInterval) : Matrix3(true);
 #else
-		ParentTransform = Parent ? Parent->GetNodeTM(CurrentTime, &ValidInterval) : Matrix3();
+		ParentTransform = ParentNode ? ParentNode->GetNodeTM(CurrentTime, &ValidInterval) : Matrix3();
 #endif
 
 		Matrix3 InvParentTransform = Inverse(ParentTransform);
diff --git a/Engine/Source/Programs/Enterprise/Datasmith/DatasmithMaxExporter/Public/DatasmithMaxSceneExporter.h b/Engine/Source/Programs/Enterprise/Datasmith/DatasmithMaxExporter/Public/DatasmithMaxSceneExporter.h
--- a/Engine/Source/Programs/Enterprise/Datasmith/DatasmithMaxExporter/Public/DatasmithMaxSceneExporter.h
+++ b/Engine/Source/Programs/Enterprise/Datasmith/DatasmithMaxExporter/Public/DatasmithMaxSceneExporter.h
@@ -66,7 +66,7 @@ public:
 	static bool ExportCameraActor(TSharedRef< IDatasmithScene > DatasmithScene, INode* Parent, INodeTab Instances, int InstanceIndex, const TCHAR* Name, float UnitMultiplier);
 	static void WriteEnvironment(TSharedRef< IDatasmithScene > DatasmithScene, bool bOnlySelection);
 	static void ExportToneOperator(TSharedRef< IDatasmithScene > DatasmithScene);
-	static void ExportAnimation(TSharedRef< IDatasmithLevelSequenceElement > LevelSequence, INode* Node, const TCHAR* ActorName, float UnitMultiplier, const FMaxLightCoordinateConversionParams& LightParams = FMaxLightCoordinateConversionParams());
+	static void ExportAnimation(TSharedRef< IDatasmithLevelSequenceElement > LevelSequence, INode* ParentNode, INode* Node, const TCHAR* ActorName, float UnitMultiplier, const FMaxLightCoordinateConversionParams& LightParams = FMaxLightCoordinateConversionParams());
 
 	static FString GetActualPath(const TCHAR* OriginalPath);
 
@@ -105,6 +105,6 @@ private:
 	static bool ParseVRayLightIES(DatasmithMaxDirectLink::FLightNodeConverter&, LightObject& Light, TSharedRef< IDatasmithPointLightElement > PointLightElement, TSharedRef< IDatasmithScene > DatasmithScene);
 	static bool ParseLightParameters(DatasmithMaxDirectLink::FLightNodeConverter&, EMaxLightClass LightClass, LightObject& Light, TSharedRef< IDatasmithLightActorElement > LightElement, TSharedRef< IDatasmithScene > DatasmithScene);
 	static bool ProcessLightTexture(TSharedRef< IDatasmithLightActorElement > LightElement, Texmap* LightTexture, TSharedRef< IDatasmithScene > DatasmithScene);
-	static bool ParseTransformAnimation(INode* Node, TSharedRef< IDatasmithTransformAnimationElement > AnimationElement, float UnitMultiplier, const FMaxLightCoordinateConversionParams& LightParams);
+	static bool ParseTransformAnimation(INode* ParentNode, INode* Node, TSharedRef< IDatasmithTransformAnimationElement > AnimationElement, float UnitMultiplier, const FMaxLightCoordinateConversionParams& LightParams);
 	static void ParseSun(INode* Node, TSharedRef<IDatasmithLightActorElement> LightElement);
 };
