[3ds Max] DirectLlink: RailClone transform math fix

- fixed order of transform multiplication (note: this is improvement over old exporter)

diff --git a/Engine/Source/Programs/Enterprise/Datasmith/DatasmithMaxExporter/Private/DatasmithMaxSceneExporter.cpp b/Engine/Source/Programs/Enterprise/Datasmith/DatasmithMaxExporter/Private/DatasmithMaxSceneExporter.cpp
--- a/Engine/Source/Programs/Enterprise/Datasmith/DatasmithMaxExporter/Private/DatasmithMaxSceneExporter.cpp
+++ b/Engine/Source/Programs/Enterprise/Datasmith/DatasmithMaxExporter/Private/DatasmithMaxSceneExporter.cpp
@@ -399,7 +399,7 @@ TSharedRef< IDatasmithActorElement > FDatasmithMaxSceneExporter::ExportHierarchi
 		FVector InstancePos, InstanceScale;
 		FQuat InstanceRotation;
 		MaxToUnrealCoordinates((*Instances)[0], InstancePos, InstanceRotation, InstanceScale, UnitMultiplier);
-		FTransform WithInstance(WithoutInstance * FTransform(InstanceRotation, InstancePos, InstanceScale));
+		FTransform WithInstance(FTransform(InstanceRotation, InstancePos, InstanceScale) * WithoutInstance);
 
 		MeshActor->SetTranslation(WithInstance.GetLocation());
 		MeshActor->SetScale(WithInstance.GetScale3D());
