[Navisworks Exporter] Add support for Navisworks 2023

diff --git a/Engine/Source/Programs/Enterprise/Datasmith/DatasmithNavisworksExporter/DatasmithNavisworksPlugin/DatasmithNavisworksPlugin.csproj b/Engine/Source/Programs/Enterprise/Datasmith/DatasmithNavisworksExporter/DatasmithNavisworksPlugin/DatasmithNavisworksPlugin.csproj
--- a/Engine/Source/Programs/Enterprise/Datasmith/DatasmithNavisworksExporter/DatasmithNavisworksPlugin/DatasmithNavisworksPlugin.csproj
+++ b/Engine/Source/Programs/Enterprise/Datasmith/DatasmithNavisworksExporter/DatasmithNavisworksPlugin/DatasmithNavisworksPlugin.csproj
@@ -6,7 +6,10 @@
       None
     </ResolveAssemblyWarnOrErrorOnTargetArchitectureMismatch>
     <ProjectGuid>{D16F2120-CCE4-4FCC-92C8-86124AB1FA25}</ProjectGuid>
+    <TargetFrameworkProfile />
   </PropertyGroup>
+  <Import Project="..\Navisworks2023.props" Condition="'$(Configuration)|$(Platform)'=='Debug2023|AnyCPU'" />
+  <Import Project="..\Navisworks2023.props" Condition="'$(Configuration)|$(Platform)'=='Release2023|AnyCPU'" />
   <Import Project="..\Navisworks2022.props" Condition="'$(Configuration)|$(Platform)'=='Debug2022|AnyCPU'" />
   <Import Project="..\Navisworks2022.props" Condition="'$(Configuration)|$(Platform)'=='Release2022|AnyCPU'" />
   <Import Project="..\Navisworks2021.props" Condition="'$(Configuration)|$(Platform)'=='Debug2021|AnyCPU'" />
@@ -22,9 +25,33 @@
     <AppDesignerFolder>Properties</AppDesignerFolder>
     <RootNamespace>DatasmithNavisworks</RootNamespace>
     <AssemblyName>DatasmithNavisworksPlugin$(Navisworks_Version)</AssemblyName>
-    <TargetFrameworkVersion>v4.7</TargetFrameworkVersion>
+    <TargetFrameworkVersion>v4.8</TargetFrameworkVersion>
     <FileAlignment>512</FileAlignment>
   </PropertyGroup>
+  <PropertyGroup Condition=" '$(Configuration)|$(Platform)' == 'Debug2023|AnyCPU' ">
+    <DebugSymbols>true</DebugSymbols>
+    <DebugType>full</DebugType>
+    <Optimize>false</Optimize>
+    <OutputPath>bin\Debug2023\</OutputPath>
+    <DefineConstants>DEBUG;TRACE</DefineConstants>
+    <ErrorReport>prompt</ErrorReport>
+    <WarningLevel>4</WarningLevel>
+    <StartAction>Program</StartAction>
+    <StartProgram>$(ProgramW6432)\Autodesk\Navisworks Manage $(Navisworks_Version)\Roamer.exe</StartProgram>
+    <EnableUnmanagedDebugging>true</EnableUnmanagedDebugging>
+    <AllowUnsafeBlocks>true</AllowUnsafeBlocks>
+  </PropertyGroup>
+  <PropertyGroup Condition=" '$(Configuration)|$(Platform)' == 'Release2023|AnyCPU' ">
+    <DebugType>pdbonly</DebugType>
+    <Optimize>true</Optimize>
+    <OutputPath>bin\Release2023\</OutputPath>
+    <DefineConstants>TRACE</DefineConstants>
+    <ErrorReport>prompt</ErrorReport>
+    <WarningLevel>4</WarningLevel>
+    <StartAction>Program</StartAction>
+    <StartProgram>$(ProgramW6432)\Autodesk\Navisworks Manage $(Navisworks_Version)\Roamer.exe</StartProgram>
+    <AllowUnsafeBlocks>true</AllowUnsafeBlocks>
+  </PropertyGroup>
   <PropertyGroup Condition=" '$(Configuration)|$(Platform)' == 'Debug2022|AnyCPU' ">
     <DebugSymbols>true</DebugSymbols>
     <DebugType>full</DebugType>
diff --git a/Engine/Source/Programs/Enterprise/Datasmith/DatasmithNavisworksExporter/DatasmithNavisworksPlugin/PackageContents.xml b/Engine/Source/Programs/Enterprise/Datasmith/DatasmithNavisworksExporter/DatasmithNavisworksPlugin/PackageContents.xml
--- a/Engine/Source/Programs/Enterprise/Datasmith/DatasmithNavisworksExporter/DatasmithNavisworksPlugin/PackageContents.xml
+++ b/Engine/Source/Programs/Enterprise/Datasmith/DatasmithNavisworksExporter/DatasmithNavisworksPlugin/PackageContents.xml
@@ -17,4 +17,8 @@
 	    <RuntimeRequirements OS="Win64" Platform="NAVMAN|NAVSIM" SeriesMin="Nw19" SeriesMax="Nw19" />
         <ComponentEntry AppName="Datasmith Exporter for Autodesk® Navisworks® 2022" AppType="ManagedPlugin" ModuleName="./DatasmithNavisworksPlugin2022.dll" />
     </Components>
+    <Components Description="2023">
+	    <RuntimeRequirements OS="Win64" Platform="NAVMAN|NAVSIM" SeriesMin="Nw20" SeriesMax="Nw20" />
+        <ComponentEntry AppName="Datasmith Exporter for Autodesk® Navisworks® 2023" AppType="ManagedPlugin" ModuleName="./DatasmithNavisworksPlugin2023.dll" />
+    </Components>
 </ApplicationPackage>
diff --git a/Engine/Source/Programs/Enterprise/Datasmith/DatasmithNavisworksExporter/UBTModules/DatasmithNavisworks2023.Build.cs b/Engine/Source/Programs/Enterprise/Datasmith/DatasmithNavisworksExporter/UBTModules/DatasmithNavisworks2023.Build.cs
new file mode 100644
--- /dev/null
+++ b/Engine/Source/Programs/Enterprise/Datasmith/DatasmithNavisworksExporter/UBTModules/DatasmithNavisworks2023.Build.cs
@@ -0,0 +1,20 @@
+// Copyright Epic Games, Inc. All Rights Reserved.
+
+using System.IO;
+
+namespace UnrealBuildTool.Rules
+{
+	[SupportedPlatforms("Win64")]
+	public class DatasmithNavisworks2023 : DatasmithNavisworksBase
+	{
+		public DatasmithNavisworks2023(ReadOnlyTargetRules Target)
+			: base(Target)
+		{
+		}
+
+		public override string GetNavisworksVersion()
+		{
+			return "2023";
+		}
+	}
+}
diff --git a/Engine/Source/Programs/Enterprise/Datasmith/DatasmithNavisworksExporter/UBTModules/DatasmithNavisworks2023.Target.cs b/Engine/Source/Programs/Enterprise/Datasmith/DatasmithNavisworksExporter/UBTModules/DatasmithNavisworks2023.Target.cs
new file mode 100644
--- /dev/null
+++ b/Engine/Source/Programs/Enterprise/Datasmith/DatasmithNavisworksExporter/UBTModules/DatasmithNavisworks2023.Target.cs
@@ -0,0 +1,14 @@
+// Copyright Epic Games, Inc. All Rights Reserved.
+
+using UnrealBuildTool;
+using System.IO;
+
+public class DatasmithNavisworks2023Target : DatasmithNavisworksBaseTarget
+{
+	public DatasmithNavisworks2023Target(TargetInfo Target)
+		: base(Target)
+	{
+	}
+
+	public override string GetVersion() { return "2023"; }
+}
