Auto Sync not working in Revit

diff --git a/Engine/Source/Programs/Enterprise/Datasmith/DatasmithRevitExporter/Private/DatasmithRevitDirectLink.cs b/Engine/Source/Programs/Enterprise/Datasmith/DatasmithRevitExporter/Private/DatasmithRevitDirectLink.cs
--- a/Engine/Source/Programs/Enterprise/Datasmith/DatasmithRevitExporter/Private/DatasmithRevitDirectLink.cs
+++ b/Engine/Source/Programs/Enterprise/Datasmith/DatasmithRevitExporter/Private/DatasmithRevitDirectLink.cs
@@ -335,6 +335,11 @@ namespace DatasmithRevitExporter
 				#error This version of Revit is not supported yet.
 #endif
 
+				if (UIApp == null)
+				{
+					UIApp = new UIApplication(SyncView.Document.Application);
+				}
+
 				RevitCommandId CmdId = RevitCommandId.LookupCommandId(CmdGUID.ToLower());
 				if (CmdId != null)
 				{
