Revit - Revit does not close its connection when closing the RVT file

diff --git a/Engine/Source/Programs/Enterprise/Datasmith/DatasmithRevitExporter/Private/DatasmithRevitDirectLink.cs b/Engine/Source/Programs/Enterprise/Datasmith/DatasmithRevitExporter/Private/DatasmithRevitDirectLink.cs
--- a/Engine/Source/Programs/Enterprise/Datasmith/DatasmithRevitExporter/Private/DatasmithRevitDirectLink.cs
+++ b/Engine/Source/Programs/Enterprise/Datasmith/DatasmithRevitExporter/Private/DatasmithRevitDirectLink.cs
@@ -370,6 +370,7 @@ namespace DatasmithRevitExporter
 			InApp.DocumentChanged -= DocumentChangedHandler;
 			DocumentChangedHandler = null;
 
+			DatasmithDirectLink?.CloseCurrentSource();
 			DatasmithDirectLink = null;
 			DatasmithScene = null;
 			RootCache = null;
