[Revit Exporter] Implement export of decals

Add support for curved decals (decals on curved walls)

diff --git a/Engine/Source/Programs/Enterprise/Datasmith/DatasmithRevitExporter/Private/DatasmithRevitDocumentData.cs b/Engine/Source/Programs/Enterprise/Datasmith/DatasmithRevitExporter/Private/DatasmithRevitDocumentData.cs
--- a/Engine/Source/Programs/Enterprise/Datasmith/DatasmithRevitExporter/Private/DatasmithRevitDocumentData.cs
+++ b/Engine/Source/Programs/Enterprise/Datasmith/DatasmithRevitExporter/Private/DatasmithRevitDocumentData.cs
@@ -2113,6 +2113,11 @@ namespace DatasmithRevitExporter
 				XYZ DecalDimensions = null;
 				FUtils.GetDecalSpatialParams(DecalElement.CurrentElement, ref DecalTransform, ref DecalDimensions);
 
+				if (DecalTransform == null || DecalDimensions == null)
+				{
+					continue;
+				}
+
 				FDatasmithFacadeActorDecal DecalActor = DecalElement.ElementActor as FDatasmithFacadeActorDecal;
 				DecalActor.SetDimensions(DecalDimensions.Z, DecalDimensions.X, DecalDimensions.Y);
 				DecalActor.SetDecalMaterialPathName(DatasmithMaterial.GetName());
diff --git a/Engine/Source/Programs/Enterprise/Datasmith/DatasmithRevitExporter/Private/DatasmithRevitUtils.cs b/Engine/Source/Programs/Enterprise/Datasmith/DatasmithRevitExporter/Private/DatasmithRevitUtils.cs
--- a/Engine/Source/Programs/Enterprise/Datasmith/DatasmithRevitExporter/Private/DatasmithRevitUtils.cs
+++ b/Engine/Source/Programs/Enterprise/Datasmith/DatasmithRevitExporter/Private/DatasmithRevitUtils.cs
@@ -119,6 +119,16 @@ namespace DatasmithRevitExporter
 				{
 					DecalQuad.Add(QuadLine);
 				}
+				else if (GeomObj is Arc QuadArc)
+				{
+					try
+					{
+						XYZ StartPoint = QuadArc.GetEndPoint(0);
+						XYZ EndPoint = QuadArc.GetEndPoint(1);
+						DecalQuad.Add(Line.CreateBound(StartPoint, EndPoint));
+					}
+					catch {}
+				}
 			}
 
 			if (DecalQuad.Count != 4)
