[SketchUp] Fixed texture export on Mac

- avoided call to CopyFile for texture files, was failing because it was copying file onto itself

diff --git a/Engine/Source/Developer/Datasmith/DatasmithExporter/Private/DatasmithSceneExporter.cpp b/Engine/Source/Developer/Datasmith/DatasmithExporter/Private/DatasmithSceneExporter.cpp
--- a/Engine/Source/Developer/Datasmith/DatasmithExporter/Private/DatasmithSceneExporter.cpp
+++ b/Engine/Source/Developer/Datasmith/DatasmithExporter/Private/DatasmithSceneExporter.cpp
@@ -95,7 +95,10 @@ void FDatasmithSceneExporterImpl::UpdateTextureElements( TSharedRef< IDatasmithS
 			NewFilename = FPaths::Combine(AssetsOutputPath, UniqueFileName + FileExtension);
 
 			// Copy image file to new location
-			PlatformFile.CopyFile(*NewFilename, *TextureFileName);
+			if (!FPaths::IsSamePath(*NewFilename, *TextureFileName))
+			{
+				PlatformFile.CopyFile(*NewFilename, *TextureFileName);
+			}
 		}
 
 		// Update texture element
diff --git a/Engine/Source/Programs/Enterprise/Datasmith/DatasmithSketchUpRubyExporter/Private/DatasmithSketchUpExporter.cpp b/Engine/Source/Programs/Enterprise/Datasmith/DatasmithSketchUpRubyExporter/Private/DatasmithSketchUpExporter.cpp
--- a/Engine/Source/Programs/Enterprise/Datasmith/DatasmithSketchUpRubyExporter/Private/DatasmithSketchUpExporter.cpp
+++ b/Engine/Source/Programs/Enterprise/Datasmith/DatasmithSketchUpRubyExporter/Private/DatasmithSketchUpExporter.cpp
@@ -21,6 +21,8 @@
 #include "IDatasmithExporterUIModule.h"
 #include "IDirectLinkUI.h"
 
+#include "DatasmithExportOptions.h"
+
 #include "DatasmithSketchUpSDKBegins.h"
 #include <SketchUpAPI/sketchup.h>
 // SketchUp prior to 2020.2 doesn't have api to convert SU entities between Ruby and C 
@@ -634,6 +636,7 @@ VALUE DatasmithSketchUpDirectLinkExporter_export_current_datasmith_scene(VALUE s
 	// Done converting args
 
 	ptr->ExportCurrentDatasmithScene();
+
 	return Qtrue;
 }
 
@@ -812,6 +815,10 @@ VALUE on_load(VALUE self, VALUE enable_ui, VALUE engine_path) {
 	// Converting args
 	Check_Type(engine_path, T_STRING);
 
+	// We always write textures during export and write them directly where they need to be
+	// This options prevents extra copying of texture files (and copying a file was causing empty file on Mac)
+	FDatasmithExportOptions::PathTexturesMode = EDSResizedTexturesPath::OriginalFolder;
+
 	bool bEnableUI = RTEST(enable_ui);
 	FString EnginePathUnreal = RubyStringToUnreal(engine_path);
 	// Done converting args
