[3ds Max] DirectLink: Fixed sync for bakeable texmaps; ForestColor texmap approximation

diff --git a/Engine/Source/Programs/Enterprise/Datasmith/DatasmithMaxExporter/Private/DatasmithMaxDirectLinkMaterial.cpp b/Engine/Source/Programs/Enterprise/Datasmith/DatasmithMaxExporter/Private/DatasmithMaxDirectLinkMaterial.cpp
--- a/Engine/Source/Programs/Enterprise/Datasmith/DatasmithMaxExporter/Private/DatasmithMaxDirectLinkMaterial.cpp
+++ b/Engine/Source/Programs/Enterprise/Datasmith/DatasmithMaxExporter/Private/DatasmithMaxDirectLinkMaterial.cpp
@@ -135,13 +135,13 @@ void FMaterialsCollectionTracker::ConvertMaterial(Mtl* Material, TSharedRef<IDat
 		}
 	}
 
-	TArray<Texmap*> TexmapsUsedByMaterial;
+	TMap<Texmap*, TSet<TSharedPtr<IDatasmithTextureElement>>> TexmapsUsedByMaterial;
 	TSharedPtr<IDatasmithBaseMaterialElement> DatasmithMaterial;
 
 	if (FDatasmithMaxMaterialsToUEPbr* MaterialConverter = FDatasmithMaxMaterialsToUEPbrManager::GetMaterialConverter(Material))
 	{
 
-		TGuardValue<TArray<Texmap*>*> TexmapsConvertedGuard(MaterialConverter->TexmapsConverted, &TexmapsUsedByMaterial);
+		TGuardValue<TMap<Texmap*, TSet<TSharedPtr<IDatasmithTextureElement>>>*> TexmapsConvertedGuard(MaterialConverter->TexmapsConverted, &TexmapsUsedByMaterial);
 		 
 		MaterialConverter->Convert(DatasmithScene, DatasmithMaterial, Material, AssetsPath);
 
@@ -155,17 +155,20 @@ void FMaterialsCollectionTracker::ConvertMaterial(Mtl* Material, TSharedRef<IDat
 	}
 
 	// Tie texture used by an actual material to tracked material
-	for (Texmap* Texture : TexmapsUsedByMaterial)
+	for (TPair<Texmap*, TSet<TSharedPtr<IDatasmithTextureElement>>> TexmapAndDatasmithElements : TexmapsUsedByMaterial)
 	{
-
+		Texmap* Tex = TexmapAndDatasmithElements.Key;
 		for (FMaterialTracker* MaterialTracker : UsedMaterialToMaterialTracker[Material])
 		{
-			MaterialTracker->AddActualTexture(Texture);
-			UsedTextureToMaterialTracker.FindOrAdd(Texture).Add(MaterialTracker);
+			MaterialTracker->AddActualTexture(Tex);
+			UsedTextureToMaterialTracker.FindOrAdd(Tex).Add(MaterialTracker);
 		}
-	}
 
-	TexmapsConverted.Append(TexmapsUsedByMaterial);
+		// Some material converters create texture elements during material conversion(like bakeable) - record created texture elements
+		UsedTextureToDatasmithElement.FindOrAdd(Tex).Append(TexmapAndDatasmithElements.Value);
+
+		TexmapsConverted.Add(Tex);
+	}
 
 }
 
diff --git a/Engine/Source/Programs/Enterprise/Datasmith/DatasmithMaxExporter/Private/DatasmithMaxExporterWithDirectLink.cpp b/Engine/Source/Programs/Enterprise/Datasmith/DatasmithMaxExporter/Private/DatasmithMaxExporterWithDirectLink.cpp
--- a/Engine/Source/Programs/Enterprise/Datasmith/DatasmithMaxExporter/Private/DatasmithMaxExporterWithDirectLink.cpp
+++ b/Engine/Source/Programs/Enterprise/Datasmith/DatasmithMaxExporter/Private/DatasmithMaxExporterWithDirectLink.cpp
@@ -1602,7 +1602,8 @@ public:
 
 				TArray<TSharedPtr<IDatasmithTextureElement>> TextureElements;
 				FDatasmithMaxMatExport::GetXMLTexture(ExportedScene.GetDatasmithScene().ToSharedRef(), Texture, ExportedScene.GetSceneExporter().GetAssetsOutputPath(), &TextureElements);
-				MaterialsCollectionTracker.UsedTextureToDatasmithElement.Add(Texture, TextureElements);
+
+				MaterialsCollectionTracker.UsedTextureToDatasmithElement.FindOrAdd(Texture).Append(TextureElements);
 			}
 		}
 
diff --git a/Engine/Source/Programs/Enterprise/Datasmith/DatasmithMaxExporter/Private/DatasmithMaxLogger.cpp b/Engine/Source/Programs/Enterprise/Datasmith/DatasmithMaxExporter/Private/DatasmithMaxLogger.cpp
--- a/Engine/Source/Programs/Enterprise/Datasmith/DatasmithMaxExporter/Private/DatasmithMaxLogger.cpp
+++ b/Engine/Source/Programs/Enterprise/Datasmith/DatasmithMaxExporter/Private/DatasmithMaxLogger.cpp
@@ -2,6 +2,8 @@
 
 #include "DatasmithMaxLogger.h"
 
+#include "DatasmithMaxDirectLink.h"
+
 #include "DatasmithMaxExporterDefines.h"
 
 #include "HAL/UnrealMemory.h"
@@ -75,6 +77,7 @@ void DatasmithMaxLogger::AddPartialSupportedMap(Texmap* Map)
 
 void DatasmithMaxLogger::AddUnsupportedMap(Texmap* Map)
 {
+
 	for (int i = 0; i < UnsupportedMaps.Num(); i++)
 	{
 		if (Map->ClassID() == UnsupportedMaps[i]->ClassID())
@@ -83,6 +86,10 @@ void DatasmithMaxLogger::AddUnsupportedMap(Texmap* Map)
 		}
 	}
 
+	MSTR Classname;
+	Map->GetClassName(Classname);
+	DatasmithMaxDirectLink::LogWarning(FString::Printf(TEXT("Unsupported texmap \"%s\" of type %s (0x%08x-0x%08x)"), Map->GetName().ToBSTR(), Classname.ToBSTR(), Map->ClassID().PartA(), Map->ClassID().PartB()));
+
 	UnsupportedMaps.Add(Map);
 }
 
diff --git a/Engine/Source/Programs/Enterprise/Datasmith/DatasmithMaxExporter/Private/MaxMaterialsToUEPbr/DatasmithMaxMaterialsToUEPbr.cpp b/Engine/Source/Programs/Enterprise/Datasmith/DatasmithMaxExporter/Private/MaxMaterialsToUEPbr/DatasmithMaxMaterialsToUEPbr.cpp
--- a/Engine/Source/Programs/Enterprise/Datasmith/DatasmithMaxExporter/Private/MaxMaterialsToUEPbr/DatasmithMaxMaterialsToUEPbr.cpp
+++ b/Engine/Source/Programs/Enterprise/Datasmith/DatasmithMaxExporter/Private/MaxMaterialsToUEPbr/DatasmithMaxMaterialsToUEPbr.cpp
@@ -189,7 +189,7 @@ IDatasmithMaterialExpression* FDatasmithMaxMaterialsToUEPbr::ConvertTexmap( cons
 
 	if (TexmapsConverted)
 	{
-		TexmapsConverted->Add(MapParameter.Map);
+		TexmapsConverted->FindOrAdd(MapParameter.Map);
 	}
 
 	// Check if the texmap has a texture ouput
diff --git a/Engine/Source/Programs/Enterprise/Datasmith/DatasmithMaxExporter/Private/MaxMaterialsToUEPbr/DatasmithMaxMaterialsToUEPbr.h b/Engine/Source/Programs/Enterprise/Datasmith/DatasmithMaxExporter/Private/MaxMaterialsToUEPbr/DatasmithMaxMaterialsToUEPbr.h
--- a/Engine/Source/Programs/Enterprise/Datasmith/DatasmithMaxExporter/Private/MaxMaterialsToUEPbr/DatasmithMaxMaterialsToUEPbr.h
+++ b/Engine/Source/Programs/Enterprise/Datasmith/DatasmithMaxExporter/Private/MaxMaterialsToUEPbr/DatasmithMaxMaterialsToUEPbr.h
@@ -15,6 +15,7 @@ class IDatasmithMaterialExpressionTexture;
 class IDatasmithMaxTexmapToUEPbr;
 class IDatasmithScene;
 class IDatasmithUEPbrMaterialElement;
+class IDatasmithTextureElement;
 class Mtl;
 class Texmap;
 
@@ -59,7 +60,10 @@ public:
 
 	} ConvertState;
 
-	TArray<Texmap*>* TexmapsConverted = nullptr;
+	// Used to record texmaps that were asked for conversion whil ematerial conversion is performed
+	// recordking also datasmith elements created for those maps.
+	// Baked texmaps currently create datasmith elements during material convertion
+	TMap<Texmap*, TSet<TSharedPtr<IDatasmithTextureElement>>>* TexmapsConverted = nullptr;
 
 protected:
 	TIndirectArray< IDatasmithMaxTexmapToUEPbr > TexmapConverters;
diff --git a/Engine/Source/Programs/Enterprise/Datasmith/DatasmithMaxExporter/Private/MaxMaterialsToUEPbr/DatasmithMaxTexmapToUEPbr.cpp b/Engine/Source/Programs/Enterprise/Datasmith/DatasmithMaxExporter/Private/MaxMaterialsToUEPbr/DatasmithMaxTexmapToUEPbr.cpp
--- a/Engine/Source/Programs/Enterprise/Datasmith/DatasmithMaxExporter/Private/MaxMaterialsToUEPbr/DatasmithMaxTexmapToUEPbr.cpp
+++ b/Engine/Source/Programs/Enterprise/Datasmith/DatasmithMaxExporter/Private/MaxMaterialsToUEPbr/DatasmithMaxTexmapToUEPbr.cpp
@@ -1703,6 +1703,7 @@ IDatasmithMaterialExpression* FDatasmithMaxBakeableToUEPbr::Convert( FDatasmithM
 
 	if ( BakedTextureElement )
 	{
+		MaxMaterialToUEPbr->TexmapsConverted->FindOrAdd(InTexmap).Add(BakedTextureElement);
 		IDatasmithMaterialExpression* MaterialExpression = MaxMaterialToUEPbr->ConvertState.MaterialElement->AddMaterialExpression( EDatasmithMaterialExpressionType::Texture );
 
 		if ( MaterialExpression )
@@ -1740,14 +1741,22 @@ bool FDatasmithMaxPassthroughToUEPbr::IsSupported( const FDatasmithMaxMaterialsT
 		{
 			return true;
 		}
+		else if ( InTexmap->ClassID() == FORESTCOLORCLASS )
+		{
+			return true;
+		}
 	}
 
 	return false;
 }
 
-const TCHAR* FDatasmithMaxPassthroughToUEPbr::GetColorParameterName() const
+const TCHAR* FDatasmithMaxPassthroughToUEPbr::GetColorParameterName(Texmap* InTexmap) const
 {
-	return TEXT("Color1");
+	if ( InTexmap->ClassID() == FORESTCOLORCLASS )
+	{
+		return TEXT("colorbase");
+	}
+	return TEXT("color1");
 }
 
 IDatasmithMaterialExpression* FDatasmithMaxPassthroughToUEPbr::Convert( FDatasmithMaxMaterialsToUEPbr* MaxMaterialToUEPbr, Texmap* InTexmap )
@@ -1759,7 +1768,7 @@ IDatasmithMaterialExpression* FDatasmithMaxPassthroughToUEPbr::Convert( FDatasmi
 	MapParameter.Weight = 1.f;
 	MapParameter.bEnabled = ( InTexmap->SubTexmapOn(0) != 0 );
 
-	FLinearColor PassthroughColor = GetColorParameter( InTexmap, GetColorParameterName() );
+	FLinearColor PassthroughColor = GetColorParameter( InTexmap, GetColorParameterName(InTexmap) );
 	IDatasmithMaterialExpression* Map1 = FDatasmithMaxTexmapToUEPbrUtils::MapOrValue( MaxMaterialToUEPbr, MapParameter, TEXT("Submap"), PassthroughColor, TOptional< float >() );
 
 	return Map1;
@@ -1801,7 +1810,7 @@ FLinearColor FDatasmithMaxPassthroughToUEPbr::GetColorParameter( Texmap* InTexma
 	return FLinearColor::Black;
 }
 
-const TCHAR* FDatasmithMaxCellularToUEPbr::GetColorParameterName() const
+const TCHAR* FDatasmithMaxCellularToUEPbr::GetColorParameterName(Texmap* InTexmap) const
 {
 	return TEXT("cellcolor");
 }
@@ -1820,7 +1829,7 @@ bool FDatasmithMaxCellularToUEPbr::IsSupported( const FDatasmithMaxMaterialsToUE
 	return false;
 }
 
-const TCHAR* FDatasmithMaxThirdPartyMultiTexmapToUEPbr::GetColorParameterName() const
+const TCHAR* FDatasmithMaxThirdPartyMultiTexmapToUEPbr::GetColorParameterName(Texmap* InTexmap) const
 {
 	return TEXT("sub_color");
 }
diff --git a/Engine/Source/Programs/Enterprise/Datasmith/DatasmithMaxExporter/Private/MaxMaterialsToUEPbr/DatasmithMaxTexmapToUEPbr.h b/Engine/Source/Programs/Enterprise/Datasmith/DatasmithMaxExporter/Private/MaxMaterialsToUEPbr/DatasmithMaxTexmapToUEPbr.h
--- a/Engine/Source/Programs/Enterprise/Datasmith/DatasmithMaxExporter/Private/MaxMaterialsToUEPbr/DatasmithMaxTexmapToUEPbr.h
+++ b/Engine/Source/Programs/Enterprise/Datasmith/DatasmithMaxExporter/Private/MaxMaterialsToUEPbr/DatasmithMaxTexmapToUEPbr.h
@@ -116,7 +116,7 @@ class FDatasmithMaxPassthroughToUEPbr : public IDatasmithMaxTexmapToUEPbr
 public:
 	virtual bool IsSupported( const FDatasmithMaxMaterialsToUEPbr* MaxMaterialToUEPbr, Texmap* InTexmap ) const override;
 
-	virtual const TCHAR* GetColorParameterName() const;
+	virtual const TCHAR* GetColorParameterName(Texmap* InTexmap) const;
 
 	virtual IDatasmithMaterialExpression* Convert( FDatasmithMaxMaterialsToUEPbr* MaxMaterialToUEPbr, Texmap* InTexmap ) override;
 	FLinearColor GetColorParameter( Texmap* InTexmap, const TCHAR* ParameterName );
@@ -124,12 +124,12 @@ public:
 
 class FDatasmithMaxCellularToUEPbr : public FDatasmithMaxPassthroughToUEPbr
 {
-	virtual const TCHAR* GetColorParameterName() const override;
+	virtual const TCHAR* GetColorParameterName(Texmap* InTexmap) const override;
 	virtual bool IsSupported( const FDatasmithMaxMaterialsToUEPbr* MaxMaterialToUEPbr, Texmap* InTexmap ) const override;
 };
 
 class FDatasmithMaxThirdPartyMultiTexmapToUEPbr : public FDatasmithMaxPassthroughToUEPbr
 {
-	virtual const TCHAR* GetColorParameterName() const override;
+	virtual const TCHAR* GetColorParameterName(Texmap* InTexmap) const override;
 	virtual bool IsSupported( const FDatasmithMaxMaterialsToUEPbr* MaxMaterialToUEPbr, Texmap* InTexmap ) const override;
 };
diff --git a/Engine/Source/Programs/Enterprise/Datasmith/DatasmithMaxExporter/Public/DatasmithMaxDirectLink.h b/Engine/Source/Programs/Enterprise/Datasmith/DatasmithMaxExporter/Public/DatasmithMaxDirectLink.h
--- a/Engine/Source/Programs/Enterprise/Datasmith/DatasmithMaxExporter/Public/DatasmithMaxDirectLink.h
+++ b/Engine/Source/Programs/Enterprise/Datasmith/DatasmithMaxExporter/Public/DatasmithMaxDirectLink.h
@@ -618,7 +618,7 @@ public:
 	TMap<Mtl*, TSharedPtr<IDatasmithBaseMaterialElement>> UsedMaterialToDatasmithMaterial;
 
 	TMap<Texmap*, TSet<FMaterialTracker*>> UsedTextureToMaterialTracker; // Materials uses by nodes keep set of assigned textures they are used for
-	TMap<Texmap*, TArray<TSharedPtr<IDatasmithTextureElement>>> UsedTextureToDatasmithElement; // Keep track of Datasmith Element created for texmap to simplify update/removal(no need to search DatasmithScene)
+	TMap<Texmap*, TSet<TSharedPtr<IDatasmithTextureElement>>> UsedTextureToDatasmithElement; // Keep track of Datasmith Element created for texmap to simplify update/removal(no need to search DatasmithScene)
 	// note: each texmap can create multiple texture elements
 
 	FSceneUpdateStats& Stats;
