Datasmith: fix SmoothingMask array size

diff --git a/Engine/Source/Programs/Enterprise/Datasmith/DatasmithARCHICADExporter/Private/DatasmithHashTools.cpp b/Engine/Source/Programs/Enterprise/Datasmith/DatasmithARCHICADExporter/Private/DatasmithHashTools.cpp
--- a/Engine/Source/Programs/Enterprise/Datasmith/DatasmithARCHICADExporter/Private/DatasmithHashTools.cpp
+++ b/Engine/Source/Programs/Enterprise/Datasmith/DatasmithARCHICADExporter/Private/DatasmithHashTools.cpp
@@ -190,11 +190,11 @@ void FDatasmithHashTools::ComputeDatasmithMeshHash(const FDatasmithMesh& Mesh)
 		TUpdate(Vertex2);
 		TUpdate(Vertex3);
 		TUpdate(MaterialId);
+		TUpdate(Mesh.GetFaceSmoothingMask(IdxFace));
 
 		for (int32 IdxComponent = 0; IdxComponent < 3; IdxComponent++)
 		{
 			HashFloatVector(FVector(Mesh.GetNormal(IdxFace * 3 + IdxComponent)));
-			TUpdate(Mesh.GetFaceSmoothingMask(IdxFace * 3 + IdxComponent));
 		}
 
 		for (int32 IdxChannel = 0; IdxChannel < UVChannelCount; ++IdxChannel)
diff --git a/Engine/Source/Runtime/Datasmith/DatasmithCore/Private/DatasmithMesh.cpp b/Engine/Source/Runtime/Datasmith/DatasmithCore/Private/DatasmithMesh.cpp
--- a/Engine/Source/Runtime/Datasmith/DatasmithCore/Private/DatasmithMesh.cpp
+++ b/Engine/Source/Runtime/Datasmith/DatasmithCore/Private/DatasmithMesh.cpp
@@ -116,8 +116,8 @@ void FDatasmithMesh::SetFacesCount(int32 NumFaces)
 	Impl->Normals.Empty(NumFaces * 3);
 	Impl->Normals.Init( FVector3f::ZeroVector, NumFaces * 3 );
 
-	Impl->FaceSmoothingMasks.Empty( NumFaces * 3 );
-	Impl->FaceSmoothingMasks.AddZeroed( NumFaces * 3 );
+	Impl->FaceSmoothingMasks.Empty( NumFaces );
+	Impl->FaceSmoothingMasks.AddZeroed( NumFaces );
 
 	for ( TArray< int32 >& UVIndices : Impl->UVIndices )
 	{
