Enable solidworks installer to install with solidworks 2022

diff --git a/Engine/Source/Programs/Enterprise/Datasmith/DatasmithSolidworksExporter/Private/Document.cs b/Engine/Source/Programs/Enterprise/Datasmith/DatasmithSolidworksExporter/Private/Document.cs
--- a/Engine/Source/Programs/Enterprise/Datasmith/DatasmithSolidworksExporter/Private/Document.cs
+++ b/Engine/Source/Programs/Enterprise/Datasmith/DatasmithSolidworksExporter/Private/Document.cs
@@ -39,7 +39,7 @@ namespace DatasmithSolidworks
 		{
 			DocId = InDocId;
 			SwDoc = InSwDoc;
-			DatasmithScene = new FDatasmithFacadeScene("StdMaterial", "Solidworks", "Solidworks", "2021");
+			DatasmithScene = new FDatasmithFacadeScene("StdMaterial", "Solidworks", "Solidworks", "");
 			DatasmithScene.SetName(InSwDoc.GetTitle());
 
 			DirectLinkPath = Path.Combine(Path.GetTempPath(), "sw_dl_" + Guid.NewGuid().ToString());
@@ -195,7 +195,7 @@ namespace DatasmithSolidworks
 				FDatasmithFacadeScene OldScene = DatasmithScene;
 				FDatasmithExporter OldExporter = Exporter;
 
-				DatasmithScene = new FDatasmithFacadeScene("StdMaterial", "Solidworks", "Solidworks", "2021");
+				DatasmithScene = new FDatasmithFacadeScene("StdMaterial", "Solidworks", "Solidworks", "");
 				DatasmithScene.SetName(CleanFileName);
 				DatasmithScene.SetOutputPath(OutDir);
 
diff --git a/Engine/Source/Programs/Enterprise/Datasmith/DatasmithSolidworksExporter/RegistryKeyExistsCustomAction/CustomAction.cs b/Engine/Source/Programs/Enterprise/Datasmith/DatasmithSolidworksExporter/RegistryKeyExistsCustomAction/CustomAction.cs
--- a/Engine/Source/Programs/Enterprise/Datasmith/DatasmithSolidworksExporter/RegistryKeyExistsCustomAction/CustomAction.cs
+++ b/Engine/Source/Programs/Enterprise/Datasmith/DatasmithSolidworksExporter/RegistryKeyExistsCustomAction/CustomAction.cs
@@ -59,7 +59,10 @@ namespace RegistryKeyExistsCustomAction
 		{
 			var tools2020 = CheckSolidworksInstalledSub(session, "SOLIDWORKS 2020", "SOLIDWORKS2020INSTALLED");
 			var tools2021 = CheckSolidworksInstalledSub(session, "SOLIDWORKS 2021", "SOLIDWORKS2021INSTALLED");
-			if (!string.IsNullOrEmpty(tools2021))
+			var tools2022 = CheckSolidworksInstalledSub(session, "SOLIDWORKS 2022", "SOLIDWORKS2022INSTALLED");
+			if (!string.IsNullOrEmpty(tools2022))
+				session["SOLIDWORKSTOOLSPATH"] = tools2022;
+			else if (!string.IsNullOrEmpty(tools2021))
 				session["SOLIDWORKSTOOLSPATH"] = tools2021;
 			else if (!string.IsNullOrEmpty(tools2020))
 				session["SOLIDWORKSTOOLSPATH"] = tools2020;
