[3ds Max] DirectLink: Exporting bitmap texture using a missing image file produces a warning

diff --git a/Engine/Source/Programs/Enterprise/Datasmith/DatasmithMaxExporter/Private/DatasmithMaxScanline.cpp b/Engine/Source/Programs/Enterprise/Datasmith/DatasmithMaxExporter/Private/DatasmithMaxScanline.cpp
--- a/Engine/Source/Programs/Enterprise/Datasmith/DatasmithMaxExporter/Private/DatasmithMaxScanline.cpp
+++ b/Engine/Source/Programs/Enterprise/Datasmith/DatasmithMaxExporter/Private/DatasmithMaxScanline.cpp
@@ -1,6 +1,8 @@
 // Copyright Epic Games, Inc. All Rights Reserved.
 #include "DatasmithMaxWriter.h"
 
+#include "DatasmithMaxDirectLink.h"
+
 #include "DatasmithSceneFactory.h"
 #include "DatasmithSceneExporter.h"
 #include "DatasmithMaxSceneParser.h"
@@ -484,6 +486,11 @@ void FDatasmithMaxMatWriter::GetRegularTexmap(TSharedRef< IDatasmithScene > Data
 			TextureElement->SetRGBCurve(Gamma);
 		}
 	}
+	
+	if (!FPaths::FileExists(*Path))
+	{
+		DatasmithMaxDirectLink::LogWarning(FString::Printf(TEXT("Bitmap texture '%s' has missing external file '%s'"), InBitmapTex->GetName().data(), *Path));
+	}
 
 	TextureElement->SetFile(*Path);
 	DatasmithScene->AddTexture(TextureElement);
