[3ds Max] DirectLink: material names fixes

- Max materials are not changing names anymore by the plugin
- Fixed export/sync for Multisubobj submaterials with same name
- Fixed crash for Blend material with null input material

diff --git a/Engine/Source/Programs/Enterprise/Datasmith/DatasmithMaxExporter/Private/DatasmithMaxDirectLinkMaterial.cpp b/Engine/Source/Programs/Enterprise/Datasmith/DatasmithMaxExporter/Private/DatasmithMaxDirectLinkMaterial.cpp
--- a/Engine/Source/Programs/Enterprise/Datasmith/DatasmithMaxExporter/Private/DatasmithMaxDirectLinkMaterial.cpp
+++ b/Engine/Source/Programs/Enterprise/Datasmith/DatasmithMaxExporter/Private/DatasmithMaxDirectLinkMaterial.cpp
@@ -7,6 +7,7 @@
 #include "DatasmithMaxHelper.h"
 #include "DatasmithMaxWriter.h"
 #include "DatasmithMaxClassIDs.h"
+#include "DatasmithMaxSceneExporter.h"
 #include "MaxMaterialsToUEPbr/DatasmithMaxMaterialsToUEPbr.h"
 
 
@@ -25,6 +26,159 @@ MAX_INCLUDES_END
 namespace DatasmithMaxDirectLink
 {
 
+void FMaterialsCollectionTracker::AddActualMaterial(FMaterialTracker& MaterialTracker, Mtl* Material)
+{
+	// Record relationships between tracked material and actual materials used on geometry(e.g. submaterials of a tracked multisubobj material)
+	UsedMaterialToMaterialTracker.FindOrAdd(Material).Add(&MaterialTracker);
+	MaterialTracker.AddActualMaterial(Material);
+}
+
+const TCHAR* FMaterialsCollectionTracker::GetMaterialName(Mtl* Material)
+{
+	if (FString* NamePtr = UsedMaterialToDatasmithMaterialName.Find(Material))
+	{
+		return **NamePtr;
+	}
+	return *UsedMaterialToDatasmithMaterialName.Add(Material, 
+		MaterialNameProvider.GenerateUniqueName(FDatasmithUtils::SanitizeObjectName(Material->GetName().data())));
+}
+
+void FMaterialsCollectionTracker::AssignMeshMaterials(const TSharedPtr<IDatasmithMeshElement>& MeshElement, Mtl* Material, const TSet<uint16>& SupportedChannels)
+{
+	if (FDatasmithMaxMatHelper::GetMaterialClass(Material) == EDSMaterialType::XRefMat)
+	{
+		AssignMeshMaterials(MeshElement, FDatasmithMaxMatHelper::GetRenderedXRefMaterial(Material), SupportedChannels);
+		return;
+	}
+	if (Material != nullptr)
+	{
+		TArray<uint16> ChannelsSorted = SupportedChannels.Array();
+		ChannelsSorted.Sort();
+		for (uint16 Channel : ChannelsSorted)
+		{
+			//Max's channel UI is not zero-based, so we register an incremented ChannelID for better visual consistency after importing in Unreal.
+			uint16 DisplayedChannelID = Channel + 1;
+
+			// Multi mat
+			if (FDatasmithMaxMatHelper::GetMaterialClass(Material) == EDSMaterialType::MultiMat)
+			{
+				// Replicate the 3ds max behavior where material ids greater than the number of sub-materials wrap around and are mapped to the existing sub-materials
+				if ( Mtl* SubMaterial = Material->GetSubMtl(Channel % Material->NumSubMtls()) )
+				{
+					if (FDatasmithMaxMatHelper::GetMaterialClass(SubMaterial) == EDSMaterialType::TheaRandom)
+					{
+						MeshElement->SetMaterial(*FDatasmithMaxSceneExporter::GetRandomSubMaterial(SubMaterial, MeshElement->GetDimensions()), DisplayedChannelID);
+					}
+					else
+					{
+						MeshElement->SetMaterial(GetMaterialName(SubMaterial), DisplayedChannelID);
+					}
+				}
+			}
+			else if (FDatasmithMaxMatHelper::GetMaterialClass(Material) == EDSMaterialType::TheaRandom)
+			{
+				MeshElement->SetMaterial(*FDatasmithMaxSceneExporter::GetRandomSubMaterial(Material, MeshElement->GetDimensions()), DisplayedChannelID);
+			}
+			// Single material
+			else
+			{
+				MeshElement->SetMaterial(GetMaterialName(Material), DisplayedChannelID);
+			}
+		}
+	}
+}
+
+void FMaterialsCollectionTracker::AssignMeshActorMaterials(const TSharedPtr<IDatasmithMeshActorElement>& MeshActor, Mtl* Material, TSet<uint16>& SupportedChannels, const FVector3f& RandomSeed)
+{
+	if (FDatasmithMaxMatHelper::GetMaterialClass(Material) == EDSMaterialType::XRefMat)
+	{
+		AssignMeshActorMaterials(MeshActor, FDatasmithMaxMatHelper::GetRenderedXRefMaterial(Material), SupportedChannels, RandomSeed);
+		return;
+	}
+
+	if (Material != nullptr)
+	{
+		if (SupportedChannels.Num() <= 1)
+		{
+			if (FDatasmithMaxMatHelper::GetMaterialClass(Material) != EDSMaterialType::MultiMat)
+			{
+				if (FDatasmithMaxMatHelper::GetMaterialClass(Material) != EDSMaterialType::TheaRandom)
+				{
+					MeshActor->AddMaterialOverride(GetMaterialName(Material), -1);
+				}
+				else
+				{
+					MeshActor->AddMaterialOverride( *FDatasmithMaxSceneExporter::GetRandomSubMaterial(Material, RandomSeed), -1 );
+				}
+			}
+			else
+			{
+				int Mid = 0;
+
+				//Find the lowest supported material id.
+				SupportedChannels.Sort([](const uint16& A, const uint16& B) { return (A < B); });
+				for (const uint16 SupportedMid : SupportedChannels)
+				{
+					Mid = SupportedMid;
+				}
+
+				// Replicate the 3ds max behavior where material ids greater than the number of sub-materials wrap around and are mapped to the existing sub-materials
+				if (Mtl* SubMaterial = Material->GetSubMtl(Mid % Material->NumSubMtls()))
+				{
+					if (FDatasmithMaxMatHelper::GetMaterialClass(SubMaterial) != EDSMaterialType::TheaRandom)
+					{
+						MeshActor->AddMaterialOverride(GetMaterialName(SubMaterial), -1);
+					}
+					else
+					{
+						MeshActor->AddMaterialOverride( *FDatasmithMaxSceneExporter::GetRandomSubMaterial(SubMaterial, RandomSeed), -1 );
+					}
+				}
+			}
+		}
+		else
+		{
+			int ActualSubObj = 1;
+			SupportedChannels.Sort([](const uint16& A, const uint16& B) { return (A < B); });
+			for (const uint16 Mid : SupportedChannels)
+			{
+				if (FDatasmithMaxMatHelper::GetMaterialClass(Material) != EDSMaterialType::MultiMat)
+				{
+					if (FDatasmithMaxMatHelper::GetMaterialClass(Material) != EDSMaterialType::TheaRandom)
+					{
+						MeshActor->AddMaterialOverride(GetMaterialName(Material), ActualSubObj);
+					}
+					else
+					{
+						MeshActor->AddMaterialOverride( *FDatasmithMaxSceneExporter::GetRandomSubMaterial(Material, RandomSeed), ActualSubObj );
+					}
+				}
+				else
+				{
+					// Replicate the 3ds max behavior where material ids greater than the number of sub-materials wrap around and are mapped to the existing sub-materials
+					int32 MaterialIndex = Mid % Material->NumSubMtls();
+
+					Mtl* SubMaterial = Material->GetSubMtl(MaterialIndex);
+					if (SubMaterial != nullptr)
+					{
+						if (FDatasmithMaxMatHelper::GetMaterialClass(SubMaterial) != EDSMaterialType::TheaRandom)
+						{
+							//Material slots in Max are not zero-based, so we serialize our SlotID starting from 1 for better visual consistency.
+							MeshActor->AddMaterialOverride(GetMaterialName(SubMaterial), Mid + 1);
+						}
+						else
+						{
+							MeshActor->AddMaterialOverride( *FDatasmithMaxSceneExporter::GetRandomSubMaterial(SubMaterial, RandomSeed), MaterialIndex + 1 );
+						}
+					}
+				}
+				ActualSubObj++;
+			}
+		}
+	}
+}
+
+
 // Copied from
 // FDatasmithMaxSceneParser::MaterialEnum
 // FDatasmithMaxSceneParser::TexEnum
@@ -59,31 +213,10 @@ public:
 		{
 			if (bAddMaterial)
 			{
-				if (!MaterialsCollectionTracker.EncounteredMaterials.Contains(Material))
-				{
-					int DuplicateCount = 0;
-					FString ProposedName = Material->GetName().data();
-					// todo: fix this without changing max material name? Btw - this requires changing all material export functions for all types of materials(those functions tied to Mtl->GetName())
-					// todo: revert material names after export
-					MaterialsCollectionTracker.MaterialNames.Add(*ProposedName);
-
-					// Make unique material name
-					FDatasmithUtils::SanitizeNameInplace(ProposedName);
-					for (Mtl* OtherMaterial: MaterialsCollectionTracker.EncounteredMaterials)
-					{
-						if (ProposedName == FDatasmithUtils::SanitizeName(OtherMaterial->GetName().data()))
-						{
-							DuplicateCount++;
-							ProposedName = FDatasmithUtils::SanitizeName(Material->GetName().data()) + TEXT("_(") + FString::FromInt(DuplicateCount) + TEXT(")");
-						}
-					}
-					Material->SetName(*ProposedName);
-					MaterialsCollectionTracker.EncounteredMaterials.Add(Material);
-				}
-				MaterialTracker.AddActualMaterial(Material);
+				MaterialsCollectionTracker.AddActualMaterial(MaterialTracker, Material);
 			}
 
-			bool bAddRecursively = Material->ClassID() == THEARANDOMCLASS || Material->ClassID() == VRAYBLENDMATCLASS || Material->ClassID() == CORONALAYERMATCLASS;
+			bool bAddRecursively = Material->ClassID() == THEARANDOMCLASS || Material->ClassID() == VRAYBLENDMATCLASS || Material->ClassID() == CORONALAYERMATCLASS || Material->ClassID() == BLENDMATCLASS;
 			for (int i = 0; i < Material->NumSubMtls(); i++)
 			{
 				MaterialEnum(Material->GetSubMtl(i), bAddRecursively);
@@ -103,6 +236,8 @@ void FMaterialsCollectionTracker::Reset()
 
 	UsedMaterialToMaterialTracker.Reset();
 	UsedMaterialToDatasmithMaterial.Reset();
+	UsedMaterialToDatasmithMaterialName.Reset();
+	MaterialNameProvider.Clear();
 
 	UsedTextureToMaterialTracker.Reset();
 	UsedTextureToDatasmithElement.Reset();
@@ -112,46 +247,41 @@ void FMaterialsCollectionTracker::UpdateMaterial(FMaterialTracker* MaterialTrack
 {
 	RemoveConvertedMaterial(*MaterialTracker);
 	FMaterialEnum(*this, *MaterialTracker).MaterialEnum(MaterialTracker->Material, true);
-	for (Mtl* Material : MaterialTracker->GetActualMaterials())
+}
+
+void FMaterialsCollectionTracker::AddDatasmithMaterialForUsedMaterial(TSharedRef<IDatasmithScene> DatasmithScene, Mtl* Material, TSharedPtr<IDatasmithBaseMaterialElement> DatasmithMaterial)
+{
+	if (DatasmithMaterial)
 	{
-		UsedMaterialToMaterialTracker.FindOrAdd(Material).Add(MaterialTracker);
+		SCENE_UPDATE_STAT_INC(UpdateMaterials, Converted);
+		DatasmithScene->AddMaterial(DatasmithMaterial);
+		UsedMaterialToDatasmithMaterial.Add(Material, DatasmithMaterial);
+		SceneTracker.RemapConvertedMaterialUVChannels(Material, DatasmithMaterial);
 	}
 }
 
 void FMaterialsCollectionTracker::ConvertMaterial(Mtl* Material, TSharedRef<IDatasmithScene> DatasmithScene, const TCHAR* AssetsPath, TSet<Texmap*>& TexmapsConverted)
 {
-	SCENE_UPDATE_STAT_INC(UpdateMaterials, Total);
-	// todo: reworking material export - composite materials(like VRayBlend) export submaterials too
-	//   check below is needed to make sure that no material is exported twice.
-	//   This is how it was handled before bu needs to be handled without going through the list of all materials.
-	// Names should be unique prior to this
-	FString MaterialName(FDatasmithUtils::SanitizeObjectName(Material->GetName().data()));
-	for (int i = 0; i < DatasmithScene->GetMaterialsCount(); i++)
+	if (UsedMaterialToDatasmithMaterial.Contains(Material))
 	{
-		if (FString(DatasmithScene->GetMaterial(i)->GetName()) == MaterialName)
-		{
-			SCENE_UPDATE_STAT_INC(UpdateMaterials, SkippedAsAlreadyConverted);
-			return;
-		}
+		// Material might have been already converted - if it's present in UsedMaterialToDatasmithMaterial this means that it(or multisubobj it's part of) wasn't changed
+		// e.g. this happens when another multisubobj material is added with existing (and already converted) material as submaterial
+		return;
 	}
 
+	SCENE_UPDATE_STAT_INC(UpdateMaterials, Total);
+
 	TMap<Texmap*, TSet<TSharedPtr<IDatasmithTextureElement>>> TexmapsUsedByMaterial;
-	TSharedPtr<IDatasmithBaseMaterialElement> DatasmithMaterial;
+	FMaterialConversionContext MaterialConversionContext = {TexmapsUsedByMaterial, *this};
+	TGuardValue MaterialConversionContextGuard(FDatasmithMaxMaterialsToUEPbrManager::Context, &MaterialConversionContext);
 
+	TSharedPtr<IDatasmithBaseMaterialElement> DatasmithMaterial;
 	if (FDatasmithMaxMaterialsToUEPbr* MaterialConverter = FDatasmithMaxMaterialsToUEPbrManager::GetMaterialConverter(Material))
 	{
-
-		TGuardValue<TMap<Texmap*, TSet<TSharedPtr<IDatasmithTextureElement>>>*> TexmapsConvertedGuard(MaterialConverter->TexmapsConverted, &TexmapsUsedByMaterial);
 		 
 		MaterialConverter->Convert(DatasmithScene, DatasmithMaterial, Material, AssetsPath);
 
-		if (DatasmithMaterial)
-		{
-			SCENE_UPDATE_STAT_INC(UpdateMaterials, Converted);
-			DatasmithScene->AddMaterial(DatasmithMaterial);
-			UsedMaterialToDatasmithMaterial.Add(Material, DatasmithMaterial);
-			SceneTracker.RemapConvertedMaterialUVChannels(Material, DatasmithMaterial);
-		}
+		AddDatasmithMaterialForUsedMaterial(DatasmithScene, Material, DatasmithMaterial);
 	}
 
 	// Tie texture used by an actual material to tracked material
@@ -190,6 +320,11 @@ void FMaterialsCollectionTracker::RemoveConvertedMaterial(FMaterialTracker& Mate
 		{
 			UsedMaterialToMaterialTracker.Remove(Material);
 
+			if (FString Name; UsedMaterialToDatasmithMaterialName.RemoveAndCopyValue(Material, Name))
+			{
+				MaterialNameProvider.RemoveExistingName(Name);
+			}
+
 			TSharedPtr<IDatasmithBaseMaterialElement> DatasmithMaterial;
 			if(UsedMaterialToDatasmithMaterial.RemoveAndCopyValue(Material, DatasmithMaterial))
 			{
diff --git a/Engine/Source/Programs/Enterprise/Datasmith/DatasmithMaxExporter/Private/DatasmithMaxExporterWithDirectLink.cpp b/Engine/Source/Programs/Enterprise/Datasmith/DatasmithMaxExporter/Private/DatasmithMaxExporterWithDirectLink.cpp
--- a/Engine/Source/Programs/Enterprise/Datasmith/DatasmithMaxExporter/Private/DatasmithMaxExporterWithDirectLink.cpp
+++ b/Engine/Source/Programs/Enterprise/Datasmith/DatasmithMaxExporter/Private/DatasmithMaxExporterWithDirectLink.cpp
@@ -178,10 +178,10 @@ struct FInstances: FNoncopyable
 		return Converted.DatasmithMeshElement->GetName();
 	}
 
-	void AssignMaterialToStaticMesh(Mtl* InMaterial)
+	void AssignMaterialToStaticMesh(FMaterialsCollectionTracker& MaterialsCollectionTracker,  Mtl* InMaterial)
 	{
 		Material = InMaterial;
-		AssignMeshMaterials(Converted.DatasmithMeshElement, Material, Converted.SupportedChannels);
+		MaterialsCollectionTracker.AssignMeshMaterials(Converted.DatasmithMeshElement, Material, Converted.SupportedChannels);
 	}
 };
 
@@ -1571,16 +1571,24 @@ public:
 				return FormatStatsUpdateMaterials();
 			};
 			PROGRESS_STAGE_COUNTER(ActualMaterialToUpdate.Num());
+
 			for (Mtl* ActualMaterial: ActualMaterialToUpdate)
 			{
 				ProgressCounter.Next();
 				if(GetCancel())
 				{
-					return false;	
+					break;	
 				}
 
+
 				MaterialsCollectionTracker.ConvertMaterial(ActualMaterial, ExportedScene.GetDatasmithScene().ToSharedRef(), ExportedScene.GetSceneExporter().GetAssetsOutputPath(), ActualTexmapsToUpdate);
 			}
+
+			if(GetCancel())
+			{
+				return false;	
+			}
+
 		}
 
 		{
@@ -2386,7 +2394,7 @@ public:
 			static_cast<FMeshNodeConverter&>(NodeTracker.GetConverter()).bMaterialsAssignedToStaticMesh = true;
 			if (Mtl* Material = UpdateGeometryNodeMaterial(*this, Instances, NodeTracker))
 			{
-				Instances.AssignMaterialToStaticMesh(Material);
+				Instances.AssignMaterialToStaticMesh(MaterialsCollectionTracker, Material);
 			}
 			break;
 		}
@@ -2714,8 +2722,7 @@ public:
 			{
 				if (Instances.Material != Material)
 				{
-					TSharedRef<IDatasmithMeshActorElement> DatasmithMeshActorRef = NodeTracker.GetConverted().DatasmithMeshActor.ToSharedRef();
-					FDatasmithMaxSceneExporter::ParseMaterialForMeshActor(Material, DatasmithMeshActorRef, Instances.Converted.SupportedChannels, FVector3f(NodeTracker.GetConverted().DatasmithMeshActor->GetTranslation()));
+					MaterialsCollectionTracker.AssignMeshActorMaterials(NodeTracker.GetConverted().DatasmithMeshActor, Material, Instances.Converted.SupportedChannels, FVector3f(NodeTracker.GetConverted().DatasmithMeshActor->GetTranslation()));
 				}
 			}
 		}
@@ -2817,6 +2824,7 @@ public:
 	class FMeshes
 	{
 	public:
+		FORCENOINLINE
 		static void AddMesh(FSceneTracker& Scene, FMeshConverterSource& MeshSource, FMeshConverted& MeshConverted, TFunction<void(bool, FMeshConverted&)> CompletionCallback)
 		{
 			// Reset old mesh
@@ -2846,7 +2854,7 @@ public:
 				NodeConverter.Meshes.Add(MeshConverted);
 
 				RegisterNodeForMaterial(NodeTracker, Material);
-				AssignMeshMaterials(MeshConverted.DatasmithMeshElement, Material, MeshConverted.SupportedChannels);
+				MaterialsCollectionTracker.AssignMeshMaterials(MeshConverted.DatasmithMeshElement, Material, MeshConverted.SupportedChannels);
 
 				FString MeshLabel = NodeTrackersNames.GetNodeName(NodeTracker) + (TEXT("_") + FString::FromInt(MeshIndex));
 				MeshConverted.DatasmithMeshElement->SetLabel(*MeshLabel);
diff --git a/Engine/Source/Programs/Enterprise/Datasmith/DatasmithMaxExporter/Private/DatasmithMaxWriter.cpp b/Engine/Source/Programs/Enterprise/Datasmith/DatasmithMaxExporter/Private/DatasmithMaxWriter.cpp
--- a/Engine/Source/Programs/Enterprise/Datasmith/DatasmithMaxExporter/Private/DatasmithMaxWriter.cpp
+++ b/Engine/Source/Programs/Enterprise/Datasmith/DatasmithMaxExporter/Private/DatasmithMaxWriter.cpp
@@ -444,8 +444,6 @@ bool FDatasmithMaxMatHelper::IsSRGB(Bitmap& InBitmap)
 	return true;
 }
 
-bool FDatasmithMaxMatExport::bForceReexport = false;
-
 // todo: need wa wya to return created texture element(if there was one), maybe use 'last' in DatasmithScene(check that count changed)
 //   verify that only one texture created!
 void FDatasmithMaxMatExport::GetXMLTexture(TSharedRef< IDatasmithScene > DatasmithScene, Texmap* InTexMap, const TCHAR* AssetsPath, TArray<TSharedPtr< IDatasmithTextureElement >>* OutTextureElements)
@@ -531,29 +529,18 @@ bool FDatasmithMaxMatExport::UseFirstSubMapOnly(EDSMaterialType MaterialType, Mt
 
 TSharedPtr< IDatasmithBaseMaterialElement > FDatasmithMaxMatExport::ExportUniqueMaterial(TSharedRef< IDatasmithScene > DatasmithScene, Mtl* Material, const TCHAR* AssetsPath)
 {
-	if (!bForceReexport)
+	if (!Material)
 	{
-		// Names should be unique prior to this
-		FString MaterialName(FDatasmithUtils::SanitizeObjectName(Material->GetName().data()));
-		for (int i = 0; i < DatasmithScene->GetMaterialsCount(); i++)
-		{
-			if (FString(DatasmithScene->GetMaterial(i)->GetName()) == MaterialName)
-			{
-				return DatasmithScene->GetMaterial(i);
-			}
-		}
+		return {};
 	}
-	else
+
+	FString MaterialName(FDatasmithMaxMaterialsToUEPbrManager::GetDatasmithMaterialName(Material));
+
+	for (int i = 0; i < DatasmithScene->GetMaterialsCount(); i++)
 	{
-		// Names should be unique prior to this
-		FString MaterialName(FDatasmithUtils::SanitizeObjectName(Material->GetName().data()));
-		for (int i = 0; i < DatasmithScene->GetMaterialsCount(); i++)
+		if (FString(DatasmithScene->GetMaterial(i)->GetName()) == MaterialName)
 		{
-			if (FString(DatasmithScene->GetMaterial(i)->GetName()) == MaterialName)
-			{
-				DatasmithScene->RemoveMaterial(DatasmithScene->GetMaterial(i));
-				break;
-			}
+			return DatasmithScene->GetMaterial(i);
 		}
 	}
 
@@ -563,17 +550,15 @@ TSharedPtr< IDatasmithBaseMaterialElement > FDatasmithMaxMatExport::ExportUnique
 
 		MaterialConverter->Convert( DatasmithScene, UEPbrMaterial, Material, AssetsPath );
 
-		if ( UEPbrMaterial )
-		{
-			DatasmithScene->AddMaterial( UEPbrMaterial );
-		}
+		FDatasmithMaxMaterialsToUEPbrManager::AddDatasmithMaterial(DatasmithScene, Material, UEPbrMaterial);
 
 		return UEPbrMaterial;
 	}
 	else
 	{
-		TSharedRef< IDatasmithMaterialElement > MaterialElement = FDatasmithSceneFactory::CreateMaterial( (TCHAR*)Material->GetName().data() );
-		DatasmithScene->AddMaterial( MaterialElement );
+		TSharedRef< IDatasmithMaterialElement > MaterialElement = FDatasmithSceneFactory::CreateMaterial( *MaterialName );
+		
+		FDatasmithMaxMaterialsToUEPbrManager::AddDatasmithMaterial(DatasmithScene, Material, MaterialElement);
 
 		WriteXMLMaterial( DatasmithScene, MaterialElement, Material );
 
diff --git a/Engine/Source/Programs/Enterprise/Datasmith/DatasmithMaxExporter/Private/MaxMaterialsToUEPbr/DatasmithMaxCoronaMaterialsToUEPbr.cpp b/Engine/Source/Programs/Enterprise/Datasmith/DatasmithMaxExporter/Private/MaxMaterialsToUEPbr/DatasmithMaxCoronaMaterialsToUEPbr.cpp
--- a/Engine/Source/Programs/Enterprise/Datasmith/DatasmithMaxExporter/Private/MaxMaterialsToUEPbr/DatasmithMaxCoronaMaterialsToUEPbr.cpp
+++ b/Engine/Source/Programs/Enterprise/Datasmith/DatasmithMaxExporter/Private/MaxMaterialsToUEPbr/DatasmithMaxCoronaMaterialsToUEPbr.cpp
@@ -747,8 +747,8 @@ void FDatasmithMaxCoronaMaterialsToUEPbr::Convert( TSharedRef< IDatasmithScene >
 	{
 		return;
 	}
-
-	TSharedRef< IDatasmithUEPbrMaterialElement > PbrMaterialElement = FDatasmithSceneFactory::CreateUEPbrMaterial( Material->GetName().data() );
+	
+	TSharedRef< IDatasmithUEPbrMaterialElement > PbrMaterialElement = FDatasmithSceneFactory::CreateUEPbrMaterial( GetMaterialName(Material) );
 	FScopedConvertState ScopedConvertState(ConvertState);
 	ConvertState.DatasmithScene = DatasmithScene;
 	ConvertState.MaterialElement = PbrMaterialElement;
@@ -1087,7 +1087,7 @@ void FDatasmithMaxCoronaBlendMaterialToUEPbr::Convert( TSharedRef<IDatasmithScen
 {
 	using namespace DatasmithMaxCoronaMaterialsToUEPbrImpl;
 
-	TSharedRef< IDatasmithUEPbrMaterialElement > PbrMaterialElement = FDatasmithSceneFactory::CreateUEPbrMaterial(Material->GetName().data());
+	TSharedRef< IDatasmithUEPbrMaterialElement > PbrMaterialElement = FDatasmithSceneFactory::CreateUEPbrMaterial(GetMaterialName(Material));
 	FScopedConvertState ScopedConvertState(ConvertState);
 	ConvertState.DatasmithScene = DatasmithScene;
 	ConvertState.MaterialElement = PbrMaterialElement;
@@ -1097,6 +1097,7 @@ void FDatasmithMaxCoronaBlendMaterialToUEPbr::Convert( TSharedRef<IDatasmithScen
 
 	//Exporting the base material.
 	IDatasmithMaterialExpressionFunctionCall* BaseMaterialFunctionCall = PbrMaterialElement->AddMaterialExpression<IDatasmithMaterialExpressionFunctionCall>();
+
 	if (TSharedPtr<IDatasmithBaseMaterialElement> ExportedMaterial = FDatasmithMaxMatExport::ExportUniqueMaterial(DatasmithScene, CoronaBlendMaterialProperties.BaseMaterial, AssetsPath))
 	{
 		BaseMaterialFunctionCall->SetFunctionPathName(ExportedMaterial->GetName());
@@ -1169,7 +1170,7 @@ void FDatasmithMaxCoronaLightMaterialToUEPbr::Convert(TSharedRef<IDatasmithScene
 		return;
 	}
 
-	TSharedRef<IDatasmithUEPbrMaterialElement> PbrMaterialElement = FDatasmithSceneFactory::CreateUEPbrMaterial(Material->GetName().data());
+	TSharedRef<IDatasmithUEPbrMaterialElement> PbrMaterialElement = FDatasmithSceneFactory::CreateUEPbrMaterial(GetMaterialName(Material));
 	FScopedConvertState ScopedConvertState(ConvertState);
 	ConvertState.DatasmithScene = DatasmithScene;
 	ConvertState.MaterialElement = PbrMaterialElement;
@@ -1208,7 +1209,7 @@ void FDatasmithMaxCoronaPhysicalMaterialToUEPbr::Convert(TSharedRef<IDatasmithSc
 		return;
 	}
 
-	TSharedRef< IDatasmithUEPbrMaterialElement > PbrMaterialElement = FDatasmithSceneFactory::CreateUEPbrMaterial( Material->GetName().data() );
+	TSharedRef< IDatasmithUEPbrMaterialElement > PbrMaterialElement = FDatasmithSceneFactory::CreateUEPbrMaterial( GetMaterialName(Material) );
 	FScopedConvertState ScopedConvertState(ConvertState);
 	ConvertState.DatasmithScene = DatasmithScene;
 	ConvertState.MaterialElement = PbrMaterialElement;
diff --git a/Engine/Source/Programs/Enterprise/Datasmith/DatasmithMaxExporter/Private/MaxMaterialsToUEPbr/DatasmithMaxMaterialsToUEPbr.cpp b/Engine/Source/Programs/Enterprise/Datasmith/DatasmithMaxExporter/Private/MaxMaterialsToUEPbr/DatasmithMaxMaterialsToUEPbr.cpp
--- a/Engine/Source/Programs/Enterprise/Datasmith/DatasmithMaxExporter/Private/MaxMaterialsToUEPbr/DatasmithMaxMaterialsToUEPbr.cpp
+++ b/Engine/Source/Programs/Enterprise/Datasmith/DatasmithMaxExporter/Private/MaxMaterialsToUEPbr/DatasmithMaxMaterialsToUEPbr.cpp
@@ -15,6 +15,8 @@
 #include "MaxMaterialsToUEPbr/DatasmithMaxMentalMaterialToUEPbr.h"
 
 
+DatasmithMaxDirectLink::FMaterialConversionContext* FDatasmithMaxMaterialsToUEPbrManager::Context = nullptr;
+
 FDatasmithMaxMaterialsToUEPbr* FDatasmithMaxMaterialsToUEPbrManager::GetMaterialConverter( Mtl* Material )
 {
 	if (Material == nullptr)
@@ -98,11 +100,38 @@ FDatasmithMaxMaterialsToUEPbr* FDatasmithMaxMaterialsToUEPbrManager::GetMaterial
 
 	if ( MaterialConverter && MaterialConverter->IsSupported( Material ) )
 	{
+		ensure(Context);
+		MaterialConverter->Context = Context;
 		return MaterialConverter;
 	}
 	return nullptr;
 }
 
+void FDatasmithMaxMaterialsToUEPbrManager::AddDatasmithMaterial(TSharedRef<IDatasmithScene> DatasmithScene, Mtl* Material, TSharedPtr<IDatasmithBaseMaterialElement> DatasmithMaterial)
+{
+	if (!DatasmithMaterial)
+	{
+		return;
+	}
+
+	if (!Context)
+	{
+		return;
+	}
+
+	Context->MaterialsCollectionTracker.AddDatasmithMaterialForUsedMaterial(DatasmithScene, Material, DatasmithMaterial);
+
+}
+
+const TCHAR* FDatasmithMaxMaterialsToUEPbrManager::GetDatasmithMaterialName(Mtl* Material)
+{
+	if (Context)
+	{
+		return Context->MaterialsCollectionTracker.GetMaterialName(Material);
+	}
+	return Material->GetName().data();
+}
+
 FDatasmithMaxMaterialsToUEPbr::FDatasmithMaxMaterialsToUEPbr()
 {
 	TexmapConverters.Add( new FDatasmithMaxBitmapToUEPbr() );
@@ -163,6 +192,14 @@ bool FDatasmithMaxMaterialsToUEPbr::IsTexmapSupported( Texmap* InTexmap ) const
 	return bIsTexmapSupported;
 }
 
+void FDatasmithMaxMaterialsToUEPbr::AddConvertedMap(const DatasmithMaxTexmapParser::FMapParameter& MapParameter)
+{
+	if (Context)
+	{
+		Context->TexmapsConverted.FindOrAdd(MapParameter.Map);
+	}
+}
+
 IDatasmithMaterialExpression* FDatasmithMaxMaterialsToUEPbr::ConvertTexmap( const DatasmithMaxTexmapParser::FMapParameter& MapParameter )
 {
 	if ( !MapParameter.bEnabled || !MapParameter.Map || FMath::IsNearlyZero( MapParameter.Weight ) )
@@ -187,10 +224,7 @@ IDatasmithMaterialExpression* FDatasmithMaxMaterialsToUEPbr::ConvertTexmap( cons
 		return nullptr;
 	}
 
-	if (TexmapsConverted)
-	{
-		TexmapsConverted->FindOrAdd(MapParameter.Map);
-	}
+	AddConvertedMap(MapParameter);
 
 	// Check if the texmap has a texture ouput
 	TextureOutput* TexOutput = nullptr;
diff --git a/Engine/Source/Programs/Enterprise/Datasmith/DatasmithMaxExporter/Private/MaxMaterialsToUEPbr/DatasmithMaxMaterialsToUEPbr.h b/Engine/Source/Programs/Enterprise/Datasmith/DatasmithMaxExporter/Private/MaxMaterialsToUEPbr/DatasmithMaxMaterialsToUEPbr.h
--- a/Engine/Source/Programs/Enterprise/Datasmith/DatasmithMaxExporter/Private/MaxMaterialsToUEPbr/DatasmithMaxMaterialsToUEPbr.h
+++ b/Engine/Source/Programs/Enterprise/Datasmith/DatasmithMaxExporter/Private/MaxMaterialsToUEPbr/DatasmithMaxMaterialsToUEPbr.h
@@ -4,6 +4,7 @@
 
 #include "CoreMinimal.h"
 #include "DatasmithDefinitions.h"
+#include "DatasmithMaxDirectLink.h"
 #include "Containers/IndirectArray.h"
 #include "Templates/SharedPointer.h"
 
@@ -24,10 +25,28 @@ namespace DatasmithMaxTexmapParser
 	struct FMapParameter;
 }
 
+namespace DatasmithMaxDirectLink
+{
+	class FMaterialsCollectionTracker;
+
+	struct FMaterialConversionContext
+	{
+		// Used to record texmaps that were asked for conversion whil ematerial conversion is performed
+		// recordking also datasmith elements created for those maps.
+		// Baked texmaps currently create datasmith elements during material convertion
+		TMap<Texmap*, TSet<TSharedPtr<IDatasmithTextureElement>>>& TexmapsConverted;
+		FMaterialsCollectionTracker& MaterialsCollectionTracker;
+	};
+}
+
+
 class FDatasmithMaxMaterialsToUEPbrManager
 {
 public:
 	static FDatasmithMaxMaterialsToUEPbr* GetMaterialConverter( Mtl* Material );
+	static void AddDatasmithMaterial(TSharedRef<IDatasmithScene> DatasmithScene, Mtl* Material, TSharedPtr<IDatasmithBaseMaterialElement> DatasmithMaterial);
+	static const TCHAR* GetDatasmithMaterialName(Mtl* Material);
+	static DatasmithMaxDirectLink::FMaterialConversionContext* Context;
 };
 
 class FDatasmithMaxMaterialsToUEPbr
@@ -60,10 +79,22 @@ public:
 
 	} ConvertState;
 
-	// Used to record texmaps that were asked for conversion whil ematerial conversion is performed
-	// recordking also datasmith elements created for those maps.
-	// Baked texmaps currently create datasmith elements during material convertion
-	TMap<Texmap*, TSet<TSharedPtr<IDatasmithTextureElement>>>* TexmapsConverted = nullptr;
+	void AddConvertedMap(const DatasmithMaxTexmapParser::FMapParameter& MapParameter);
+
+	void AddConvertedMapDatasmithElement(Texmap* InTexmap, TSharedPtr<IDatasmithTextureElement> TextureElement)
+	{
+		if (Context)
+		{
+			Context->TexmapsConverted.FindOrAdd(InTexmap).Add(TextureElement);
+		}
+	}
+
+	const TCHAR* GetMaterialName(Mtl* Material) const
+	{
+		return Context->MaterialsCollectionTracker.GetMaterialName(Material);
+	}
+
+	DatasmithMaxDirectLink::FMaterialConversionContext* Context = nullptr;
 
 protected:
 	TIndirectArray< IDatasmithMaxTexmapToUEPbr > TexmapConverters;
diff --git a/Engine/Source/Programs/Enterprise/Datasmith/DatasmithMaxExporter/Private/MaxMaterialsToUEPbr/DatasmithMaxMentalMaterialToUEPbr.cpp b/Engine/Source/Programs/Enterprise/Datasmith/DatasmithMaxExporter/Private/MaxMaterialsToUEPbr/DatasmithMaxMentalMaterialToUEPbr.cpp
--- a/Engine/Source/Programs/Enterprise/Datasmith/DatasmithMaxExporter/Private/MaxMaterialsToUEPbr/DatasmithMaxMentalMaterialToUEPbr.cpp
+++ b/Engine/Source/Programs/Enterprise/Datasmith/DatasmithMaxExporter/Private/MaxMaterialsToUEPbr/DatasmithMaxMentalMaterialToUEPbr.cpp
@@ -217,7 +217,7 @@ void FDatasmithMaxMentalMaterialToUEPbr::Convert( TSharedRef< IDatasmithScene >
 		return;
 	}
 
-	TSharedRef< IDatasmithUEPbrMaterialElement > PbrMaterialElement = FDatasmithSceneFactory::CreateUEPbrMaterial( Material->GetName().data() );
+	TSharedRef< IDatasmithUEPbrMaterialElement > PbrMaterialElement = FDatasmithSceneFactory::CreateUEPbrMaterial( GetMaterialName(Material) );
 	FScopedConvertState ScopedConvertState(ConvertState);
 	ConvertState.DatasmithScene = DatasmithScene;
 	ConvertState.MaterialElement = PbrMaterialElement;
diff --git a/Engine/Source/Programs/Enterprise/Datasmith/DatasmithMaxExporter/Private/MaxMaterialsToUEPbr/DatasmithMaxPhysicalMaterialToUEPbr.cpp b/Engine/Source/Programs/Enterprise/Datasmith/DatasmithMaxExporter/Private/MaxMaterialsToUEPbr/DatasmithMaxPhysicalMaterialToUEPbr.cpp
--- a/Engine/Source/Programs/Enterprise/Datasmith/DatasmithMaxExporter/Private/MaxMaterialsToUEPbr/DatasmithMaxPhysicalMaterialToUEPbr.cpp
+++ b/Engine/Source/Programs/Enterprise/Datasmith/DatasmithMaxExporter/Private/MaxMaterialsToUEPbr/DatasmithMaxPhysicalMaterialToUEPbr.cpp
@@ -677,7 +677,7 @@ void FDatasmithMaxPhysicalMaterialToUEPbr::Convert( TSharedRef< IDatasmithScene
 		return;
 	}
 
-	TSharedRef< IDatasmithUEPbrMaterialElement > PbrMaterialElement = FDatasmithSceneFactory::CreateUEPbrMaterial( Material->GetName().data() );
+	TSharedRef< IDatasmithUEPbrMaterialElement > PbrMaterialElement = FDatasmithSceneFactory::CreateUEPbrMaterial( GetMaterialName(Material) );
 	FScopedConvertState ScopedConvertState(ConvertState);
 	ConvertState.DatasmithScene = DatasmithScene;
 	ConvertState.MaterialElement = PbrMaterialElement;
diff --git a/Engine/Source/Programs/Enterprise/Datasmith/DatasmithMaxExporter/Private/MaxMaterialsToUEPbr/DatasmithMaxScanlineMaterialsToUEPbr.cpp b/Engine/Source/Programs/Enterprise/Datasmith/DatasmithMaxExporter/Private/MaxMaterialsToUEPbr/DatasmithMaxScanlineMaterialsToUEPbr.cpp
--- a/Engine/Source/Programs/Enterprise/Datasmith/DatasmithMaxExporter/Private/MaxMaterialsToUEPbr/DatasmithMaxScanlineMaterialsToUEPbr.cpp
+++ b/Engine/Source/Programs/Enterprise/Datasmith/DatasmithMaxExporter/Private/MaxMaterialsToUEPbr/DatasmithMaxScanlineMaterialsToUEPbr.cpp
@@ -148,7 +148,7 @@ void FDatasmithMaxScanlineMaterialsToUEPbr::Convert( TSharedRef< IDatasmithScene
 		return;
 	}
 
-	TSharedRef< IDatasmithUEPbrMaterialElement > PbrMaterialElement = FDatasmithSceneFactory::CreateUEPbrMaterial( Material->GetName().data() );
+	TSharedRef< IDatasmithUEPbrMaterialElement > PbrMaterialElement = FDatasmithSceneFactory::CreateUEPbrMaterial( GetMaterialName(Material) );
 	FScopedConvertState ScopedConvertState(ConvertState);
 	ConvertState.DatasmithScene = DatasmithScene;
 	ConvertState.MaterialElement = PbrMaterialElement;
@@ -344,7 +344,7 @@ void FDatasmithMaxBlendMaterialsToUEPbr::Convert( TSharedRef< IDatasmithScene >
 		ParamBlock2->ReleaseDesc();
 	}
 
-	TSharedRef< IDatasmithUEPbrMaterialElement > PbrMaterialElement = FDatasmithSceneFactory::CreateUEPbrMaterial(Material->GetName().data());
+	TSharedRef< IDatasmithUEPbrMaterialElement > PbrMaterialElement = FDatasmithSceneFactory::CreateUEPbrMaterial(GetMaterialName(Material));
 	FScopedConvertState ScopedConvertState(ConvertState);
 	ConvertState.DatasmithScene = DatasmithScene;
 	ConvertState.MaterialElement = PbrMaterialElement;
diff --git a/Engine/Source/Programs/Enterprise/Datasmith/DatasmithMaxExporter/Private/MaxMaterialsToUEPbr/DatasmithMaxTexmapToUEPbr.cpp b/Engine/Source/Programs/Enterprise/Datasmith/DatasmithMaxExporter/Private/MaxMaterialsToUEPbr/DatasmithMaxTexmapToUEPbr.cpp
--- a/Engine/Source/Programs/Enterprise/Datasmith/DatasmithMaxExporter/Private/MaxMaterialsToUEPbr/DatasmithMaxTexmapToUEPbr.cpp
+++ b/Engine/Source/Programs/Enterprise/Datasmith/DatasmithMaxExporter/Private/MaxMaterialsToUEPbr/DatasmithMaxTexmapToUEPbr.cpp
@@ -1703,7 +1703,8 @@ IDatasmithMaterialExpression* FDatasmithMaxBakeableToUEPbr::Convert( FDatasmithM
 
 	if ( BakedTextureElement )
 	{
-		MaxMaterialToUEPbr->TexmapsConverted->FindOrAdd(InTexmap).Add(BakedTextureElement);
+		MaxMaterialToUEPbr->AddConvertedMapDatasmithElement(InTexmap, BakedTextureElement);
+
 		IDatasmithMaterialExpression* MaterialExpression = MaxMaterialToUEPbr->ConvertState.MaterialElement->AddMaterialExpression( EDatasmithMaterialExpressionType::Texture );
 
 		if ( MaterialExpression )
diff --git a/Engine/Source/Programs/Enterprise/Datasmith/DatasmithMaxExporter/Private/MaxMaterialsToUEPbr/DatasmithMaxVrayMaterialsToUEPbr.cpp b/Engine/Source/Programs/Enterprise/Datasmith/DatasmithMaxExporter/Private/MaxMaterialsToUEPbr/DatasmithMaxVrayMaterialsToUEPbr.cpp
--- a/Engine/Source/Programs/Enterprise/Datasmith/DatasmithMaxExporter/Private/MaxMaterialsToUEPbr/DatasmithMaxVrayMaterialsToUEPbr.cpp
+++ b/Engine/Source/Programs/Enterprise/Datasmith/DatasmithMaxExporter/Private/MaxMaterialsToUEPbr/DatasmithMaxVrayMaterialsToUEPbr.cpp
@@ -520,7 +520,7 @@ void FDatasmithMaxVRayMaterialsToUEPbr::Convert( TSharedRef< IDatasmithScene > D
 		return;
 	}
 
-	TSharedRef< IDatasmithUEPbrMaterialElement > PbrMaterialElement = FDatasmithSceneFactory::CreateUEPbrMaterial( Material->GetName().data() );
+	TSharedRef< IDatasmithUEPbrMaterialElement > PbrMaterialElement = FDatasmithSceneFactory::CreateUEPbrMaterial( GetMaterialName(Material) );
 	FScopedConvertState ScopedConvertState(ConvertState);
 	ConvertState.DatasmithScene = DatasmithScene;
 	ConvertState.MaterialElement = PbrMaterialElement;
@@ -940,7 +940,7 @@ void FDatasmithMaxVRay2SidedMaterialsToUEPbr::Convert( TSharedRef<IDatasmithScen
 
 			if ( MaterialElement )
 			{
-				MaterialElement->SetName( Material->GetName().data() ); // Name it with the main material not the front material
+				MaterialElement->SetName( GetMaterialName(Material) ); // Name it with the main material not the front material
 
 				if ( MaterialElement->IsA( EDatasmithElementType::UEPbrMaterial ) )
 				{
@@ -980,7 +980,7 @@ void FDatasmithMaxVRayWrapperMaterialsToUEPbr::Convert( TSharedRef<IDatasmithSce
 
 			if ( MaterialElement )
 			{
-				MaterialElement->SetName( Material->GetName().data() ); // Name it with the main material not the base material
+				MaterialElement->SetName( GetMaterialName(Material) ); // Name it with the main material not the base material
 			}
 		}
 	}
@@ -1023,7 +1023,7 @@ void FDatasmithMaxVRayBlendMaterialToUEPbr::Convert( TSharedRef<IDatasmithScene>
 {
 	using namespace DatasmithMaxVRayMaterialsToUEPbrImpl;
 
-	TSharedRef< IDatasmithUEPbrMaterialElement > PbrMaterialElement = FDatasmithSceneFactory::CreateUEPbrMaterial(Material->GetName().data());
+	TSharedRef< IDatasmithUEPbrMaterialElement > PbrMaterialElement = FDatasmithSceneFactory::CreateUEPbrMaterial(GetMaterialName(Material));
 	FScopedConvertState ScopedConvertState(ConvertState);
 	ConvertState.DatasmithScene = DatasmithScene;
 	ConvertState.MaterialElement = PbrMaterialElement;
@@ -1113,7 +1113,7 @@ void FDatasmithMaxVRayLightMaterialToUEPbr::Convert(TSharedRef<IDatasmithScene>
 		return;
 	}
 
-	TSharedRef<IDatasmithUEPbrMaterialElement> PbrMaterialElement = FDatasmithSceneFactory::CreateUEPbrMaterial(Material->GetName().data());
+	TSharedRef<IDatasmithUEPbrMaterialElement> PbrMaterialElement = FDatasmithSceneFactory::CreateUEPbrMaterial(GetMaterialName(Material));
 	FScopedConvertState ScopedConvertState(ConvertState);
 	ConvertState.DatasmithScene = DatasmithScene;
 	ConvertState.MaterialElement = PbrMaterialElement;
diff --git a/Engine/Source/Programs/Enterprise/Datasmith/DatasmithMaxExporter/Public/DatasmithMaxDirectLink.h b/Engine/Source/Programs/Enterprise/Datasmith/DatasmithMaxExporter/Public/DatasmithMaxDirectLink.h
--- a/Engine/Source/Programs/Enterprise/Datasmith/DatasmithMaxExporter/Public/DatasmithMaxDirectLink.h
+++ b/Engine/Source/Programs/Enterprise/Datasmith/DatasmithMaxExporter/Public/DatasmithMaxDirectLink.h
@@ -9,6 +9,7 @@
 
 #include "IDatasmithSceneElements.h"
 #include "DatasmithMesh.h"
+#include "DatasmithUtils.h"
 
 #include "Logging/LogMacros.h"
 
@@ -582,6 +583,14 @@ public:
 
 	FMaterialTracker* AddMaterial(Mtl* Material); // Add material to track its changes
 
+	void AddActualMaterial(FMaterialTracker& MaterialTracker, Mtl* Material); // Add material used in a tracked material's graph
+	const TCHAR* GetMaterialName(Mtl* SubMaterial); // Get name used for Datasmith material. Datasmith material names must be unique(used to identify elements)
+
+	void AssignMeshMaterials(const TSharedPtr<IDatasmithMeshElement>& MeshElement, Mtl* Material, const TSet<uint16>& SupportedChannels);
+	void AssignMeshActorMaterials(const TSharedPtr<IDatasmithMeshActorElement>& MeshActor, Mtl* Material, TSet<uint16>& SupportedChannels, const FVector3f& RandomSeed);
+
+	void AddDatasmithMaterialForUsedMaterial(TSharedRef<IDatasmithScene> DatasmithScene, Mtl* Material, TSharedPtr<IDatasmithBaseMaterialElement> DatasmithMaterial); // Record which Datasmith material was created for a Max material, not only for tracked(assigned) materials
+
 	void InvalidateMaterial(Mtl* Material); // Mark changed
 
 	void UpdateMaterial(FMaterialTracker* MaterialTracker); // Reparse source material
@@ -617,11 +626,15 @@ public:
 	TMap<Mtl*, TSet<FMaterialTracker*>> UsedMaterialToMaterialTracker; // Materials uses by nodes keep set of assigned materials they are used for
 	TMap<Mtl*, TSharedPtr<IDatasmithBaseMaterialElement>> UsedMaterialToDatasmithMaterial;
 
+	TMap<Mtl*, FString> UsedMaterialToDatasmithMaterialName;
+
 	TMap<Texmap*, TSet<FMaterialTracker*>> UsedTextureToMaterialTracker; // Materials uses by nodes keep set of assigned textures they are used for
 	TMap<Texmap*, TSet<TSharedPtr<IDatasmithTextureElement>>> UsedTextureToDatasmithElement; // Keep track of Datasmith Element created for texmap to simplify update/removal(no need to search DatasmithScene)
 	// note: each texmap can create multiple texture elements
 
 	FSceneUpdateStats& Stats;
+
+	FDatasmithUniqueNameProvider MaterialNameProvider;
 };
 
 //----
diff --git a/Engine/Source/Programs/Enterprise/Datasmith/DatasmithMaxExporter/Public/DatasmithMaxWriter.h b/Engine/Source/Programs/Enterprise/Datasmith/DatasmithMaxExporter/Public/DatasmithMaxWriter.h
--- a/Engine/Source/Programs/Enterprise/Datasmith/DatasmithMaxExporter/Public/DatasmithMaxWriter.h
+++ b/Engine/Source/Programs/Enterprise/Datasmith/DatasmithMaxExporter/Public/DatasmithMaxWriter.h
@@ -109,8 +109,6 @@ public:
 	static TSharedPtr< IDatasmithBaseMaterialElement > ExportUniqueMaterial(TSharedRef< IDatasmithScene > DatasmithScene, Mtl* Material, const TCHAR* AssetsPath);
 	static void WriteXMLMaterial(TSharedRef< IDatasmithScene > DatasmithScene, TSharedPtr< IDatasmithMaterialElement > MaterialElement, Mtl* Material);
 
-	static bool bForceReexport; // todo: ExportUniqueMaterial calls go deep into composite materials export - need to fix this for DirectLink Sync
-
 private:
 	static bool UseFirstSubMapOnly(EDSMaterialType MaterialType, Mtl* Material);
 };
