[Sketchup] Fixed Component/instance name change sync

diff --git a/Engine/Source/Programs/Enterprise/Datasmith/DatasmithSketchUpRubyExporter/Private/DatasmithSketchUpComponent.cpp b/Engine/Source/Programs/Enterprise/Datasmith/DatasmithSketchUpRubyExporter/Private/DatasmithSketchUpComponent.cpp
--- a/Engine/Source/Programs/Enterprise/Datasmith/DatasmithSketchUpRubyExporter/Private/DatasmithSketchUpComponent.cpp
+++ b/Engine/Source/Programs/Enterprise/Datasmith/DatasmithSketchUpRubyExporter/Private/DatasmithSketchUpComponent.cpp
@@ -111,7 +111,7 @@ void FNodeOccurence::UpdateMeshActors(FExportContext& Context)
 	MeshActors.Reset(EntitiesGeometry->GetMeshCount());
 
 	FString ComponentActorName = GetActorName();
-	FString MeshActorLabel = GetActorLabel();
+	
 	for (int32 MeshIndex = 0; MeshIndex < EntitiesGeometry->GetMeshCount(); ++MeshIndex)
 	{
 		FString MeshActorName = FString::Printf(TEXT("%ls_%d"), *ComponentActorName, MeshIndex + 1); // Count meshes/mesh actors from 1
@@ -121,9 +121,6 @@ void FNodeOccurence::UpdateMeshActors(FExportContext& Context)
 
 		MeshActors.Add(DMeshActorPtr);
 
-		// Set the mesh actor label used in the Unreal UI.
-		DMeshActorPtr->SetLabel(*MeshActorLabel);
-
 		// Add the Datasmith actor component depth tag.
 		// We use component depth + 1 to factor in the added Datasmith scene root once imported in Unreal.
 		FString ComponentDepthTag = FString::Printf(TEXT("SU.DEPTH.%d"), Depth + 1);
@@ -608,10 +605,13 @@ void FEntity::UpdateOccurrence(FExportContext& Context, FNodeOccurence& Node)
 		Node.InheritedMaterialID = Node.ParentNode->InheritedMaterialID;
 	}
 
+
+	FString MeshActorLabel = Node.GetActorLabel();
 	// Update Datasmith Mesh Actors
 	for (int32 MeshIndex = 0; MeshIndex < Node.MeshActors.Num(); ++MeshIndex)
 	{
 		const TSharedPtr<IDatasmithMeshActorElement>& MeshActor = Node.MeshActors[MeshIndex];
+		MeshActor->SetLabel(*MeshActorLabel);
 		MeshActor->SetLayer(*EffectiveLayerName);
 
 		// Update Override(Inherited)  Material
@@ -698,6 +698,11 @@ void FComponentInstance::UpdateOccurrence(FExportContext& Context, FNodeOccurenc
 
 	SUComponentInstanceRef InSComponentInstanceRef = GetComponentInstanceRef();
 
+	if (FDefinition* EntityDefinition = GetDefinition())
+	{
+		EntityDefinition->BuildNodeNames(Node);
+	}
+
 	// Set the actor label used in the Unreal UI.
 	Node.DatasmithActorElement->SetLabel(*Node.DatasmithActorLabel);
 
@@ -728,8 +733,6 @@ void FComponentInstance::UpdateOccurrence(FExportContext& Context, FNodeOccurenc
 		MeshActor->SetTranslation(Node.DatasmithActorElement->GetTranslation());
 	}
 
-	
-
 	Super::UpdateOccurrence(Context, Node);
 }
 
diff --git a/Engine/Source/Programs/Enterprise/Datasmith/DatasmithSketchUpRubyExporter/Private/DatasmithSketchUpComponent.h b/Engine/Source/Programs/Enterprise/Datasmith/DatasmithSketchUpRubyExporter/Private/DatasmithSketchUpComponent.h
--- a/Engine/Source/Programs/Enterprise/Datasmith/DatasmithSketchUpRubyExporter/Private/DatasmithSketchUpComponent.h
+++ b/Engine/Source/Programs/Enterprise/Datasmith/DatasmithSketchUpRubyExporter/Private/DatasmithSketchUpComponent.h
@@ -126,6 +126,7 @@ namespace DatasmithSketchUp
 		virtual void CreateActor(FExportContext& Context, FNodeOccurence& Node) = 0; // Create Datasmith actor for node occurrence
 		virtual void UpdateGeometry(FExportContext& Context) = 0; // Convert definition's Entities geometry to Datasmith Mesh
 		virtual void UpdateMetadata(FExportContext& Context) = 0; 
+		virtual void BuildNodeNames(FNodeOccurence& Node) = 0;
 
 		void EntityVisible(FEntity* Entity, bool bVisible);
 
@@ -152,7 +153,6 @@ namespace DatasmithSketchUp
 		void UpdateDefinition(FExportContext& Context);
 
 	protected:
-		virtual void BuildNodeNames(FNodeOccurence& Node) = 0;
 
 		TSharedPtr<DatasmithSketchUp::FEntities> Entities;
 
@@ -172,19 +172,19 @@ namespace DatasmithSketchUp
 		);
 
 		// Begin FDefinition
-		void Parse(FExportContext& Context);
-		void CreateActor(FExportContext& Context, FNodeOccurence& Node) override;
-		void BuildNodeNames(FNodeOccurence& Node) override;
-		void UpdateGeometry(FExportContext& Context) override;
-		void UpdateMetadata(FExportContext& Context) override;
-
-		void AddInstance(FExportContext& Context, TSharedPtr<FComponentInstance> Instance);
-		void InvalidateInstancesGeometry(FExportContext& Context) override;
-		void InvalidateInstancesMetadata(FExportContext& Context) override;
-		void FillOccurrenceActorMetadata(FNodeOccurence& Node) override;
-
-		FString GetSketchupSourceGUID() override;
-		FString GetSketchupSourceName()  override;
+		virtual void Parse(FExportContext& Context) override;
+		virtual void CreateActor(FExportContext& Context, FNodeOccurence& Node) override;
+		virtual void BuildNodeNames(FNodeOccurence& Node) override;
+		virtual void UpdateGeometry(FExportContext& Context) override;
+		virtual void UpdateMetadata(FExportContext& Context) override;
+
+		virtual void AddInstance(FExportContext& Context, TSharedPtr<FComponentInstance> Instance) override;
+		virtual void InvalidateInstancesGeometry(FExportContext& Context) override;
+		virtual void InvalidateInstancesMetadata(FExportContext& Context) override;
+		virtual void FillOccurrenceActorMetadata(FNodeOccurence& Node) override;
+
+		virtual FString GetSketchupSourceGUID() override;
+		virtual FString GetSketchupSourceName()  override;
 		// End FDefinition
 
 
@@ -213,19 +213,19 @@ namespace DatasmithSketchUp
 		FModelDefinition(SUModelRef Model);
 
 		// Being FDefinition
-		void Parse(FExportContext& Context);
-		void CreateActor(FExportContext& Context, FNodeOccurence& Node) override;
-		void BuildNodeNames(FNodeOccurence& Node) override;
-		void UpdateGeometry(FExportContext& Context) override;
-		void UpdateMetadata(FExportContext& Context) override;
-
-		void AddInstance(FExportContext& Context, TSharedPtr<FComponentInstance> Instance);
-		void InvalidateInstancesGeometry(FExportContext& Context) override;
-		void InvalidateInstancesMetadata(FExportContext& Context) override;
-		void FillOccurrenceActorMetadata(FNodeOccurence& Node) override;
-
-		FString GetSketchupSourceGUID() override;
-		FString GetSketchupSourceName()  override;
+		virtual void Parse(FExportContext& Context) override;
+		virtual void CreateActor(FExportContext& Context, FNodeOccurence& Node) override;
+		virtual void BuildNodeNames(FNodeOccurence& Node) override;
+		virtual void UpdateGeometry(FExportContext& Context) override;
+		virtual void UpdateMetadata(FExportContext& Context) override;
+
+		virtual void AddInstance(FExportContext& Context, TSharedPtr<FComponentInstance> Instance) override;
+		virtual void InvalidateInstancesGeometry(FExportContext& Context) override;
+		virtual void InvalidateInstancesMetadata(FExportContext& Context) override;
+		virtual void FillOccurrenceActorMetadata(FNodeOccurence& Node) override;
+
+		virtual FString GetSketchupSourceGUID() override;
+		virtual FString GetSketchupSourceName()  override;
 		// End FDefinition
 
 	private:
