[SketchUp] Enable parent non-uniform scaling of a rotated component

 - Unreal doesn't allow non-uniform scaling of a parent actor to be applied to a rotated child object(scaling is applied before rotation, in mesh's local space). SketchUp allows this and in some cases it's useful/possible to 'bake' rotation into component's geometry to have same effect in Unreal as in SU.

- Restrictions: Transform 'baking' is done only for purely geometric components(without child components). And when component that has oly single instance(single transform to be baked). Multiple occurrencies of the same instance are allowed since they have the same local transform(which was baked).

- Translation and uniform scaling is not baked(only non-uniform scale, rotate and shear). These is not needed and more convenient for user to keep them in the actor.

diff --git a/Engine/Source/Programs/Enterprise/Datasmith/DatasmithSketchUpRubyExporter/Private/DatasmithSketchUpComponent.cpp b/Engine/Source/Programs/Enterprise/Datasmith/DatasmithSketchUpRubyExporter/Private/DatasmithSketchUpComponent.cpp
--- a/Engine/Source/Programs/Enterprise/Datasmith/DatasmithSketchUpRubyExporter/Private/DatasmithSketchUpComponent.cpp
+++ b/Engine/Source/Programs/Enterprise/Datasmith/DatasmithSketchUpRubyExporter/Private/DatasmithSketchUpComponent.cpp
@@ -18,6 +18,7 @@
 #include "SketchUpAPI/model/layer.h"
 #include "SketchUpAPI/model/model.h"
 #include "SketchUpAPI/model/component_instance.h"
+#include "SketchUpAPI/geometry/transformation.h"
 
 #if !defined(SKP_SDK_2019) && !defined(SKP_SDK_2020)
 #include "SketchUpAPI/model/layer_folder.h"
@@ -393,6 +394,13 @@ FString FModelDefinition::GetSketchupSourceName()
 	return SketchupSourceName;
 }
 
+SUTransformation FModelDefinition::GetMeshBakedTransform()
+{
+	SUTransformation Transform;
+	SUTransformationScale(&Transform, 1.0);
+	return Transform;
+}
+
 FString FModelDefinition::GetSketchupSourceGUID()
 {
 	return TEXT("MODEL");
@@ -470,6 +478,7 @@ void FComponentDefinition::CreateActor(FExportContext& Context, FNodeOccurence&
 void FComponentDefinition::UpdateGeometry(FExportContext& Context)
 {
 	Entities->UpdateGeometry(Context);
+	bBakeTransformIntoMesh  = ShouldBakeTransformIntoMesh();
 }
 
 void FComponentDefinition::UpdateMetadata(FExportContext& Context)
@@ -535,6 +544,118 @@ FString FComponentDefinition::GetSketchupSourceName()
 	return SuGetString(SUComponentDefinitionGetName, ComponentDefinitionRef);
 }
 
+// Bake transform into the mesh only when
+// - component has no child components, 
+// - only SINGLE instance
+// - local transform has non-trivial rotation component
+// Baking transform into mesh might be useful to convert transformation which is not representable in Unreal
+// for example rotated geometry with scaled parent. In Unreal scaling would only be applied in local static mesh space
+// So that scaled(on parent node) rotated cube that should be a rhombus in SU in unreal would become just a box because scaling would apply along local axes(i.e. before rotation)
+// But baking rotation into mesh itself makes scaling apply properly
+bool FComponentDefinition::ShouldBakeTransformIntoMesh()
+{
+	size_t ComponentInstanceCount = 0;
+	SUEntitiesGetNumInstances(GetEntities().EntitiesRef, &ComponentInstanceCount);
+
+	size_t GroupCount = 0;
+	SUEntitiesGetNumGroups(GetEntities().EntitiesRef, &GroupCount);
+
+	if ((Instances.Num() != 1) || ((ComponentInstanceCount + GroupCount) != 0))
+	{
+		return false;
+	}
+
+	// Check the (only) instance's transform
+	FComponentInstance* ComponentInstance = Instances.Array()[0];
+
+	SUTransformation Transform;
+	SUComponentInstanceGetTransform(ComponentInstance->GetComponentInstanceRef(), &Transform);
+
+	// Grab 3x3 rotation submatrix and check if it is not an identity transform
+	const double R[] = {
+		Transform.values[0], Transform.values[1], Transform.values[2],
+		Transform.values[4], Transform.values[5], Transform.values[6],
+		Transform.values[8], Transform.values[9], Transform.values[10]
+	};
+
+	static const double Identity[] = {
+		1, 0, 0,
+		0, 1, 0,
+		0, 0, 1,
+	};
+
+	static_assert(sizeof(Identity)/sizeof(Identity[0]) == sizeof(R)/sizeof(R[0]));
+
+	for (int32 I = 0; I < sizeof(Identity)/sizeof(Identity[0]); ++I)
+	{
+		if (!FMath::IsNearlyEqual(Identity[I], R[I]))
+		{
+			return true;
+		}
+	}
+
+	return false;
+}
+
+SUTransformation FComponentDefinition::GetMeshBakedTransform()
+{
+	if (!ShouldBakeTransformIntoMesh())
+	{
+		SUTransformation Transform;
+		SUTransformationScale(&Transform, 1.0);
+		return Transform;
+	}
+
+	FComponentInstance* ComponentInstance = Instances.Array()[0];
+	SUTransformation Transform;
+	SUComponentInstanceGetTransform(ComponentInstance->GetComponentInstanceRef(), &Transform);
+
+	// Don't bake translation 
+	Transform.values[12] = 0;
+	Transform.values[13] = 0;
+	Transform.values[14] = 0;
+	// and uniform scale
+	Transform.values[15] = 1;
+	// There's no need to bake translation and uniform scale into mesh and keeping those (at least translation) at actor transform is more convenient
+
+	return Transform;
+}
+
+void FComponentDefinition::GetLocalTransform(FComponentInstance& ComponentInstance, SUTransformation& SComponentInstanceLocalTransform)
+{
+	if (ShouldBakeTransformIntoMesh())
+	{
+		SUTransformation SComponentInstanceTransform;
+		SUComponentInstanceGetTransform(ComponentInstance.GetComponentInstanceRef(), &SComponentInstanceTransform);
+
+		// Local instance transform rotation part is baked into mesh, scale and translation aren't
+		double UniformScale = SComponentInstanceTransform.values[15];
+		const SUVector3D Translation{SComponentInstanceTransform.values[12], SComponentInstanceTransform.values[13], SComponentInstanceTransform.values[14]};
+
+		SUTransformation Transform;
+		SUTransformationTranslation(&Transform, &Translation);
+		Transform.values[15] = UniformScale;
+
+		SComponentInstanceLocalTransform = Transform;
+	}
+	else
+	{
+		SUComponentInstanceGetTransform(ComponentInstance.GetComponentInstanceRef(), &SComponentInstanceLocalTransform);
+	}
+}
+
+void FComponentDefinition::ComponentInstancePropertiesInvalidated()
+{
+	// If mesh transform baking is affecting geometry then invalidate geometry
+	// i.e. if transform should be baked or need to bake it has changed
+	bool bShouldBakeTransformIntoMesh = ShouldBakeTransformIntoMesh();
+	if (bShouldBakeTransformIntoMesh || (bShouldBakeTransformIntoMesh != bBakeTransformIntoMesh))
+	{
+		InvalidateDefinitionGeometry();
+	}
+}
+
+
 FString FComponentDefinition::GetSketchupSourceGUID()
 {
 	// Retrieve the SketchUp component definition IFC GUID.
@@ -696,8 +817,6 @@ void FComponentInstance::UpdateOccurrence(FExportContext& Context, FNodeOccurenc
 		return;
 	}
 
-	SUComponentInstanceRef InSComponentInstanceRef = GetComponentInstanceRef();
-
 	if (FDefinition* EntityDefinition = GetDefinition())
 	{
 		EntityDefinition->BuildNodeNames(Node);
@@ -713,11 +832,15 @@ void FComponentInstance::UpdateOccurrence(FExportContext& Context, FNodeOccurenc
 	// Set the Datasmith actor layer name.
 	Node.DatasmithActorElement->SetLayer(*FDatasmithUtils::SanitizeObjectName(SEffectiveLayerName));
 
-	SUTransformation SComponentInstanceWorldTransform = DatasmithSketchUpUtils::GetComponentInstanceTransform(InSComponentInstanceRef, Node.ParentNode->WorldTransform);
-	Node.WorldTransform = SComponentInstanceWorldTransform;
+	// Compute the world transform of the SketchUp component instance.
+	SUTransformation LocalTransform;
+	Definition.GetLocalTransform(*this, LocalTransform);
+	SUTransformation WorldTransform;
+	SUTransformationMultiply(&Node.ParentNode->WorldTransform, &LocalTransform, &WorldTransform);
+	Node.WorldTransform = WorldTransform; // Store world transform to be used by children to compute its
 
 	// Set the Datasmith actor world transform.
-	DatasmithSketchUpUtils::SetActorTransform(Node.DatasmithActorElement, SComponentInstanceWorldTransform);
+	DatasmithSketchUpUtils::SetActorTransform(Node.DatasmithActorElement, Node.WorldTransform);
 
 	Node.ResetMetadataElement(Context);// todo: can enable/disable metadata export by toggling this code
 	FillOccurrenceActorMetadata(Node);
@@ -768,6 +891,16 @@ void FComponentInstance::UpdateMetadata(FExportContext& Context)
 	ParsedMetadata = MakeUnique<FMetadata>(SUComponentInstanceToEntity(GetComponentInstanceRef()));
 }
 
+void FComponentInstance::UpdateEntityProperties(FExportContext& Context)
+{
+	if (bPropertiesInvalidated)
+	{
+		Definition.ComponentInstancePropertiesInvalidated();
+	}
+	
+	FEntity::UpdateEntityProperties(Context);
+}
+
 void FComponentInstance::InvalidateOccurrencesGeometry(FExportContext& Context)
 {
 	for (FNodeOccurence* Node : Occurrences)
diff --git a/Engine/Source/Programs/Enterprise/Datasmith/DatasmithSketchUpRubyExporter/Private/DatasmithSketchUpComponent.h b/Engine/Source/Programs/Enterprise/Datasmith/DatasmithSketchUpRubyExporter/Private/DatasmithSketchUpComponent.h
--- a/Engine/Source/Programs/Enterprise/Datasmith/DatasmithSketchUpRubyExporter/Private/DatasmithSketchUpComponent.h
+++ b/Engine/Source/Programs/Enterprise/Datasmith/DatasmithSketchUpRubyExporter/Private/DatasmithSketchUpComponent.h
@@ -130,7 +130,7 @@ namespace DatasmithSketchUp
 
 		void EntityVisible(FEntity* Entity, bool bVisible);
 
-		// Modfication methods
+		// Modification methods
 		virtual void AddInstance(FExportContext& Context, TSharedPtr<FComponentInstance> Instance) = 0; // Register child CompoenntInstance Entity of Definition's Entities
 		virtual void InvalidateInstancesGeometry(FExportContext& Context) = 0; // Mark that all instances(and their occurrences) needed to be updated
 		virtual void InvalidateInstancesMetadata(FExportContext& Context) = 0; // Mark that all instances(and their occurrences) needed to be updated
@@ -139,7 +139,8 @@ namespace DatasmithSketchUp
 		virtual FString GetSketchupSourceGUID() = 0;
 		virtual FString GetSketchupSourceName() = 0;
 
-		
+		virtual SUTransformation GetMeshBakedTransform() = 0; // Transform that may be used to bake geometry
+
 		FEntities& GetEntities()
 		{
 			return *Entities;
@@ -184,15 +185,22 @@ namespace DatasmithSketchUp
 		virtual void FillOccurrenceActorMetadata(FNodeOccurence& Node) override;
 
 		virtual FString GetSketchupSourceGUID() override;
-		virtual FString GetSketchupSourceName()  override;
+		virtual FString GetSketchupSourceName() override;
+
+		virtual SUTransformation GetMeshBakedTransform() override;
 		// End FDefinition
 
 
+
 		// Register/unregister instanced of this definition
 		void LinkComponentInstance(FComponentInstance* ComponentInstance);
 		void UnlinkComponentInstance(FComponentInstance* ComponentInstance);
 		void RemoveComponentDefinition(FExportContext& Context);
 
+		bool ShouldBakeTransformIntoMesh();
+		void GetLocalTransform(FComponentInstance& ComponentInstance, SUTransformation& SComponentInstanceLocalTransform);
+		void ComponentInstancePropertiesInvalidated();
+
 		// Source SketchUp component ID.
 		FComponentDefinitionIDType SketchupSourceID;
 		TSet<FComponentInstance*> Instances; // Tracked instances of this ComponentDefinition
@@ -204,6 +212,8 @@ namespace DatasmithSketchUp
 
 		// Whether or not the source SketchUp component behaves like a billboard, always presenting a 2D surface perpendicular to the direction of camera.
 		bool bSketchupSourceFaceCamera = false;
+
+		bool bBakeTransformIntoMesh = false;
 	};
 
 
@@ -226,6 +236,8 @@ namespace DatasmithSketchUp
 
 		virtual FString GetSketchupSourceGUID() override;
 		virtual FString GetSketchupSourceName()  override;
+
+		virtual SUTransformation GetMeshBakedTransform() override;
 		// End FDefinition
 
 	private:
@@ -323,7 +335,7 @@ namespace DatasmithSketchUp
 			bGeometryInvalidated = true;
 		}
 
-		void UpdateEntityProperties(FExportContext& Context);
+		virtual void UpdateEntityProperties(FExportContext& Context);
 		void UpdateEntityGeometry(FExportContext& Context);
 
 		TSet<FNodeOccurence*> VisibleNodes;
@@ -345,18 +357,21 @@ namespace DatasmithSketchUp
 		{}
 
 		// >>> FEntity
-		FDefinition* GetDefinition() override;
-		bool GetAssignedMaterial(FMaterialIDType& MaterialId) override;
-		void InvalidateOccurrencesGeometry(FExportContext& Context) override;
-		void UpdateOccurrence(FExportContext& Context, FNodeOccurence& Node) override;
-		void InvalidateOccurrencesProperties(FExportContext& Context) override;
-		int64 GetPersistentId() override;
-		FString GetName() override;
-		void UpdateOccurrenceVisibility(FExportContext& Context, FNodeOccurence&) override;
-		void DeleteOccurrence(FExportContext& Context, FNodeOccurence* Node) override;
-		void UpdateMetadata(FExportContext& Context) override;
+		virtual FDefinition* GetDefinition() override;
+		virtual bool GetAssignedMaterial(FMaterialIDType& MaterialId) override;
+		virtual void InvalidateOccurrencesGeometry(FExportContext& Context) override;
+		virtual void UpdateOccurrence(FExportContext& Context, FNodeOccurence& Node) override;
+		virtual void InvalidateOccurrencesProperties(FExportContext& Context) override;
+		virtual int64 GetPersistentId() override;
+		virtual FString GetName() override;
+		virtual void UpdateOccurrenceVisibility(FExportContext& Context, FNodeOccurence&) override;
+		virtual void DeleteOccurrence(FExportContext& Context, FNodeOccurence* Node) override;
+		virtual void UpdateMetadata(FExportContext& Context) override;
+		virtual void UpdateEntityProperties(FExportContext& Context) override;
 		// <<< FEntity
 
+
+
 		// Set Definition which Entities contains this entity
 		void SetParentDefinition(FExportContext& Context, FDefinition* InParent);
 		bool IsParentDefinition(FDefinition* InParent)
diff --git a/Engine/Source/Programs/Enterprise/Datasmith/DatasmithSketchUpRubyExporter/Private/DatasmithSketchUpMesh.cpp b/Engine/Source/Programs/Enterprise/Datasmith/DatasmithSketchUpRubyExporter/Private/DatasmithSketchUpMesh.cpp
--- a/Engine/Source/Programs/Enterprise/Datasmith/DatasmithSketchUpRubyExporter/Private/DatasmithSketchUpMesh.cpp
+++ b/Engine/Source/Programs/Enterprise/Datasmith/DatasmithSketchUpRubyExporter/Private/DatasmithSketchUpMesh.cpp
@@ -22,6 +22,7 @@
 #include "SketchUpAPI/model/layer.h"
 #include "SketchUpAPI/model/mesh_helper.h"
 #include "SketchUpAPI/model/uv_helper.h"
+#include "SketchUpAPI/geometry/point3d.h"
 #include "DatasmithSketchUpSDKCeases.h"
 
 // Datasmith SDK.
@@ -64,7 +65,7 @@ namespace DatasmithSketchUp
 	public:
 
 		// Convert the combined mesh into a Datasmith mesh.
-		void ConvertMeshToDatasmith(FDatasmithMesh& OutDMesh) const;
+		void ConvertMeshToDatasmith(SUTransformation Transform, FDatasmithMesh& OutDMesh) const;
 
 		// Tessellate a SketchUp face into a triangle mesh merged into the combined mesh.
 		void AddFace(
@@ -295,7 +296,7 @@ namespace DatasmithSketchUp
 		}
 	}
 
-	void FDatasmithSketchUpMesh::ConvertMeshToDatasmith(FDatasmithMesh& OutDMesh) const
+	void FDatasmithSketchUpMesh::ConvertMeshToDatasmith(SUTransformation Transform, FDatasmithMesh& OutDMesh) const
 	{
 		// Get the number of mesh vertices (must be > 0).
 		int32 VertexCount = MeshVertexPoints.Num();
@@ -305,7 +306,9 @@ namespace DatasmithSketchUp
 
 		for (int32 VertexNo = 0; VertexNo < VertexCount; VertexNo++)
 		{
-			SUPoint3D const& VertexPoint = MeshVertexPoints[VertexNo];
+			SUPoint3D VertexPoint = MeshVertexPoints[VertexNo];
+			SUPoint3DTransform(&Transform, &VertexPoint);
+			
 			FVector P = FVector(DatasmithSketchUpUtils::FromSketchUp::ConvertPosition(VertexPoint));
 			OutDMesh.SetVertex(VertexNo, P.X, P.Y, P.Z);
 		}
@@ -437,14 +440,16 @@ void FEntities::UpdateGeometry(FExportContext& Context)
 					Mesh->DatasmithMesh->SetMaterial(Material->GetName(), SlotId);
 				}
 			}
+			
+			SUTransformation Transform = Definition.GetMeshBakedTransform();
 
 			Context.MeshExportTasks.Emplace(Async(
 				EAsyncExecution::ThreadPool,
-				[Mesh, ExtractedMeshPtr, &Context]()
+				[&Context, Mesh, ExtractedMeshPtr, Transform]()
 				{
 					FDatasmithMeshExporter DatasmithMeshExporter;
 					FDatasmithMesh DatasmithMesh;
-					ExtractedMeshPtr->ConvertMeshToDatasmith(DatasmithMesh);
+					ExtractedMeshPtr->ConvertMeshToDatasmith(Transform, DatasmithMesh);
 
 					FGCScopeGuard GCGuard; // Prevent GC from running while UDatasmithMesh is created in ExportToUObject. 
 					return DatasmithMeshExporter.ExportToUObject(Mesh->DatasmithMesh, Context.GetAssetsOutputPath(), DatasmithMesh, nullptr, FDatasmithExportOptions::LightmapUV);
