[SketchUp] Layer folder(group) tracking for individual faces

- also optimized(cached) layer visibility

diff --git a/Engine/Source/Programs/Enterprise/Datasmith/DatasmithSketchUpRubyExporter/Private/DatasmithSketchUpCommon.h b/Engine/Source/Programs/Enterprise/Datasmith/DatasmithSketchUpRubyExporter/Private/DatasmithSketchUpCommon.h
--- a/Engine/Source/Programs/Enterprise/Datasmith/DatasmithSketchUpRubyExporter/Private/DatasmithSketchUpCommon.h
+++ b/Engine/Source/Programs/Enterprise/Datasmith/DatasmithSketchUpRubyExporter/Private/DatasmithSketchUpCommon.h
@@ -39,6 +39,7 @@ namespace DatasmithSketchUp
 	typedef FEntityIDType FMaterialIDType;
 	typedef FEntityIDType FSceneIDType;
 	typedef FEntityIDType FTextureIDType;
+	typedef FEntityIDType FLayerIDType;
 
 	static FORCEINLINE uint32 GetTypeHash(FEntityIDType Key)
 	{
diff --git a/Engine/Source/Programs/Enterprise/Datasmith/DatasmithSketchUpRubyExporter/Private/DatasmithSketchUpComponent.cpp b/Engine/Source/Programs/Enterprise/Datasmith/DatasmithSketchUpRubyExporter/Private/DatasmithSketchUpComponent.cpp
--- a/Engine/Source/Programs/Enterprise/Datasmith/DatasmithSketchUpRubyExporter/Private/DatasmithSketchUpComponent.cpp
+++ b/Engine/Source/Programs/Enterprise/Datasmith/DatasmithSketchUpRubyExporter/Private/DatasmithSketchUpComponent.cpp
@@ -927,25 +927,7 @@ void FComponentInstance::InvalidateOccurrencesProperties(FExportContext& Context
 	SUDrawingElementGetHidden(DrawingElementRef, &bNewHidden);
 
 	SUDrawingElementGetLayer(DrawingElementRef, &LayerRef);
-	bool bNewLayerVisible = true;
-	SULayerGetVisibility(LayerRef, &bNewLayerVisible);
-
-	// Search for invisible ancestor folder (parent invisibility overrides child's visibility) 
-	// LayerFolder introduced in SketchUp 2021
-#if !defined(SKP_SDK_2019) && !defined(SKP_SDK_2020)
-	SULayerFolderRef LayerFolderRef = SU_INVALID;
-	SULayerGetParentLayerFolder(LayerRef, &LayerFolderRef);
-	while (SUIsValid(LayerFolderRef))
-	{
-		bool bLayerFolderVisible = true;
-		SULayerFolderGetVisibility(LayerFolderRef, &bLayerFolderVisible);
-		bNewLayerVisible = bNewLayerVisible && bLayerFolderVisible;
-
-		SULayerFolderRef ParentLayerFolderRef = SU_INVALID;
-		SULayerFolderGetParentLayerFolder(LayerFolderRef, &ParentLayerFolderRef);
-		LayerFolderRef = ParentLayerFolderRef;
-	}
-#endif
+	bool bNewLayerVisible = Context.Layers.IsLayerVisible(LayerRef);
 
 	if (bHidden != bNewHidden || bLayerVisible != bNewLayerVisible)
 	{
diff --git a/Engine/Source/Programs/Enterprise/Datasmith/DatasmithSketchUpRubyExporter/Private/DatasmithSketchUpExportContext.cpp b/Engine/Source/Programs/Enterprise/Datasmith/DatasmithSketchUpRubyExporter/Private/DatasmithSketchUpExportContext.cpp
--- a/Engine/Source/Programs/Enterprise/Datasmith/DatasmithSketchUpRubyExporter/Private/DatasmithSketchUpExportContext.cpp
+++ b/Engine/Source/Programs/Enterprise/Datasmith/DatasmithSketchUpRubyExporter/Private/DatasmithSketchUpExportContext.cpp
@@ -57,9 +57,8 @@ FExportContext::FExportContext()
 	, Materials(*this)
 	, Scenes(*this)
 	, Textures(*this)
-
+	, Layers(*this)
 {
-
 }
 
 const TCHAR* FExportContext::GetAssetsOutputPath() const
@@ -98,6 +97,7 @@ void FExportContext::Populate()
 	RootNode->DatasmithActorLabel = TEXT("Model");
 
 	// Parse/convert Model
+	Layers.PopulateFromModel(ModelRef);
 	Scenes.PopulateFromModel(ModelRef);
 	ComponentDefinitions.PopulateFromModel(ModelRef);
 
@@ -206,6 +206,32 @@ void FSceneCollection::PopulateFromModel(SUModelRef InModelRef)
 	}
 }
 
+void FLayerCollection::PopulateFromModel(SUModelRef InModelRef)
+{
+	size_t LayerCount = 0;
+	SUModelGetNumLayers(InModelRef, &LayerCount);
+
+	TArray<SULayerRef> Layers;
+	Layers.Init(SU_INVALID, LayerCount);
+	SUResult SResult = SUModelGetLayers(InModelRef, LayerCount, Layers.GetData(), &LayerCount);
+	Layers.SetNum(LayerCount);
+
+	for (SULayerRef LayerRef : Layers)
+	{
+		UpdateLayer(LayerRef);
+	}
+}
+
+void FLayerCollection::UpdateLayer(SULayerRef LayerRef)
+{
+	LayerVisibility.FindOrAdd(DatasmithSketchUpUtils::GetEntityID(SULayerToEntity(LayerRef))) = DatasmithSketchUpUtils::IsLayerVisible(LayerRef);
+}
+
+bool FLayerCollection::IsLayerVisible(SULayerRef LayerRef)
+{
+	bool* Found = LayerVisibility.Find(DatasmithSketchUpUtils::GetEntityID(SULayerToEntity(LayerRef)));
+	return Found ? *Found : true;
+}
 
 void FComponentDefinitionCollection::PopulateFromModel(SUModelRef InModelRef)
 {
diff --git a/Engine/Source/Programs/Enterprise/Datasmith/DatasmithSketchUpRubyExporter/Private/DatasmithSketchUpExportContext.h b/Engine/Source/Programs/Enterprise/Datasmith/DatasmithSketchUpRubyExporter/Private/DatasmithSketchUpExportContext.h
--- a/Engine/Source/Programs/Enterprise/Datasmith/DatasmithSketchUpRubyExporter/Private/DatasmithSketchUpExportContext.h
+++ b/Engine/Source/Programs/Enterprise/Datasmith/DatasmithSketchUpRubyExporter/Private/DatasmithSketchUpExportContext.h
@@ -208,6 +208,21 @@ namespace DatasmithSketchUp
 		FMaterialOccurrence DefaultMaterial;
 	};
 
+	// Tracks information related to SketchUp "Tags"/"Layers"
+	class FLayerCollection
+	{
+	public:
+		FLayerCollection(FExportContext& InContext) : Context(InContext) {}
+
+		void PopulateFromModel(SUModelRef InSModelRef);
+		void UpdateLayer(SULayerRef LayerRef);
+		bool IsLayerVisible(SULayerRef LayerRef);
+
+		TMap<FLayerIDType, bool> LayerVisibility;
+	private:
+		FExportContext& Context;
+	};
+
 	// Holds all the data needed during export and incremental updates
 	class FExportContext
 	{
@@ -238,6 +253,7 @@ namespace DatasmithSketchUp
 		FMaterialCollection Materials;
 		FSceneCollection Scenes;
 		FTextureCollection Textures;
+		FLayerCollection Layers;
 
 		TArray<TFuture<bool>> MeshExportTasks;
 	};
diff --git a/Engine/Source/Programs/Enterprise/Datasmith/DatasmithSketchUpRubyExporter/Private/DatasmithSketchUpExporter.cpp b/Engine/Source/Programs/Enterprise/Datasmith/DatasmithSketchUpRubyExporter/Private/DatasmithSketchUpExporter.cpp
--- a/Engine/Source/Programs/Enterprise/Datasmith/DatasmithSketchUpRubyExporter/Private/DatasmithSketchUpExporter.cpp
+++ b/Engine/Source/Programs/Enterprise/Datasmith/DatasmithSketchUpRubyExporter/Private/DatasmithSketchUpExporter.cpp
@@ -513,6 +513,8 @@ public:
 		DatasmithSketchUp::FEntityIDType LayerId = DatasmithSketchUpUtils::GetEntityID(Entity);
 		Context.ComponentInstances.LayerModified(LayerId);
 		Context.EntitiesObjects.LayerModified(LayerId);
+		Context.Layers.UpdateLayer(SULayerFromEntity(Entity));
+
 		return true;
 	}
 
diff --git a/Engine/Source/Programs/Enterprise/Datasmith/DatasmithSketchUpRubyExporter/Private/DatasmithSketchUpMesh.cpp b/Engine/Source/Programs/Enterprise/Datasmith/DatasmithSketchUpRubyExporter/Private/DatasmithSketchUpMesh.cpp
--- a/Engine/Source/Programs/Enterprise/Datasmith/DatasmithSketchUpRubyExporter/Private/DatasmithSketchUpMesh.cpp
+++ b/Engine/Source/Programs/Enterprise/Datasmith/DatasmithSketchUpRubyExporter/Private/DatasmithSketchUpMesh.cpp
@@ -375,7 +375,7 @@ const TCHAR* FEntitiesGeometry::GetMeshElementName(int32 MeshIndex)
 	return Meshes[MeshIndex]->DatasmithMesh->GetName();
 }
 
-void ScanSketchUpEntitiesFaces(SUEntitiesRef EntitiesRef, FEntitiesGeometry& Geometry, TFunctionRef<void(TSharedPtr<FDatasmithSketchUpMesh> ExtractedMesh)> OnNewExtractedMesh);
+void ScanSketchUpEntitiesFaces(FExportContext& Context, SUEntitiesRef EntitiesRef, FEntitiesGeometry& Geometry, TFunctionRef<void(TSharedPtr<FDatasmithSketchUpMesh> ExtractedMesh)> OnNewExtractedMesh);
 
 void FEntities::UpdateGeometry(FExportContext& Context)
 {
@@ -458,7 +458,7 @@ void FEntities::UpdateGeometry(FExportContext& Context)
 		}
 	};
 
-	ScanSketchUpEntitiesFaces(EntitiesRef, *EntitiesGeometry, ProcessExtractedMesh);
+	ScanSketchUpEntitiesFaces(Context, EntitiesRef, *EntitiesGeometry, ProcessExtractedMesh);
 	EntitiesGeometry->Meshes.SetNum(MeshCount);
 
 	Context.EntitiesObjects.RegisterEntities(*this);
@@ -523,7 +523,7 @@ TArray<SUComponentInstanceRef> FEntities::GetComponentInstances()
 }
 
 
-void ScanSketchUpEntitiesFaces(SUEntitiesRef EntitiesRef, FEntitiesGeometry& Geometry, TFunctionRef<void(TSharedPtr<FDatasmithSketchUpMesh> ExtractedMesh)> OnNewExtractedMesh)
+void ScanSketchUpEntitiesFaces(FExportContext& Context, SUEntitiesRef EntitiesRef, FEntitiesGeometry& Geometry, TFunctionRef<void(TSharedPtr<FDatasmithSketchUpMesh> ExtractedMesh)> OnNewExtractedMesh)
 {
 	// Get the number of faces in the source SketchUp entities.
 	size_t SFaceCount = 0;
@@ -582,9 +582,10 @@ void ScanSketchUpEntitiesFaces(SUEntitiesRef EntitiesRef, FEntitiesGeometry& Geo
 			SUDrawingElementGetLayer(SUFaceToDrawingElement(SScannedFaceRef), &LayerRef);
 			Geometry.Layers.Add(DatasmithSketchUpUtils::GetEntityID(SULayerToEntity(LayerRef)));
 
+			bool bFaceHidden = false;
+			SUDrawingElementGetHidden(SUFaceToDrawingElement(SScannedFaceRef), &bFaceHidden);
 
-			// Get whether or not the SketckUp face is visible in the current SketchUp scene.
-			if (DatasmithSketchUpUtils::IsVisible(SScannedFaceRef))
+			if (!bFaceHidden && Context.Layers.IsLayerVisible(LayerRef))
 			{
 				ExtractedMesh.AddFace(SScannedFaceRef);
 			}
diff --git a/Engine/Source/Programs/Enterprise/Datasmith/DatasmithSketchUpRubyExporter/Private/DatasmithSketchUpUtils.cpp b/Engine/Source/Programs/Enterprise/Datasmith/DatasmithSketchUpRubyExporter/Private/DatasmithSketchUpUtils.cpp
--- a/Engine/Source/Programs/Enterprise/Datasmith/DatasmithSketchUpRubyExporter/Private/DatasmithSketchUpUtils.cpp
+++ b/Engine/Source/Programs/Enterprise/Datasmith/DatasmithSketchUpRubyExporter/Private/DatasmithSketchUpUtils.cpp
@@ -21,6 +21,9 @@
 #include "SketchUpAPI/model/geometry.h"
 #include "SketchUpAPI/model/group.h"
 #include "SketchUpAPI/model/layer.h"
+#if !defined(SKP_SDK_2019) && !defined(SKP_SDK_2020)
+#include "SketchUpAPI/model/layer_folder.h"
+#endif
 #include "SketchUpAPI/model/material.h"
 #include "SketchUpAPI/model/mesh_helper.h"
 #include "SketchUpAPI/model/model.h"
@@ -148,26 +151,30 @@ bool DatasmithSketchUpUtils::IsVisible(
 	return (!bSComponentInstanceHidden && bSEffectiveLayerVisible);
 }
 
-bool DatasmithSketchUpUtils::IsVisible(
-	SUFaceRef  InFaceRef
+bool DatasmithSketchUpUtils::IsLayerVisible(
+	SULayerRef LayerRef
 )
 {
-	// Get the flag indicating whether or not the SketchUp face is hidden.
-	bool bFaceHidden = false;
-	SUDrawingElementGetHidden(SUFaceToDrawingElement(InFaceRef), &bFaceHidden); // we can ignore the returned SU_RESULT
-
-	// Retrieve the SketckUp face layer.
-	SULayerRef SFaceLayerRef = SU_INVALID;
-	SUDrawingElementGetLayer(SUFaceToDrawingElement(InFaceRef), &SFaceLayerRef); // we can ignore the returned SU_RESULT
-
-	// Retrieve the SketchUp face layer name.
-	FString SFaceLayerName = SuGetString(SULayerGetName, SFaceLayerRef);
-
-	// Get the flag indicating whether or not the SketchUp face effective layer is visible.
-	bool bSFaceEffectiveLayerVisible = true;
-	SULayerGetVisibility(SFaceLayerRef, &bSFaceEffectiveLayerVisible); // we can ignore the returned SU_RESULT
+	bool bVisible = true;
+	SULayerGetVisibility(LayerRef, &bVisible);
+
+	// Search for invisible ancestor folder (parent invisibility overrides child's visibility) 
+	// LayerFolder introduced in SketchUp 2021
+#if !defined(SKP_SDK_2019) && !defined(SKP_SDK_2020)
+	SULayerFolderRef LayerFolderRef = SU_INVALID;
+	SULayerGetParentLayerFolder(LayerRef, &LayerFolderRef);
+	while (bVisible && SUIsValid(LayerFolderRef))
+	{
+		bool bLayerFolderVisible = true;
+		SULayerFolderGetVisibility(LayerFolderRef, &bLayerFolderVisible);
+		bVisible = bVisible && bLayerFolderVisible;
 
-	return (!bFaceHidden && bSFaceEffectiveLayerVisible);
+		SULayerFolderRef ParentLayerFolderRef = SU_INVALID;
+		SULayerFolderGetParentLayerFolder(LayerFolderRef, &ParentLayerFolderRef);
+		LayerFolderRef = ParentLayerFolderRef;
+	}
+#endif
+	return bVisible;
 }
 
 SUMaterialRef DatasmithSketchUpUtils::GetMaterial(
diff --git a/Engine/Source/Programs/Enterprise/Datasmith/DatasmithSketchUpRubyExporter/Private/DatasmithSketchUpUtils.h b/Engine/Source/Programs/Enterprise/Datasmith/DatasmithSketchUpRubyExporter/Private/DatasmithSketchUpUtils.h
--- a/Engine/Source/Programs/Enterprise/Datasmith/DatasmithSketchUpRubyExporter/Private/DatasmithSketchUpUtils.h
+++ b/Engine/Source/Programs/Enterprise/Datasmith/DatasmithSketchUpRubyExporter/Private/DatasmithSketchUpUtils.h
@@ -85,9 +85,9 @@ namespace DatasmithSketchUpUtils
 		SULayerRef             InSEffectiveLayerRef     // SketchUp component instance effective layer
 	);
 
-	// Return whether or not a SketckUp face is visible in the current SketchUp scene.
-	bool IsVisible(
-		SUFaceRef  InSFaceRef          // valid SketchUp face
+	// Return whether or not a SketckUp layer is visible in the current SketchUp scene taking into account folder visibility
+	bool IsLayerVisible(
+		SULayerRef LayerRef
 	);
 
 	// Get the material of a SketckUp component instance.
