[SketchUp] "Color By Tag"

- Switching Color By Tag rebuilds materials for the whole scene to use SU Layer(aka Tag)
materials
- Also includes texture/uv scaling information of Layer materials.
- Layer/Tag override works as in SU - default(Layer0) layer material overriden if component instance overrides Layer

diff --git a/Engine/Source/Programs/Enterprise/Datasmith/DatasmithSketchUpRubyExporter/Private/DatasmithSketchUpComponent.cpp b/Engine/Source/Programs/Enterprise/Datasmith/DatasmithSketchUpRubyExporter/Private/DatasmithSketchUpComponent.cpp
--- a/Engine/Source/Programs/Enterprise/Datasmith/DatasmithSketchUpRubyExporter/Private/DatasmithSketchUpComponent.cpp
+++ b/Engine/Source/Programs/Enterprise/Datasmith/DatasmithSketchUpRubyExporter/Private/DatasmithSketchUpComponent.cpp
@@ -726,7 +726,6 @@ void FEntity::UpdateOccurrence(FExportContext& Context, FNodeOccurence& Node)
 		Node.InheritedMaterialID = Node.ParentNode->InheritedMaterialID;
 	}
 
-
 	FString MeshActorLabel = Node.GetActorLabel();
 	// Update Datasmith Mesh Actors
 	for (int32 MeshIndex = 0; MeshIndex < Node.MeshActors.Num(); ++MeshIndex)
@@ -739,13 +738,7 @@ void FEntity::UpdateOccurrence(FExportContext& Context, FNodeOccurence& Node)
 		// todo: set inherited material only on mesh actors that have faces with default material, right now setting on every mesh, hot harmful but excessive
 		if (EntitiesGeometry.IsMeshUsingInheritedMaterial(MeshIndex))
 		{
-			if (FMaterialOccurrence* Material = Context.Materials.RegisterInstance(Node.InheritedMaterialID, &Node))
-			{
-				// SketchUp has 'material override' only for single('Default') material. 
-				// So we reset overrides on the actor to remove this single override(if it was set) and re-add new override
-				MeshActor->ResetMaterialOverrides(); // Clear previous override if was set
-				MeshActor->AddMaterialOverride(Material->GetName(), EntitiesGeometry.GetInheritedMaterialOverrideSlotId());
-			}
+			Context.Materials.SetMeshActorOverrideMaterial(Node, EntitiesGeometry, MeshActor);
 		}
 	}
 }
@@ -817,6 +810,8 @@ void FComponentInstance::UpdateOccurrence(FExportContext& Context, FNodeOccurenc
 		return;
 	}
 
+	Node.EffectiveLayerRef = DatasmithSketchUpUtils::GetEffectiveLayer(GetComponentInstanceRef(), Node.ParentNode->EffectiveLayerRef);
+
 	if (FDefinition* EntityDefinition = GetDefinition())
 	{
 		EntityDefinition->BuildNodeNames(Node);
@@ -996,8 +991,6 @@ void FComponentInstance::FillOccurrenceActorMetadata(FNodeOccurence& Node)
 
 void FComponentInstance::UpdateOccurrenceVisibility(FExportContext& Context, FNodeOccurence& Node)
 {
-	Node.EffectiveLayerRef = DatasmithSketchUpUtils::GetEffectiveLayer(GetComponentInstanceRef(), Node.ParentNode->EffectiveLayerRef);
-
 	// Parent node, component instance and layer - all should be visible to have node visible
 	Node.SetVisibility(Node.ParentNode->bVisible && !bHidden && bLayerVisible);
 
diff --git a/Engine/Source/Programs/Enterprise/Datasmith/DatasmithSketchUpRubyExporter/Private/DatasmithSketchUpComponent.h b/Engine/Source/Programs/Enterprise/Datasmith/DatasmithSketchUpRubyExporter/Private/DatasmithSketchUpComponent.h
--- a/Engine/Source/Programs/Enterprise/Datasmith/DatasmithSketchUpRubyExporter/Private/DatasmithSketchUpComponent.h
+++ b/Engine/Source/Programs/Enterprise/Datasmith/DatasmithSketchUpRubyExporter/Private/DatasmithSketchUpComponent.h
@@ -291,6 +291,17 @@ namespace DatasmithSketchUp
 		TSet<FMaterial*> MaterialsUsed;
 		bool bDefaultMaterialUsed = false;
 
+		bool IsDefaultLayerMaterialUsed()
+		{
+			return bDefaultMaterialUsed;
+		}
+
+		void SetDefaultLayerMaterialUsed()
+		{
+			bDefaultMaterialUsed = true;
+		}
+
+
 		// todo: update reusing datasmith elements? 
 		// todo: merge ALL faces that are present in Entities into single mesh? do we really need separate mesh for every isolated set of faces?
 		// todo: occurrences using these entities must ne referenced
diff --git a/Engine/Source/Programs/Enterprise/Datasmith/DatasmithSketchUpRubyExporter/Private/DatasmithSketchUpExportContext.cpp b/Engine/Source/Programs/Enterprise/Datasmith/DatasmithSketchUpRubyExporter/Private/DatasmithSketchUpExportContext.cpp
--- a/Engine/Source/Programs/Enterprise/Datasmith/DatasmithSketchUpRubyExporter/Private/DatasmithSketchUpExportContext.cpp
+++ b/Engine/Source/Programs/Enterprise/Datasmith/DatasmithSketchUpRubyExporter/Private/DatasmithSketchUpExportContext.cpp
@@ -29,19 +29,12 @@
 #include "SketchUpAPI/model/camera.h"
 #include "SketchUpAPI/model/component_definition.h"
 #include "SketchUpAPI/model/component_instance.h"
-#include "SketchUpAPI/model/drawing_element.h"
-#include "SketchUpAPI/model/edge.h"
-#include "SketchUpAPI/model/entities.h"
-#include "SketchUpAPI/model/entity.h"
-#include "SketchUpAPI/model/face.h"
-#include "SketchUpAPI/model/geometry.h"
-#include "SketchUpAPI/model/group.h"
 #include "SketchUpAPI/model/layer.h"
-#include "SketchUpAPI/model/mesh_helper.h"
 #include "SketchUpAPI/model/model.h"
+#include <SketchUpAPI/model/rendering_options.h>
 #include "SketchUpAPI/model/scene.h"
-#include "SketchUpAPI/model/texture.h"
-#include "SketchUpAPI/model/uv_helper.h"
+#include <SketchUpAPI/model/typed_value.h>
+
 
 #include "SketchUpAPI/application/application.h"
 
@@ -104,8 +97,75 @@ void FExportContext::Populate()
 	RootNode->ToDatasmith(*this);
 }
 
-void FExportContext::Update()
+// ColorByLayer(or ColorByTag) changes material display on every entity by useinf either regular materials or colors(materials) set up for layers
+bool FExportContext::PreUpdateColorByLayer()
 {
+	if (bColorByLayerInvaliated)
+	{
+		bColorByLayerInvaliated = false;
+
+		bool bColorByLayerNew = bColorByLayer;
+		SURenderingOptionsRef RenderingOptionsRef = SU_INVALID;
+		if (SUModelGetRenderingOptions(ModelRef, &RenderingOptionsRef) == SU_ERROR_NONE)
+		{
+			SUTypedValueRef DisplayColorByLayerTypedValue = SU_INVALID;
+			SUTypedValueCreate(&DisplayColorByLayerTypedValue);
+			if (SURenderingOptionsGetValue(RenderingOptionsRef, "DisplayColorByLayer", &DisplayColorByLayerTypedValue) == SU_ERROR_NONE)
+			{
+				bool DisplayColorByLayer;
+				if (SUTypedValueGetBool(DisplayColorByLayerTypedValue, &DisplayColorByLayer) == SU_ERROR_NONE)
+				{
+					bColorByLayerNew = DisplayColorByLayer;
+				}
+			}
+		}
+
+		// Check that flag was actually toggled to different state since last update
+		if (bColorByLayer != bColorByLayerNew)
+		{
+			// Invalidate essentially everything when mode changed
+			//    - layer materials vs regular materials need to be build/removed (material rebuild is caused by nodes/geom rebuild)
+			//	  - meshes need to be rebuild(to split by layer, not by regular material)
+			//	  - mesh actors need rebuild tochange override materials
+			for (TPair<FComponentInstanceIDType, TSharedPtr<FComponentInstance>> Instance : ComponentInstances.ComponentInstanceMap)
+			{
+				Instance.Value->InvalidateEntityProperties();
+			}
+			
+			for (const auto& IdValue : ComponentDefinitions.ComponentDefinitionMap)
+			{
+				TSharedPtr<FComponentDefinition> Definition = IdValue.Value;
+				Definition->InvalidateDefinitionGeometry();
+			}
+			ModelDefinition->InvalidateDefinitionGeometry();
+
+			bColorByLayer = bColorByLayerNew;
+			return true;
+		}
+	}
+	else
+	{
+		if (bColorByLayer)
+		{
+			return Materials.LayerMaterials.CheckForModifications();
+		}
+	}
+	return false;
+}
+
+
+bool FExportContext::Update(bool bModifiedHint)
+{
+	bool bModified = bModifiedHint;
+
+	bModified |= PreUpdateColorByLayer();
+
+	if (!bModified)
+	{
+		// code below is not supposed to change Datasmith scene if not modification hint(i.e. no invalidated data) was present
+		return false;
+	}
+
 	// Invalidate occurrences for changed instances first
 	Model->UpdateEntityProperties(*this);
 	ComponentInstances.UpdateProperties(); 
@@ -132,6 +192,8 @@ void FExportContext::Update()
 		Task.Get();
 	}
 	MeshExportTasks.Reset();
+
+	return bModified;
 }
 
 FDefinition* FExportContext::GetDefinition(SUEntityRef Entity)
@@ -165,6 +227,12 @@ FDefinition* FExportContext::GetDefinition(FEntityIDType DefinitionEntityId)
 	return nullptr;
 }
 
+bool FExportContext::InvalidateColorByLayer()
+{
+	bColorByLayerInvaliated = true;
+	return false; // Don't return 'modified'  - in case user toggles the flag back and forth Update will check if it's actually changed
+}
+
 void FComponentDefinitionCollection::Update()
 {
 	for (const auto& IdValue : ComponentDefinitionMap)
@@ -233,6 +301,33 @@ bool FLayerCollection::IsLayerVisible(SULayerRef LayerRef)
 	return Found ? *Found : true;
 }
 
+
+
+SULayerRef FLayerCollection::GetLayer(FLayerIDType LayerId)
+{
+	size_t LayerCount = 0;
+	SUModelGetNumLayers(Context.ModelRef, &LayerCount);
+
+	TArray<SULayerRef> Layers;
+	Layers.Init(SU_INVALID, LayerCount);
+	SUResult SResult = SUModelGetLayers(Context.ModelRef, LayerCount, Layers.GetData(), &LayerCount);
+	Layers.SetNum(LayerCount);
+
+	for (SULayerRef LayerRef : Layers)
+	{
+		if (GetLayerId(LayerRef) == LayerId)
+		{
+			return LayerRef;
+		}
+	}
+	return SU_INVALID;
+}
+
+FLayerIDType FLayerCollection::GetLayerId(SULayerRef LayerRef)
+{
+	return DatasmithSketchUpUtils::GetEntityID(SULayerToEntity(LayerRef));
+}
+
 void FComponentDefinitionCollection::PopulateFromModel(SUModelRef InModelRef)
 {
 	// Get the number of normal component definitions in the SketchUp model.
@@ -492,30 +587,7 @@ void FComponentInstanceCollection::LayerModified(DatasmithSketchUp::FEntityIDTyp
 	}
 }
 
-void FMaterialCollection::PopulateFromModel(SUModelRef InModelRef)
-{
-	// Get the number of material definitions in the SketchUp model.
-	size_t SMaterialDefinitionCount = 0;
-	SUModelGetNumMaterials(InModelRef, &SMaterialDefinitionCount); // we can ignore the returned SU_RESULT
-
-	if (SMaterialDefinitionCount > 0)
-	{
-		// Retrieve the material definitions in the SketchUp model.
-		TArray<SUMaterialRef> SMaterialDefinitions;
-		SMaterialDefinitions.Init(SU_INVALID, SMaterialDefinitionCount);
-		SUModelGetMaterials(InModelRef, SMaterialDefinitionCount, SMaterialDefinitions.GetData(), &SMaterialDefinitionCount); // we can ignore the returned SU_RESULT
-		SMaterialDefinitions.SetNum(SMaterialDefinitionCount);
-
-		// Add the material definitions to our dictionary.
-		for (SUMaterialRef SMaterialDefinitionRef : SMaterialDefinitions)
-		{
-
-			CreateMaterial(SMaterialDefinitionRef);
-		}
-	}
-}
-
-FMaterialOccurrence* FMaterialCollection::RegisterInstance(FMaterialIDType MaterialID, FNodeOccurence* NodeOccurrence)
+FMaterialOccurrence* FRegularMaterials::RegisterInstance(FMaterialIDType MaterialID, FNodeOccurence* NodeOccurrence)
 {
 	if (NodeOccurrence->MaterialOverride)
 	{
@@ -527,10 +599,19 @@ FMaterialOccurrence* FMaterialCollection::RegisterInstance(FMaterialIDType Mater
 		return &Material->RegisterInstance(NodeOccurrence);
 	}
 
-	return nullptr; // Don't use a material if material id is unknown(of default)
+	return {}; // Don't use a material if material id is unknown(of default)
+}
+
+void FRegularMaterials::UnregisterGeometry(DatasmithSketchUp::FEntitiesGeometry* EntitiesGeometry)
+{
+	if (EntitiesGeometry->bDefaultMaterialUsed)
+	{
+		DefaultMaterial.UnregisterGeometry(Context, EntitiesGeometry);
+		EntitiesGeometry->bDefaultMaterialUsed = false;
+	}
 }
 
-FMaterialOccurrence* FMaterialCollection::RegisterGeometry(FMaterialIDType MaterialID, DatasmithSketchUp::FEntitiesGeometry* EntitiesGeometry)
+FMaterialOccurrence* FRegularMaterials::RegisterGeometry(FMaterialIDType MaterialID, DatasmithSketchUp::FEntitiesGeometry* EntitiesGeometry)
 {
 	if (const TSharedPtr<DatasmithSketchUp::FMaterial> Material = FindOrCreateMaterial(MaterialID))
 	{
@@ -547,6 +628,17 @@ FMaterialOccurrence* FMaterialCollection::RegisterGeometry(FMaterialIDType Mater
 	return &DefaultMaterial; 
 }
 
+FMaterialOccurrence* FLayerMaterials::RegisterGeometryForLayer(FLayerIDType LayerID, FEntitiesGeometry* EntitiesGeometry)
+{
+	if (const TSharedPtr<DatasmithSketchUp::FMaterial> Material = FindOrCreateMaterialForLayer(LayerID))
+	{
+		EntitiesGeometry->MaterialsUsed.Add(Material.Get());
+		return &Material->RegisterGeometry(EntitiesGeometry);
+	}
+
+	return {};
+}
+
 void FMaterialCollection::UnregisterGeometry(DatasmithSketchUp::FEntitiesGeometry* EntitiesGeometry)
 {
 	if (!EntitiesGeometry)
@@ -562,11 +654,7 @@ void FMaterialCollection::UnregisterGeometry(DatasmithSketchUp::FEntitiesGeometr
 		Material.UnregisterGeometry(Context, EntitiesGeometry);
 	}
 
-	if (EntitiesGeometry->bDefaultMaterialUsed)
-	{
-		DefaultMaterial.UnregisterGeometry(Context, EntitiesGeometry);
-		EntitiesGeometry->bDefaultMaterialUsed = false;
-	}
+	RegularMaterials.UnregisterGeometry(EntitiesGeometry);
 
 	EntitiesGeometry->MaterialsUsed.Reset();
 }
@@ -574,47 +662,43 @@ void FMaterialCollection::UnregisterGeometry(DatasmithSketchUp::FEntitiesGeometr
 TSharedPtr<FMaterial> FMaterialCollection::CreateMaterial(SUMaterialRef MaterialDefinitionRef)
 {
 	TSharedPtr<FMaterial> Material = FMaterial::Create(Context, MaterialDefinitionRef);
-	MaterialDefinitionMap.Emplace(DatasmithSketchUpUtils::GetMaterialID(MaterialDefinitionRef), Material);
+	Materials.Add(Material.Get());
 	return Material;
 }
 
-TSharedPtr<FMaterial> FMaterialCollection::CreateMaterial(FMaterialIDType MaterialID)
+void FMaterialCollection::SetMeshActorOverrideMaterial(FNodeOccurence& Node, DatasmithSketchUp::FEntitiesGeometry& EntitiesGeometry, const TSharedPtr<IDatasmithMeshActorElement>& MeshActor)
 {
-	// Get the number of material definitions in the SketchUp model.
-	size_t SMaterialDefinitionCount = 0;
-	SUModelGetNumMaterials(Context.ModelRef, &SMaterialDefinitionCount); // we can ignore the returned SU_RESULT
+	// SketchUp has 'material override' only for single material - 'Default' material or material from default('Layer0') layer in ColorByTag mode. 
+	// So we reset overrides on the actor to remove this single override(if it was set) and re-add new override(if there's one)
+	MeshActor->ResetMaterialOverrides();
 
-	// Retrieve the material definitions in the SketchUp model.
-	TArray<SUMaterialRef> SMaterialDefinitions;
-	SMaterialDefinitions.Init(SU_INVALID, SMaterialDefinitionCount);
-	SUModelGetMaterials(Context.ModelRef, SMaterialDefinitionCount, SMaterialDefinitions.GetData(), &SMaterialDefinitionCount); // we can ignore the returned SU_RESULT
-	SMaterialDefinitions.SetNum(SMaterialDefinitionCount);
-
-	// Add the material definitions to our dictionary.
-	for (SUMaterialRef SMaterialDefinitionRef : SMaterialDefinitions)
+	if (Context.bColorByLayer)
 	{
-		if (MaterialID == DatasmithSketchUpUtils::GetMaterialID(SMaterialDefinitionRef))
+		// Retrieve the default layer in the SketchUp model.
+		SULayerRef DefaultLayerRef = SU_INVALID;
+		SUModelGetDefaultLayer(Context.ModelRef, &DefaultLayerRef);
+
+		if (!SUAreEqual(DefaultLayerRef, Node.EffectiveLayerRef)) // Don't apply default layer(Layer0) material as override
 		{
-			return CreateMaterial(SMaterialDefinitionRef);
+			FEntityIDType LayerId = DatasmithSketchUpUtils::GetEntityID(SULayerToEntity(Node.EffectiveLayerRef));
+			if (FMaterialOccurrence* Material = LayerMaterials.RegisterInstance(LayerId, &Node))
+			{
+				MeshActor->AddMaterialOverride(Material->GetName(), EntitiesGeometry.GetInheritedMaterialOverrideSlotId());
+			}
 		}
 	}
-	return nullptr;
-}
-
-void FMaterialCollection::InvalidateMaterial(SUMaterialRef MaterialDefinitionRef)
-{
-	FMaterialIDType MateriadId = DatasmithSketchUpUtils::GetMaterialID(MaterialDefinitionRef);
-
-	if (InvalidateMaterial(MateriadId))
+	else
 	{
-		return;
+		if (FMaterialOccurrence* Material = RegularMaterials.RegisterInstance(Node.InheritedMaterialID, &Node))
+		{
+			MeshActor->AddMaterialOverride(Material->GetName(), EntitiesGeometry.GetInheritedMaterialOverrideSlotId());
+		}
 	}
-	CreateMaterial(MaterialDefinitionRef);
 }
 
-bool FMaterialCollection::InvalidateMaterial(FMaterialIDType MateriadId)
+bool FRegularMaterials::InvalidateMaterial(FMaterialIDType MateriadId)
 {
-	if (TSharedPtr<FMaterial>* Ptr = MaterialDefinitionMap.Find(MateriadId))
+	if (TSharedPtr<FMaterial>* Ptr = MaterialForMaterialId.Find(MateriadId))
 	{
 		FMaterial& Material = **Ptr;
 
@@ -624,29 +708,83 @@ bool FMaterialCollection::InvalidateMaterial(FMaterialIDType MateriadId)
 	return false;
 }
 
-bool FMaterialCollection::RemoveMaterial(FEntityIDType EntityId)
+bool FRegularMaterials::RemoveMaterial(FMaterialIDType EntityId)
 {
 	TSharedPtr<FMaterial> Material;
-	if (MaterialDefinitionMap.RemoveAndCopyValue(EntityId, Material))
+	if (MaterialForMaterialId.RemoveAndCopyValue(EntityId, Material))
 	{
-		Material->Remove(Context);
+		MaterialIdForMaterial.Remove(Material.Get());
+		Context.Materials.RemoveMaterial(Material.Get());
 		return true;
 	}
 	return false;
 }
 
-bool FMaterialCollection::InvalidateDefaultMaterial()
+bool FMaterialCollection::RemoveMaterial(FMaterial* Material)
+{
+	Materials.Remove(Material);
+	Material->Remove(Context);
+	return true;
+}
+
+bool FRegularMaterials::InvalidateDefaultMaterial()
 {
 	DefaultMaterial.Invalidate(Context);
 	return true;
 }
 
-TSharedPtr<FMaterial> FMaterialCollection::FindOrCreateMaterial(FMaterialIDType MaterialID)
+TSharedPtr<FMaterial> FRegularMaterials::FindOrCreateMaterial(FMaterialIDType MaterialId)
+{
+	if (TSharedPtr<FMaterial>* Ptr = MaterialForMaterialId.Find(MaterialId))
+	{
+		return *Ptr;
+	}
+
+	// Get the number of material definitions in the SketchUp model.
+	size_t SMaterialDefinitionCount = 0;
+	SUModelGetNumMaterials(Context.ModelRef, &SMaterialDefinitionCount); // we can ignore the returned SU_RESULT
+
+	// Retrieve the material definitions in the SketchUp model.
+	TArray<SUMaterialRef> SMaterialDefinitions;
+	SMaterialDefinitions.Init(SU_INVALID, SMaterialDefinitionCount);
+	SUModelGetMaterials(Context.ModelRef, SMaterialDefinitionCount, SMaterialDefinitions.GetData(), &SMaterialDefinitionCount); // we can ignore the returned SU_RESULT
+	SMaterialDefinitions.SetNum(SMaterialDefinitionCount);
+
+	// Add the material definitions to our dictionary.
+	for (SUMaterialRef SMaterialDefinitionRef : SMaterialDefinitions)
+	{
+		if (MaterialId == DatasmithSketchUpUtils::GetMaterialID(SMaterialDefinitionRef))
+		{
+			TSharedPtr<FMaterial> Material = Context.Materials.CreateMaterial(SMaterialDefinitionRef);
+			MaterialForMaterialId.Emplace(MaterialId, Material);
+			MaterialIdForMaterial.Add(Material.Get(), MaterialId);
+			return Material;
+		}
+	}
+	return {};
+}
+
+TSharedPtr<FMaterial> FLayerMaterials::FindOrCreateMaterialForLayer(FLayerIDType LayerID)
 {
-	if (TSharedPtr<FMaterial>* Ptr = MaterialDefinitionMap.Find(MaterialID))
+	if (TSharedPtr<FMaterial>* Ptr = MaterialForLayerId.Find(LayerID))
 	{
 		return *Ptr;
 	}
-	return CreateMaterial(MaterialID);
+
+	SULayerRef LayerRef = Context.Layers.GetLayer(LayerID);
+
+	if (SUIsValid(LayerRef))
+	{
+		SUMaterialRef MaterialRef;
+		if (SULayerGetMaterial(LayerRef, &MaterialRef) == SU_ERROR_NONE)
+		{
+			TSharedPtr<FMaterial> Material = Context.Materials.CreateMaterial(MaterialRef);
+			Material->bGeometryHasScalingBakedIntoUvs = false;
+			MaterialForLayerId.Add(LayerID, Material);
+			LayerIdForMaterial.Add(Material.Get(), LayerID);
+			return Material;
+		}
+	}
+	return {};
 }
 
diff --git a/Engine/Source/Programs/Enterprise/Datasmith/DatasmithSketchUpRubyExporter/Private/DatasmithSketchUpExportContext.h b/Engine/Source/Programs/Enterprise/Datasmith/DatasmithSketchUpRubyExporter/Private/DatasmithSketchUpExportContext.h
--- a/Engine/Source/Programs/Enterprise/Datasmith/DatasmithSketchUpRubyExporter/Private/DatasmithSketchUpExportContext.h
+++ b/Engine/Source/Programs/Enterprise/Datasmith/DatasmithSketchUpRubyExporter/Private/DatasmithSketchUpExportContext.h
@@ -94,11 +94,9 @@ namespace DatasmithSketchUp
 
 		void Update();
 
+		TMap<FEntityIDType, TSharedPtr<FComponentDefinition>> ComponentDefinitionMap;
 	private:
 		FExportContext& Context;
-
-		TMap<FEntityIDType, TSharedPtr<FComponentDefinition>> ComponentDefinitionMap;
-
 	};
 
 	class FTextureCollection
@@ -168,44 +166,93 @@ namespace DatasmithSketchUp
 		FExportContext& Context;
 	};
 
-	class FMaterialCollection
+	class FRegularMaterials
 	{
 	public:
-		FMaterialCollection(FExportContext & InContext) : Context(InContext) {}
+		FRegularMaterials(FExportContext& InContext): Context(InContext)
+		{
+		}
 
-		// Initialize the dictionary of material definitions.
-		void PopulateFromModel(
-			SUModelRef InSModelRef // model containing SketchUp material definitions
-		);
-		void Update();
+		TSharedPtr<FMaterial> FindOrCreateMaterial(FMaterialIDType MaterialID);
 
-		TSharedPtr<FMaterial> CreateMaterial(SUMaterialRef SMaterialDefinitionRef);
-		TSharedPtr<FMaterial> CreateMaterial(FMaterialIDType MaterialID);
-		void InvalidateMaterial(SUMaterialRef SMaterialDefinitionRef);
 		bool InvalidateMaterial(FMaterialIDType MateriadId);
-		bool RemoveMaterial(FEntityIDType EntityId);
+		bool RemoveMaterial(FMaterialIDType EntityId);
+
+		// Tell that this materials is assigned directly to a face on the geometry
+		FMaterialOccurrence* RegisterGeometry(FMaterialIDType MaterialID, DatasmithSketchUp::FEntitiesGeometry* EntitiesGeometry);
+		// Tell that this materials is assigned on the node
+		FMaterialOccurrence* RegisterInstance(FMaterialIDType MaterialID, FNodeOccurence* NodeOccurrence);
+
+
+		void UnregisterGeometry(DatasmithSketchUp::FEntitiesGeometry* EntitiesGeometry);
 
 		bool InvalidateDefaultMaterial();
 
+		void RemoveUnused();
+		void UpdateDefaultMaterial();
+		void Apply(FMaterial* Material);
 
-		TSharedPtr<FMaterial> FindOrCreateMaterial(FMaterialIDType MaterialID);
+	private:
+		FExportContext& Context;
 
-		// Tell that this materials is assigned on the node
-		FMaterialOccurrence* RegisterInstance(FMaterialIDType MaterialID, FNodeOccurence* NodeOccurrence);
+		TMap<FMaterialIDType, TSharedPtr<DatasmithSketchUp::FMaterial>> MaterialForMaterialId;
+		TMap<DatasmithSketchUp::FMaterial*, FMaterialIDType> MaterialIdForMaterial;
+		// TMap<FMaterialIDType, TSharedPtr<DatasmithSketchUp::FMaterial>> MaterialDefinitionMap;
+
+		// todo: include default material into whole update cycle(textures, unused)
+		FMaterialOccurrence DefaultMaterial;
+	};
 
-		// Tell that this materials is assigned directly to a face on the geometry
-		FMaterialOccurrence* RegisterGeometry(FMaterialIDType MaterialID, DatasmithSketchUp::FEntitiesGeometry* EntitiesGeometry);
 
-		void UnregisterGeometry(DatasmithSketchUp::FEntitiesGeometry * EntitiesGeometry);
+	class FLayerMaterials
+	{
+	public:
+		FLayerMaterials(FExportContext& InContext): Context(InContext)
+		{
+		}
 
-		const TCHAR* GetDefaultMaterialName();
+		TSharedPtr<FMaterial> FindOrCreateMaterialForLayer(FLayerIDType LayerID);
+
+		FMaterialOccurrence* RegisterGeometryForLayer(FLayerIDType LayerID, FEntitiesGeometry* EntitiesGeometry);
+		void RemoveUnused();
+		void Apply(FMaterial* Material);
+		void UpdateLayer(SULayerRef LayerRef);
+		FMaterialOccurrence* RegisterInstance(FLayerIDType LayerID, FNodeOccurence* Node);
+		bool CheckForModifications();
 
 	private:
 		FExportContext& Context;
 
-		TMap<FMaterialIDType, TSharedPtr<DatasmithSketchUp::FMaterial>> MaterialDefinitionMap;
+		TMap<FLayerIDType, TSharedPtr<DatasmithSketchUp::FMaterial>> MaterialForLayerId;
+		TMap<FLayerIDType, FMD5Hash> MaterialHashForLayerId;
+		TMap<DatasmithSketchUp::FMaterial*, FLayerIDType> LayerIdForMaterial;
+	};
+
+	class FMaterialCollection
+	{
+	public:
+		FMaterialCollection(FExportContext & InContext) : Context(InContext), LayerMaterials(InContext), RegularMaterials(InContext)
+		{
+		}
+
+		void Update();
+
+		TSharedPtr<FMaterial> CreateMaterial(SUMaterialRef SMaterialDefinitionRef);
+		void InvalidateMaterial(SUMaterialRef SMaterialDefinitionRef);
+		bool RemoveMaterial(FMaterial* Material);
+
+		void UnregisterGeometry(DatasmithSketchUp::FEntitiesGeometry* EntitiesGeometry);
+
+		void SetMeshActorOverrideMaterial(FNodeOccurence& Node, DatasmithSketchUp::FEntitiesGeometry& EntitiesGeometry, const TSharedPtr<IDatasmithMeshActorElement>& MeshActor);
+
+	private:
+		FExportContext& Context;
+	public:
+		FLayerMaterials LayerMaterials;
+		FRegularMaterials RegularMaterials;
+	private:
+		TSet<DatasmithSketchUp::FMaterial*> Materials;
 
-		FMaterialOccurrence DefaultMaterial;
 	};
 
 	// Tracks information related to SketchUp "Tags"/"Layers"
@@ -217,6 +264,9 @@ namespace DatasmithSketchUp
 		void PopulateFromModel(SUModelRef InSModelRef);
 		void UpdateLayer(SULayerRef LayerRef);
 		bool IsLayerVisible(SULayerRef LayerRef);
+		SULayerRef GetLayer(FLayerIDType LayerID);
+
+		FLayerIDType GetLayerId(SULayerRef LayerRef);
 
 		TMap<FLayerIDType, bool> LayerVisibility;
 	private:
@@ -233,11 +283,15 @@ namespace DatasmithSketchUp
 		const TCHAR* GetAssetsOutputPath() const;
 
 		void Populate(); // Create Datasmith scene from the Model
-		void Update(); // Update Datasmith scene to reflect iterative changes done to the Model 
+		bool Update(bool bModifiedHint); // Update Datasmith scene to reflect iterative changes done to the Model 
 
 		FDefinition* GetDefinition(SUEntityRef Entity);
 		FDefinition* GetDefinition(FEntityIDType DefinitionEntityId);
 
+		bool InvalidateColorByLayer();
+
+		bool PreUpdateColorByLayer();
+
 		SUModelRef ModelRef = SU_INVALID;
 
 		TSharedPtr<IDatasmithScene> DatasmithScene;
@@ -255,6 +309,9 @@ namespace DatasmithSketchUp
 		FTextureCollection Textures;
 		FLayerCollection Layers;
 
+		bool bColorByLayer = false;
+		bool bColorByLayerInvaliated = true;
+
 		TArray<TFuture<bool>> MeshExportTasks;
 	};
 }
diff --git a/Engine/Source/Programs/Enterprise/Datasmith/DatasmithSketchUpRubyExporter/Private/DatasmithSketchUpExporter.cpp b/Engine/Source/Programs/Enterprise/Datasmith/DatasmithSketchUpRubyExporter/Private/DatasmithSketchUpExporter.cpp
--- a/Engine/Source/Programs/Enterprise/Datasmith/DatasmithSketchUpRubyExporter/Private/DatasmithSketchUpExporter.cpp
+++ b/Engine/Source/Programs/Enterprise/Datasmith/DatasmithSketchUpRubyExporter/Private/DatasmithSketchUpExporter.cpp
@@ -312,9 +312,9 @@ public:
 		return true;
 	}
 
-	void Update()
+	bool Update(bool bModifiedHint)
 	{
-		Context.Update();
+		return Context.Update(bModifiedHint);
 	}
 
 	void SendUpdate()
@@ -373,10 +373,9 @@ public:
 		}
 
 		// Try Material
-		if (Context.Materials.RemoveMaterial(EntityId))
-		{
-			return true;
-		}
+		// Doesn't seem like material removal event ever comes through properly
+		// MaterialsObserver::onMaterialRemove simply has deleted material entity(meaning there's not way to retrieve correct reference/id for it) 
+		ensure(!Context.Materials.RegularMaterials.RemoveMaterial(EntityId));
 
 		return false;
 	}
@@ -402,7 +401,7 @@ public:
 			}
 		}
 
-		if (Context.Materials.InvalidateMaterial(EntityId))
+		if (Context.Materials.RegularMaterials.InvalidateMaterial(EntityId))
 		{
 			return true;
 		}
@@ -491,7 +490,6 @@ public:
 		}
 		case SURefType_Material:
 		{
-			Context.Materials.CreateMaterial(SUMaterialFromEntity(Entity));
 			break;
 		}
 		default:
@@ -514,13 +512,19 @@ public:
 		Context.ComponentInstances.LayerModified(LayerId);
 		Context.EntitiesObjects.LayerModified(LayerId);
 		Context.Layers.UpdateLayer(SULayerFromEntity(Entity));
+		Context.Materials.LayerMaterials.UpdateLayer(SULayerFromEntity(Entity));
 
 		return true;
 	}
 
 	bool OnStyleModified()
 	{
-		return Context.Materials.InvalidateDefaultMaterial();
+		return Context.Materials.RegularMaterials.InvalidateDefaultMaterial();
+	}
+
+	bool OnColorByLayerModified()
+	{
+		return Context.InvalidateColorByLayer();
 	}
 };
 
@@ -619,15 +623,16 @@ VALUE DatasmithSketchUpDirectLinkExporter_send_update(VALUE self)
 	return Qtrue;
 }
 
-VALUE DatasmithSketchUpDirectLinkExporter_update(VALUE self)
+VALUE DatasmithSketchUpDirectLinkExporter_update(VALUE self, VALUE modified_hint)
 {
 	// Converting args
 	FDatasmithSketchUpDirectLinkExporter* ptr;
 	Data_Get_Struct(self, FDatasmithSketchUpDirectLinkExporter, ptr);
+
+	bool bModifiedHint = RTEST(modified_hint);
 	// Done converting args
 
-	ptr->Update();
-	return Qtrue;
+	return ptr->Update(bModifiedHint) ? Qtrue : Qfalse;
 }
 
 VALUE DatasmithSketchUpDirectLinkExporter_export_current_datasmith_scene(VALUE self)
@@ -813,6 +818,18 @@ VALUE DatasmithSketchUpDirectLinkExporter_on_style_changed(VALUE self)
 	return Qtrue;
 }
 
+VALUE DatasmithSketchUpDirectLinkExporter_on_color_by_layer_changed(VALUE self)
+{
+	// Converting args
+	FDatasmithSketchUpDirectLinkExporter* Ptr;
+	Data_Get_Struct(self, FDatasmithSketchUpDirectLinkExporter, Ptr);
+	// Done converting args
+
+	Ptr->OnColorByLayerModified();
+
+	return Qtrue;
+}
+
 VALUE on_load(VALUE self, VALUE enable_ui, VALUE engine_path) {
 	// Converting args
 	Check_Type(engine_path, T_STRING);
@@ -899,8 +916,9 @@ extern "C" DLLEXPORT void Init_DatasmithSketchUp()
 	rb_define_method(DatasmithSketchUpDirectLinkExporterCRubyClass, "on_entity_removed", ToRuby(DatasmithSketchUpDirectLinkExporter_on_entity_removed), 2);
 
 	rb_define_method(DatasmithSketchUpDirectLinkExporterCRubyClass, "on_style_changed", ToRuby(DatasmithSketchUpDirectLinkExporter_on_style_changed), 0);
+	rb_define_method(DatasmithSketchUpDirectLinkExporterCRubyClass, "on_color_by_layer_changed", ToRuby(DatasmithSketchUpDirectLinkExporter_on_color_by_layer_changed), 0);
 
-	rb_define_method(DatasmithSketchUpDirectLinkExporterCRubyClass, "update", ToRuby(DatasmithSketchUpDirectLinkExporter_update), 0);
+	rb_define_method(DatasmithSketchUpDirectLinkExporterCRubyClass, "update", ToRuby(DatasmithSketchUpDirectLinkExporter_update), 1);
 	rb_define_method(DatasmithSketchUpDirectLinkExporterCRubyClass, "send_update", ToRuby(DatasmithSketchUpDirectLinkExporter_send_update), 0);
 	rb_define_method(DatasmithSketchUpDirectLinkExporterCRubyClass, "export_current_datasmith_scene", ToRuby(DatasmithSketchUpDirectLinkExporter_export_current_datasmith_scene), 0);
 
diff --git a/Engine/Source/Programs/Enterprise/Datasmith/DatasmithSketchUpRubyExporter/Private/DatasmithSketchUpMaterial.cpp b/Engine/Source/Programs/Enterprise/Datasmith/DatasmithSketchUpRubyExporter/Private/DatasmithSketchUpMaterial.cpp
--- a/Engine/Source/Programs/Enterprise/Datasmith/DatasmithSketchUpRubyExporter/Private/DatasmithSketchUpMaterial.cpp
+++ b/Engine/Source/Programs/Enterprise/Datasmith/DatasmithSketchUpRubyExporter/Private/DatasmithSketchUpMaterial.cpp
@@ -15,8 +15,7 @@
 // SketchUp SDK.
 #include "DatasmithSketchUpSDKBegins.h"
 #include "SketchUpAPI/model/component_instance.h"
-#include "SketchUpAPI/model/drawing_element.h"
-#include "SketchUpAPI/model/entity.h"
+#include <SketchUpAPI/model/layer.h>
 #include "SketchUpAPI/model/model.h"
 #include "SketchUpAPI/model/texture.h"
 #include "SketchUpAPI/model/typed_value.h"
@@ -24,6 +23,7 @@
 #include "DatasmithSketchUpSDKCeases.h"
 
 // Datasmith SDK.
+
 #include "Containers/Array.h"
 #include "DatasmithSceneFactory.h"
 #include "DatasmithUtils.h"
@@ -262,7 +262,6 @@ TSharedRef<IDatasmithBaseMaterialElement> FMaterial::CreateDefaultMaterialElemen
 	SURenderingOptionsRef RenderingOptionsRef = SU_INVALID;
 	if (SUModelGetRenderingOptions(Context.ModelRef, &RenderingOptionsRef) == SU_ERROR_NONE)
 	{
-		
 		SUTypedValueRef ColorTypedValue = SU_INVALID;
 		SUTypedValueCreate(&ColorTypedValue);
 		if (SURenderingOptionsGetValue(RenderingOptionsRef, "FaceFrontColor", &ColorTypedValue) == SU_ERROR_NONE)
@@ -290,9 +289,6 @@ TSharedRef<IDatasmithBaseMaterialElement> FMaterial::CreateDefaultMaterialElemen
 
 void FMaterial::Remove(FExportContext& Context)
 {
-
-	MaterialDirectlyAppliedToMeshes.RemoveDatasmithElement(Context);
-
 	// Expecting that SketchUp removes material from geometry/instances prior to removing it from the scene, but in case it doesn't...
 	if (!ensure(MaterialDirectlyAppliedToMeshes.MeshesMaterialDirectlyAppliedTo.IsEmpty()))
 	{
@@ -302,8 +298,6 @@ void FMaterial::Remove(FExportContext& Context)
 		}
 	}
 
-	MaterialInheritedByNodes.RemoveDatasmithElement(Context);
-
 	// Expecting that SketchUp removes material from geometry/instances prior to removing it from the scene, but in case it doesn't...
 	if (!ensure(MaterialInheritedByNodes.NodesMaterialInheritedBy.IsEmpty()))
 	{
@@ -317,6 +311,36 @@ void FMaterial::Remove(FExportContext& Context)
 	Context.Textures.UnregisterMaterial(this);
 }
 
+FMD5Hash FMaterial::ComputeHash(FExportContext& Context)
+{
+	FMD5 MD5;
+	FExtractedMaterial ExtractedMaterial(Context, MaterialRef);
+
+	MD5.Update(reinterpret_cast<uint8*>(&ExtractedMaterial.SourceType), sizeof(ExtractedMaterial.SourceType));
+	MD5.Update(reinterpret_cast<uint8*>(&ExtractedMaterial.SourceColor), sizeof(ExtractedMaterial.SourceColor));
+	MD5.Update(reinterpret_cast<uint8*>(&ExtractedMaterial.bSourceColorAlphaUsed), sizeof(ExtractedMaterial.bSourceColorAlphaUsed));
+	// MD5.Update(reinterpret_cast<uint8*>(&ExtractedMaterial.SketchupSourceName), sizeof(ExtractedMaterial.SketchupSourceName));
+
+	size_t TextureWidth = 0;
+	size_t TextureHeight = 0;
+	double TextureSScale = 1.0;
+	double TextureTScale = 1.0;
+	SUTextureGetDimensions(ExtractedMaterial.TextureRef, &TextureWidth, &TextureHeight, &TextureSScale, &TextureTScale);
+
+	bool bUseAlphaChannel = false;
+	SUTextureGetUseAlphaChannel(ExtractedMaterial.TextureRef, &bUseAlphaChannel);
+
+	MD5.Update(reinterpret_cast<uint8*>(&TextureWidth), sizeof(TextureWidth));
+	MD5.Update(reinterpret_cast<uint8*>(&TextureHeight), sizeof(TextureHeight));
+	MD5.Update(reinterpret_cast<uint8*>(&TextureSScale), sizeof(TextureSScale));
+	MD5.Update(reinterpret_cast<uint8*>(&TextureTScale), sizeof(TextureTScale));
+	MD5.Update(reinterpret_cast<uint8*>(&bUseAlphaChannel), sizeof(bUseAlphaChannel));
+
+	FMD5Hash Hash;
+	Hash.Set(MD5);
+	return Hash;
+}
+
 void FMaterial::UpdateTexturesUsage(FExportContext& Context)
 {
 	if (!bInvalidated)
@@ -344,18 +368,13 @@ void FMaterial::UpdateTexturesUsage(FExportContext& Context)
 
 void FMaterial::Update(FExportContext& Context)
 {
-	if(!bInvalidated)
-	{
-		return;
-	}
-
 	FExtractedMaterial ExtractedMaterial(Context, MaterialRef);
 	EntityId = ExtractedMaterial.SketchupSourceID.EntityID;
 
 	MaterialDirectlyAppliedToMeshes.RemoveDatasmithElement(Context);
 	if (MaterialDirectlyAppliedToMeshes.HasUsers())
 	{
-		MaterialDirectlyAppliedToMeshes.DatasmithElement = CreateMaterialElement(Context, ExtractedMaterial, *ExtractedMaterial.LocalizedMaterialName, Texture, false);
+		MaterialDirectlyAppliedToMeshes.DatasmithElement = CreateMaterialElement(Context, ExtractedMaterial, *ExtractedMaterial.LocalizedMaterialName, Texture, !bGeometryHasScalingBakedIntoUvs);
 	}
 
 	MaterialInheritedByNodes.RemoveDatasmithElement(Context);
@@ -364,10 +383,6 @@ void FMaterial::Update(FExportContext& Context)
 		MaterialInheritedByNodes.DatasmithElement = CreateMaterialElement(Context, ExtractedMaterial, *ExtractedMaterial.InheritedMaterialName, Texture, true);
 	}
 
-	FMaterialIDType MaterialId = DatasmithSketchUpUtils::GetMaterialID(MaterialRef);
-	MaterialDirectlyAppliedToMeshes.Apply(MaterialId);
-	MaterialInheritedByNodes.Apply(MaterialId);
-
 	bInvalidated = false;
 }
 
@@ -418,7 +433,10 @@ const TCHAR* FMaterialOccurrence::GetName()
 	return DatasmithElement->GetName();
 }
 
-void FMaterialOccurrence::Apply(FMaterialIDType MaterialId)
+// todo: make virtual FMaterialOccurrence - for layer/regular materials different?
+// they differ by the way they are applied to geometry/nodes -
+// VERY tied to meshes composition - they either map slots from SU Material or from SU Layer
+void FMaterialOccurrence::ApplyRegularMaterial(FMaterialIDType MaterialId)
 {
 	// Apply material to meshes
 	for (FEntitiesGeometry* Geometry : MeshesMaterialDirectlyAppliedTo)
@@ -426,7 +444,48 @@ void FMaterialOccurrence::Apply(FMaterialIDType MaterialId)
 		//Geometry->SetMaterialElementName(DatasmithSketchUpUtils::GetMaterialID(MaterialRef), MaterialDirectlyAppliedToMeshes->GetName());
 		for (const TSharedPtr<FDatasmithInstantiatedMesh>& Mesh : Geometry->Meshes)
 		{
-			if (int32* SlotIdPtr = Mesh->SlotIdForMaterialID.Find(MaterialId))
+			if (int32* SlotIdPtr = Mesh->SlotIdForMaterialId.Find(MaterialId))
+			{
+				Mesh->DatasmithMesh->SetMaterial(GetName(), *SlotIdPtr);
+			}
+		}
+	}
+
+	// Apply material to mesh actors
+	for (FNodeOccurence* NodePtr : NodesMaterialInheritedBy)
+	{
+		FNodeOccurence& Node = *NodePtr;
+		FDefinition* EntityDefinition = Node.Entity.GetDefinition();
+		DatasmithSketchUp::FEntitiesGeometry& EntitiesGeometry = *EntityDefinition->GetEntities().EntitiesGeometry;
+
+		for (int32 MeshIndex = 0; MeshIndex < Node.MeshActors.Num(); ++MeshIndex)
+		{
+			// Update Override(Inherited)  Material
+			// todo: set inherited material only on mesh actors that have faces with default material, right now setting on every mesh, hot harmful but excessive
+			if (EntitiesGeometry.IsMeshUsingInheritedMaterial(MeshIndex))
+			{
+				const TSharedPtr<IDatasmithMeshActorElement>& MeshActor = Node.MeshActors[MeshIndex];
+
+				// SketchUp has 'material override' only for single('Default') material. 
+				// So we reset overrides on the actor to remove this single override(if it was set) and re-add new override
+				MeshActor->ResetMaterialOverrides(); // Clear previous override if was set
+				MeshActor->AddMaterialOverride(GetName(), EntitiesGeometry.GetInheritedMaterialOverrideSlotId());
+			}
+		}
+	}
+}
+
+void FMaterialOccurrence::ApplyLayerMaterial(FLayerIDType LayerId)
+{
+	// Apply material to meshes
+	for (FEntitiesGeometry* Geometry : MeshesMaterialDirectlyAppliedTo)
+	{
+		//Geometry->SetMaterialElementName(DatasmithSketchUpUtils::GetMaterialID(MaterialRef), MaterialDirectlyAppliedToMeshes->GetName());
+		for (const TSharedPtr<FDatasmithInstantiatedMesh>& Mesh : Geometry->Meshes)
+		{
+			// todo: maybe can refactor this and unify FMaterialOccurrence for regular/layer materials?
+			// By moving this into the Mesh? It has this occurrence anyway so it can keep the association of material<->slot there?
+			if (int32* SlotIdPtr = Mesh->SlotIdForLayerId.Find(LayerId))
 			{
 				Mesh->DatasmithMesh->SetMaterial(GetName(), *SlotIdPtr);
 			}
@@ -518,18 +577,13 @@ void FMaterialOccurrence::UnregisterInstance(FExportContext& Context, FNodeOccur
 	RemoveUser(Context);
 }
 
-void FMaterialCollection::Update()
+void FRegularMaterials::RemoveUnused()
 {
-	// Update usage of materials and textures by materials before updating textures(to only update used textures)
 	TArray<FMaterialIDType> UnusedMaterials;
-	for (TPair<FMaterialIDType, TSharedPtr<DatasmithSketchUp::FMaterial>> IdAndMaterial : MaterialDefinitionMap)
+	for (TPair<FMaterialIDType, TSharedPtr<FMaterial>> IdAndMaterial : MaterialForMaterialId)
 	{
 		FMaterial& Material = *IdAndMaterial.Value;
-		if (Material.IsUsed())
-		{
-			Material.UpdateTexturesUsage(Context);
-		}
-		else
+		if (!Material.IsUsed())
 		{
 			UnusedMaterials.Add(IdAndMaterial.Key);
 		}
@@ -541,27 +595,152 @@ void FMaterialCollection::Update()
 		RemoveMaterial(MaterialId);
 	}
 
-	Context.Textures.Update();
+	if (!DefaultMaterial.HasUsers())
+	{
+		DefaultMaterial.Invalidate(Context);
+	}
+}
 
-	// Update materials after textures are updated - some materials might end up using shared texture(in case it has same contents, which is determined in Textures Update)
-	for (TPair<FMaterialIDType, TSharedPtr<FMaterial>> IdAndMaterial : MaterialDefinitionMap)
+void FRegularMaterials::UpdateDefaultMaterial()
+{
+	if (DefaultMaterial.IsInvalidated() && DefaultMaterial.HasUsers())
+	{
+		DefaultMaterial.DatasmithElement = FMaterial::CreateDefaultMaterialElement(Context);
+		DefaultMaterial.ApplyRegularMaterial(FMaterial::INHERITED_MATERIAL_ID);
+	}
+}
+
+void FRegularMaterials::Apply(FMaterial* Material)
+{
+	FMaterialIDType* MaterialIdFound = MaterialIdForMaterial.Find(Material);
+	if (!MaterialIdFound)
+	{
+		return;
+	}
+	FMaterialIDType MaterialId = *MaterialIdFound;
+	
+	Material->MaterialDirectlyAppliedToMeshes.ApplyRegularMaterial(MaterialId);
+	Material->MaterialInheritedByNodes.ApplyRegularMaterial(MaterialId);
+}
+
+void FLayerMaterials::RemoveUnused()
+{
+	TArray<FLayerIDType> UnusedMaterials;
+	for (TPair<FLayerIDType, TSharedPtr<FMaterial>> IdAndMaterial : MaterialForLayerId)
 	{
 		FMaterial& Material = *IdAndMaterial.Value;
-		if (Material.IsUsed())
+		if (!Material.IsUsed())
 		{
-			Material.Update(Context);
+			UnusedMaterials.Add(IdAndMaterial.Key);
 		}
 	}
 
-	if (DefaultMaterial.IsInvalidated() && DefaultMaterial.HasUsers())
+	for (FLayerIDType LayerId: UnusedMaterials)
 	{
-		DefaultMaterial.DatasmithElement = FMaterial::CreateDefaultMaterialElement(Context);
-		DefaultMaterial.Apply(FMaterial::INHERITED_MATERIAL_ID);
+		TSharedPtr<FMaterial> Material;
+		if (MaterialForLayerId.RemoveAndCopyValue(LayerId, Material))
+		{
+			MaterialHashForLayerId.Remove(LayerId);
+			LayerIdForMaterial.Remove(Material.Get());
+			Context.Materials.RemoveMaterial(Material.Get());
+		}
+	}
+}
+
+void FLayerMaterials::Apply(FMaterial* Material)
+{
+	FLayerIDType* Found = LayerIdForMaterial.Find(Material);
+	if (!Found)
+	{
+		return;
 	}
+	FLayerIDType LayerId = *Found;
+
+	MaterialHashForLayerId.Add(LayerId, Material->ComputeHash(Context));
+	
+	Material->MaterialDirectlyAppliedToMeshes.ApplyLayerMaterial(LayerId);
+	Material->MaterialInheritedByNodes.ApplyLayerMaterial(LayerId);
 }
 
-const TCHAR* FMaterialCollection::GetDefaultMaterialName()
+void FLayerMaterials::UpdateLayer(SULayerRef LayerRef)
 {
-	return DefaultMaterial.GetName();
+	if (TSharedPtr<DatasmithSketchUp::FMaterial>* Found = MaterialForLayerId.Find(DatasmithSketchUpUtils::GetEntityID(SULayerToEntity(LayerRef))))
+	{
+		(*Found)->Invalidate();
+	}
 }
 
+// todo: VERY similar to regular material's RegisterInstance
+FMaterialOccurrence* FLayerMaterials::RegisterInstance(FLayerIDType LayerID, FNodeOccurence* NodeOccurrence)
+{
+	if (NodeOccurrence->MaterialOverride)
+	{
+		NodeOccurrence->MaterialOverride->UnregisterInstance(Context, NodeOccurrence);
+	}
+
+	if (const TSharedPtr<DatasmithSketchUp::FMaterial> Material = FindOrCreateMaterialForLayer(LayerID))
+	{
+		return &Material->RegisterInstance(NodeOccurrence);
+	}
+
+	return {}; // Don't use a material if material id is unknown(of default)
+}
+
+// Test if a layer material has changed since last update 
+bool FLayerMaterials::CheckForModifications()
+{
+	size_t LayerCount = 0;
+	SUModelGetNumLayers(Context.ModelRef, &LayerCount);
+
+	TArray<SULayerRef> Layers;
+	Layers.Init(SU_INVALID, LayerCount);
+	SUResult SResult = SUModelGetLayers(Context.ModelRef, LayerCount, Layers.GetData(), &LayerCount);
+	Layers.SetNum(LayerCount);
+
+	bool bHasModifications = false;
+
+	for (SULayerRef LayerRef : Layers)
+	{
+		FLayerIDType LayerId = DatasmithSketchUpUtils::GetEntityID(SULayerToEntity(LayerRef));
+		if (TSharedPtr<DatasmithSketchUp::FMaterial>* Found = MaterialForLayerId.Find(LayerId))
+		{
+			TSharedPtr<DatasmithSketchUp::FMaterial> Material = *Found;
+			FMD5Hash Hash = Material->ComputeHash(Context);
+			FMD5Hash* HashFound = MaterialHashForLayerId.Find(LayerId);
+			if (!HashFound || (*HashFound != Hash))
+			{
+				Material->Invalidate();
+				bHasModifications = true;
+			}
+		}
+	}
+
+	return bHasModifications;
+}
+
+void FMaterialCollection::Update()
+{
+	LayerMaterials.RemoveUnused();
+	RegularMaterials.RemoveUnused();
+
+	// Update usage of materials and textures by materials before updating textures(to only update used textures)
+	for (FMaterial* Material : Materials)
+	{
+		Material->UpdateTexturesUsage(Context);
+	}
+
+	Context.Textures.Update();
+
+	// Update materials after textures are updated - some materials might end up using shared texture(in case it has same contents, which is determined in Textures Update)
+	for (FMaterial* Material: Materials)
+	{
+		if (Material->IsInvalidated())
+		{
+			Material->Update(Context);
+			LayerMaterials.Apply(Material);
+			RegularMaterials.Apply(Material);
+		}
+	}
+
+	RegularMaterials.UpdateDefaultMaterial();
+}
diff --git a/Engine/Source/Programs/Enterprise/Datasmith/DatasmithSketchUpRubyExporter/Private/DatasmithSketchUpMaterial.h b/Engine/Source/Programs/Enterprise/Datasmith/DatasmithSketchUpRubyExporter/Private/DatasmithSketchUpMaterial.h
--- a/Engine/Source/Programs/Enterprise/Datasmith/DatasmithSketchUpRubyExporter/Private/DatasmithSketchUpMaterial.h
+++ b/Engine/Source/Programs/Enterprise/Datasmith/DatasmithSketchUpRubyExporter/Private/DatasmithSketchUpMaterial.h
@@ -11,6 +11,7 @@
 // Datasmith SDK.
 #include "Containers/Map.h"
 #include "Containers/UnrealString.h"
+#include "Misc/SecureHash.h"
 #include "Templates/SharedPointer.h"
 
 
@@ -25,6 +26,7 @@ namespace DatasmithSketchUp
 	class FMaterial;
 	class FTexture;
 
+	// Holds datasmith material and records connections to its 'users' - meshes and nodes
 	class FMaterialOccurrence : FNoncopyable
 	{
 	public:
@@ -64,7 +66,8 @@ namespace DatasmithSketchUp
 		void RegisterInstance(FNodeOccurence* NodeOccurrence); // Tie material to an instance it's used on
 		void UnregisterInstance(FExportContext& Context, FNodeOccurence* NodeOccurrence);
 
-		void Apply(FMaterialIDType MaterialId); // Apply material to meshes/instances
+		void ApplyRegularMaterial(FMaterialIDType MaterialId); // Apply material to meshes/instances
+		void ApplyLayerMaterial(FLayerIDType LayerId); // Apply material to meshes/instances
 
 		TSet<FEntitiesGeometry*> MeshesMaterialDirectlyAppliedTo;
 		TSet<FNodeOccurence*> NodesMaterialInheritedBy;
@@ -103,6 +106,10 @@ namespace DatasmithSketchUp
 
 
 		void Invalidate();
+		bool IsInvalidated()
+		{
+			return bInvalidated;
+		}
 		void UpdateTexturesUsage(FExportContext& Context);
 		void Update(FExportContext& Context); // create datasmith elements for material occurrences
 		void Remove(FExportContext& Context);
@@ -117,6 +124,16 @@ namespace DatasmithSketchUp
 			return Texture;
 		}
 
+		FMD5Hash ComputeHash(FExportContext& Context);
+
+		// Material can be directly applied to a face in SketchUp
+		FMaterialOccurrence MaterialDirectlyAppliedToMeshes;
+		// In case face has Default material assigned it inherits material set to it's parent(in general - first non-Default material in ancestors chain)
+		FMaterialOccurrence MaterialInheritedByNodes;
+
+		// Whether UVs of geometry have texture scaling baked, when applying this material
+		// Scaling is baked for regular materials, applied directly to faces so regular materials assigned to static meshes need to take this into account.
+		bool bGeometryHasScalingBakedIntoUvs = true;
 	private:
 
 		SUMaterialRef MaterialRef;
@@ -125,11 +142,6 @@ namespace DatasmithSketchUp
 
 		FTexture* Texture = nullptr;
 
-		// Material can be directly applied to a face in SketchUp
-		FMaterialOccurrence MaterialDirectlyAppliedToMeshes;
-		// In case face has Default material assigned it inherits material set to it's parent(in general - first non-Default material in ancestors chain)
-		FMaterialOccurrence MaterialInheritedByNodes;
-
 		uint8 bInvalidated : 1;
 
 		friend class FMaterialOccurrence;
diff --git a/Engine/Source/Programs/Enterprise/Datasmith/DatasmithSketchUpRubyExporter/Private/DatasmithSketchUpMesh.cpp b/Engine/Source/Programs/Enterprise/Datasmith/DatasmithSketchUpRubyExporter/Private/DatasmithSketchUpMesh.cpp
--- a/Engine/Source/Programs/Enterprise/Datasmith/DatasmithSketchUpRubyExporter/Private/DatasmithSketchUpMesh.cpp
+++ b/Engine/Source/Programs/Enterprise/Datasmith/DatasmithSketchUpRubyExporter/Private/DatasmithSketchUpMesh.cpp
@@ -65,12 +65,10 @@ namespace DatasmithSketchUp
 	public:
 
 		// Convert the combined mesh into a Datasmith mesh.
-		void ConvertMeshToDatasmith(SUTransformation Transform, FDatasmithMesh& OutDMesh) const;
+		void ConvertMeshToDatasmith(FExportContext& Context, SUTransformation Transform, FDatasmithMesh& OutDMesh) const;
 
 		// Tessellate a SketchUp face into a triangle mesh merged into the combined mesh.
-		void AddFace(
-			SUFaceRef InSFaceRef // valid source SketchUp face to tessellate and combine
-		);
+		void AddFace(FExportContext& Context, SUFaceRef InSFaceRef, FLayerIDType LayerId);
 
 		int32 GetOrCreateSlotForMaterial(FMaterialIDType MaterialID)
 		{
@@ -87,6 +85,21 @@ namespace DatasmithSketchUp
 			}
 		}
 
+		int32 GetOrCreateSlotForLayer(FLayerIDType LayerID)
+		{
+			if (int32* SlotIdPtr = SlotIdForLayerId.Find(LayerID))
+			{
+				return *SlotIdPtr;
+			}
+			else
+			{
+				int32 SlotId = LayerIDForSlotId.Num();
+				LayerIDForSlotId.Add(LayerID);// Assign material to slot
+				SlotIdForLayerId.Add(LayerID, SlotId); // store back reference from material to slot
+				return SlotId;
+			}
+		}
+
 		// Return whether or not the combined mesh contains geometry.
 		bool ContainsGeometry() const;
 
@@ -107,6 +120,10 @@ namespace DatasmithSketchUp
 
 		TArray<FEntityIDType> MaterialIDForSlotId;
 		TMap<FEntityIDType, int32> SlotIdForMaterialId;
+
+		TArray<FLayerIDType> LayerIDForSlotId;
+		TMap<FLayerIDType, int32> SlotIdForLayerId;
+
 		bool bHasFacesWithDefaultMaterial = false;;
 
 	};
@@ -117,7 +134,7 @@ namespace DatasmithSketchUp
 		return (MeshVertexPoints.Num() > 0 && MeshTriangleIndices.Num() > 0);
 	}
 
-	void FDatasmithSketchUpMesh::AddFace(SUFaceRef InSFaceRef)
+	void FDatasmithSketchUpMesh::AddFace(FExportContext& Context, SUFaceRef InSFaceRef, FLayerIDType LayerId)
 	{
 		// Retrieve the SketchUp face front material.
 		SUMaterialRef FrontMaterialRef = SU_INVALID;
@@ -263,25 +280,32 @@ namespace DatasmithSketchUp
 
 		int32 SlotId = 0; // Default material slot
 
-		// Get the SketckUp material ID.
-		if (bUseFrontMaterial)
+		if (Context.bColorByLayer)
 		{
-			if (SUIsValid(FrontMaterialRef))
-			{
-				// Get the front material ID of the SketckUp front material.
-				SlotId = GetOrCreateSlotForMaterial(DatasmithSketchUpUtils::GetMaterialID(FrontMaterialRef));
-			}
+			// Get the front material ID of the SketckUp front material.
+			SlotId = GetOrCreateSlotForLayer(LayerId);
 		}
-		else // bUseBackMaterial
+		else
 		{
-			if (SUIsValid(BackMaterialRef))
+			// Get the SketckUp material ID.
+			if (bUseFrontMaterial)
 			{
-				// Get the back material ID of the SketckUp back material.
-				SlotId = GetOrCreateSlotForMaterial(DatasmithSketchUpUtils::GetMaterialID(BackMaterialRef));
+				if (SUIsValid(FrontMaterialRef))
+				{
+					// Get the front material ID of the SketckUp front material.
+					SlotId = GetOrCreateSlotForMaterial(DatasmithSketchUpUtils::GetMaterialID(FrontMaterialRef));
+				}
+			}
+			else // bUseBackMaterial
+			{
+				if (SUIsValid(BackMaterialRef))
+				{
+					// Get the back material ID of the SketckUp back material.
+					SlotId = GetOrCreateSlotForMaterial(DatasmithSketchUpUtils::GetMaterialID(BackMaterialRef));
+				}
 			}
 		}
 
-
 		if (SlotId == 0)
 		{
 			// todo: it's possible to skip adding slot=0 when there's no faces with 'default' material
@@ -289,14 +313,15 @@ namespace DatasmithSketchUp
 			bHasFacesWithDefaultMaterial = true;
 		}
 
-		MeshTriangleSlotIds.Reserve(MeshTriangleSlotIds.Num());
+		MeshTriangleSlotIds.Reserve(MeshTriangleSlotIds.Num() + TriangleCount);
+		
 		for(int32 TriangleIndex = 0; TriangleIndex < TriangleCount; ++TriangleIndex)
 		{
 			MeshTriangleSlotIds.Add(SlotId);
 		}
 	}
 
-	void FDatasmithSketchUpMesh::ConvertMeshToDatasmith(SUTransformation Transform, FDatasmithMesh& OutDMesh) const
+	void FDatasmithSketchUpMesh::ConvertMeshToDatasmith(FExportContext& Context, SUTransformation Transform, FDatasmithMesh& OutDMesh) const
 	{
 		// Get the number of mesh vertices (must be > 0).
 		int32 VertexCount = MeshVertexPoints.Num();
@@ -408,7 +433,8 @@ void FEntities::UpdateGeometry(FExportContext& Context)
 			if (MeshCount < EntitiesGeometry->Meshes.Num())
 			{
 				Mesh = EntitiesGeometry->Meshes[MeshCount];
-				Mesh->SlotIdForMaterialID.Reset();
+				Mesh->SlotIdForMaterialId.Reset();
+				Mesh->SlotIdForLayerId.Reset();
 			}
 			else
 			{
@@ -425,19 +451,35 @@ void FEntities::UpdateGeometry(FExportContext& Context)
 			Mesh->bIsUsingInheritedMaterial = ExtractedMeshPtr->bHasFacesWithDefaultMaterial;
 
 			// Add the non-inherited materials used by the combined mesh triangles.
-			for (int32 SlotId = 0;SlotId < ExtractedMeshPtr->MaterialIDForSlotId.Num(); ++SlotId)
+			if (Context.bColorByLayer)
 			{
-				if (SlotId == 0 && !ExtractedMeshPtr->bHasFacesWithDefaultMaterial)
+				for (int32 SlotId = 0;SlotId < ExtractedMeshPtr->LayerIDForSlotId.Num(); ++SlotId)
 				{
-					continue; // Skip adding Default material slot if it's not used
-				}
+					FLayerIDType LayerID = ExtractedMeshPtr->LayerIDForSlotId[SlotId];
+					Mesh->SlotIdForLayerId.FindOrAdd(LayerID, SlotId);
 
-				FMaterialIDType MeshMaterialID = ExtractedMeshPtr->MaterialIDForSlotId[SlotId];
-				Mesh->SlotIdForMaterialID.FindOrAdd(MeshMaterialID, SlotId);
-				// Default or (somehow)missing materials are also assigned to mesh(as a default material)
-				if (FMaterialOccurrence* Material = Context.Materials.RegisterGeometry(MeshMaterialID, EntitiesGeometry.Get()))
+					if (FMaterialOccurrence* Material = Context.Materials.LayerMaterials.RegisterGeometryForLayer(LayerID, EntitiesGeometry.Get()))
+					{
+						Mesh->DatasmithMesh->SetMaterial(Material->GetName(), SlotId);
+					}
+				}
+			}
+			else
+			{
+				for (int32 SlotId = 0;SlotId < ExtractedMeshPtr->MaterialIDForSlotId.Num(); ++SlotId)
 				{
-					Mesh->DatasmithMesh->SetMaterial(Material->GetName(), SlotId);
+					if (SlotId == 0 && !ExtractedMeshPtr->bHasFacesWithDefaultMaterial)
+					{
+						continue; // Skip adding Default material slot if it's not used
+					}
+
+					FMaterialIDType MeshMaterialID = ExtractedMeshPtr->MaterialIDForSlotId[SlotId];
+					Mesh->SlotIdForMaterialId.FindOrAdd(MeshMaterialID, SlotId);
+					// Default or (somehow)missing materials are also assigned to mesh(as a default material)
+					if (FMaterialOccurrence* Material = Context.Materials.RegularMaterials.RegisterGeometry(MeshMaterialID, EntitiesGeometry.Get()))
+					{
+						Mesh->DatasmithMesh->SetMaterial(Material->GetName(), SlotId);
+					}
 				}
 			}
 			
@@ -449,7 +491,7 @@ void FEntities::UpdateGeometry(FExportContext& Context)
 				{
 					FDatasmithMeshExporter DatasmithMeshExporter;
 					FDatasmithMesh DatasmithMesh;
-					ExtractedMeshPtr->ConvertMeshToDatasmith(Transform, DatasmithMesh);
+					ExtractedMeshPtr->ConvertMeshToDatasmith(Context, Transform, DatasmithMesh);
 
 					FGCScopeGuard GCGuard; // Prevent GC from running while UDatasmithMesh is created in ExportToUObject. 
 					return DatasmithMeshExporter.ExportToUObject(Mesh->DatasmithMesh, Context.GetAssetsOutputPath(), DatasmithMesh, nullptr, FDatasmithExportOptions::LightmapUV);
@@ -580,14 +622,15 @@ void ScanSketchUpEntitiesFaces(FExportContext& Context, SUEntitiesRef EntitiesRe
 			// this geometry needs to be rebuilt
 			SULayerRef LayerRef = SU_INVALID;
 			SUDrawingElementGetLayer(SUFaceToDrawingElement(SScannedFaceRef), &LayerRef);
-			Geometry.Layers.Add(DatasmithSketchUpUtils::GetEntityID(SULayerToEntity(LayerRef)));
+			FLayerIDType LayerId = DatasmithSketchUpUtils::GetEntityID(SULayerToEntity(LayerRef));
+			Geometry.Layers.Add(LayerId);
 
 			bool bFaceHidden = false;
 			SUDrawingElementGetHidden(SUFaceToDrawingElement(SScannedFaceRef), &bFaceHidden);
 
 			if (!bFaceHidden && Context.Layers.IsLayerVisible(LayerRef))
 			{
-				ExtractedMesh.AddFace(SScannedFaceRef);
+				ExtractedMesh.AddFace(Context, SScannedFaceRef, LayerId);
 			}
 
 			// Get the number of SketchUp face edges.
diff --git a/Engine/Source/Programs/Enterprise/Datasmith/DatasmithSketchUpRubyExporter/Private/DatasmithSketchUpMesh.h b/Engine/Source/Programs/Enterprise/Datasmith/DatasmithSketchUpRubyExporter/Private/DatasmithSketchUpMesh.h
--- a/Engine/Source/Programs/Enterprise/Datasmith/DatasmithSketchUpRubyExporter/Private/DatasmithSketchUpMesh.h
+++ b/Engine/Source/Programs/Enterprise/Datasmith/DatasmithSketchUpRubyExporter/Private/DatasmithSketchUpMesh.h
@@ -23,7 +23,8 @@ namespace DatasmithSketchUp
 	{
 	public:
 		TSharedPtr<IDatasmithMeshElement> DatasmithMesh;
-		TMap<FEntityIDType, int32> SlotIdForMaterialID;
+		TMap<FMaterialIDType, int32> SlotIdForMaterialId;
+		TMap<FLayerIDType, int32> SlotIdForLayerId;
 		bool bIsUsingInheritedMaterial; // Whether mesh has faces without material assigned(so override/inherited material will apply)
 	};
 }
