[SketchUp] Merge Component geometry into single mesh

- old code(making disconnected faces into separate meshes) is kept, compiling. In case we need an option

diff --git a/Engine/Source/Programs/Enterprise/Datasmith/DatasmithSketchUpRubyExporter/Private/DatasmithSketchUpExportContext.h b/Engine/Source/Programs/Enterprise/Datasmith/DatasmithSketchUpRubyExporter/Private/DatasmithSketchUpExportContext.h
--- a/Engine/Source/Programs/Enterprise/Datasmith/DatasmithSketchUpRubyExporter/Private/DatasmithSketchUpExportContext.h
+++ b/Engine/Source/Programs/Enterprise/Datasmith/DatasmithSketchUpRubyExporter/Private/DatasmithSketchUpExportContext.h
@@ -273,6 +273,14 @@ namespace DatasmithSketchUp
 		FExportContext& Context;
 	};
 
+	struct FOptions
+	{
+		// Enabling this reverts to legacy behavior when each set of disconnected faces results in a separate static mesh.
+		// Disconnected means faces not connected through common edges. E.g. face A and B share an edge, face C shares edge with A then ABC are 'connected' one mesh(A connects with B through C).
+		bool bSeparateDisconnectedMeshes = false;
+	};
+
+
 	// Holds all the data needed during export and incremental updates
 	class FExportContext
 	{
@@ -292,6 +300,8 @@ namespace DatasmithSketchUp
 
 		bool PreUpdateColorByLayer();
 
+		FOptions Options;
+
 		SUModelRef ModelRef = SU_INVALID;
 
 		TSharedPtr<IDatasmithScene> DatasmithScene;
diff --git a/Engine/Source/Programs/Enterprise/Datasmith/DatasmithSketchUpRubyExporter/Private/DatasmithSketchUpMesh.cpp b/Engine/Source/Programs/Enterprise/Datasmith/DatasmithSketchUpRubyExporter/Private/DatasmithSketchUpMesh.cpp
--- a/Engine/Source/Programs/Enterprise/Datasmith/DatasmithSketchUpRubyExporter/Private/DatasmithSketchUpMesh.cpp
+++ b/Engine/Source/Programs/Enterprise/Datasmith/DatasmithSketchUpRubyExporter/Private/DatasmithSketchUpMesh.cpp
@@ -401,6 +401,7 @@ const TCHAR* FEntitiesGeometry::GetMeshElementName(int32 MeshIndex)
 }
 
 void ScanSketchUpEntitiesFaces(FExportContext& Context, SUEntitiesRef EntitiesRef, FEntitiesGeometry& Geometry, TFunctionRef<void(TSharedPtr<FDatasmithSketchUpMesh> ExtractedMesh)> OnNewExtractedMesh);
+void CombineSketchUpEntitiesFaces(FExportContext& Context, SUEntitiesRef EntitiesRef, FEntitiesGeometry& Geometry, TFunctionRef<void(TSharedPtr<FDatasmithSketchUpMesh> ExtractedMesh)> OnNewExtractedMesh);
 
 void FEntities::UpdateGeometry(FExportContext& Context)
 {
@@ -500,7 +501,15 @@ void FEntities::UpdateGeometry(FExportContext& Context)
 		}
 	};
 
-	ScanSketchUpEntitiesFaces(Context, EntitiesRef, *EntitiesGeometry, ProcessExtractedMesh);
+	if (Context.Options.bSeparateDisconnectedMeshes)
+	{
+		ScanSketchUpEntitiesFaces(Context, EntitiesRef, *EntitiesGeometry, ProcessExtractedMesh);
+	}
+	else
+	{
+		CombineSketchUpEntitiesFaces(Context, EntitiesRef, *EntitiesGeometry, ProcessExtractedMesh);
+	}
+
 	EntitiesGeometry->Meshes.SetNum(MeshCount);
 
 	Context.EntitiesObjects.RegisterEntities(*this);
@@ -579,7 +588,7 @@ void ScanSketchUpEntitiesFaces(FExportContext& Context, SUEntitiesRef EntitiesRe
 	// Retrieve the faces in the source SketchUp entities.
 	TArray<SUFaceRef> SFaces;
 	SFaces.Init(SU_INVALID, SFaceCount);
-	SUEntitiesGetFaces(EntitiesRef, SFaceCount, SFaces.GetData(), &SFaceCount); // we can ignore the returned SU_RESULT
+	SUEntitiesGetFaces(EntitiesRef, SFaceCount, SFaces.GetData(), &SFaceCount);
 	SFaces.SetNum(SFaceCount);
 
 	TSet<int32> ScannedEdgeIDSet;
@@ -603,7 +612,6 @@ void ScanSketchUpEntitiesFaces(FExportContext& Context, SUEntitiesRef EntitiesRe
 		FDatasmithSketchUpMesh& ExtractedMesh = *ExtractedMeshPtr;
 		ExtractedMesh.GetOrCreateSlotForMaterial(FMaterial::INHERITED_MATERIAL_ID); // Add default material to Slot=0
 
-
 		// The source SketchUp face needs to be scanned once.
 		TArray<SUFaceRef> FacesToScan;
 		FacesToScan.Add(SSourceFaceRef);
@@ -683,4 +691,58 @@ void ScanSketchUpEntitiesFaces(FExportContext& Context, SUEntitiesRef EntitiesRe
 
 		OnNewExtractedMesh(ExtractedMeshPtr);
 	}
-}
\ No newline at end of file
+}
+
+void CombineSketchUpEntitiesFaces(FExportContext& Context, SUEntitiesRef EntitiesRef, FEntitiesGeometry& Geometry, TFunctionRef<void(TSharedPtr<FDatasmithSketchUpMesh> ExtractedMesh)> OnNewExtractedMesh)
+{
+	// Get the number of faces in the source SketchUp entities.
+	size_t SFaceCount = 0;
+	SUEntitiesGetNumFaces(EntitiesRef, &SFaceCount); // we can ignore the returned SU_RESULT
+
+	if (SFaceCount == 0)
+	{
+		return;
+	}
+
+	// Retrieve the faces in the source SketchUp entities.
+	TArray<SUFaceRef> SFaces;
+	SFaces.Init(SU_INVALID, SFaceCount);
+	SUEntitiesGetFaces(EntitiesRef, SFaceCount, SFaces.GetData(), &SFaceCount);
+	SFaces.SetNum(SFaceCount);
+
+	// Create a mesh combining the geometry of the SketchUp connected faces.
+	TSharedPtr<FDatasmithSketchUpMesh> ExtractedMeshPtr = MakeShared<FDatasmithSketchUpMesh>();
+	FDatasmithSketchUpMesh& ExtractedMesh = *ExtractedMeshPtr;
+	ExtractedMesh.GetOrCreateSlotForMaterial(FMaterial::INHERITED_MATERIAL_ID); // Add default material to Slot=0
+
+	for (SUFaceRef FaceRef : SFaces)
+	{
+		// Get the Source SketckUp face ID.
+		int32 SSourceFaceID = DatasmithSketchUpUtils::GetFaceID(FaceRef);
+
+		// Do not scan more than once a valid SketckUp face.
+		if (SUIsInvalid(FaceRef) || Geometry.FaceIds.Contains(SSourceFaceID))
+		{
+			continue;
+		}
+
+		Geometry.FaceIds.Add(SSourceFaceID);
+
+		// Record every face's layer(even for invisible faces!). When face layer visibility changes 
+		// this geometry needs to be rebuilt
+		SULayerRef LayerRef = SU_INVALID;
+		SUDrawingElementGetLayer(SUFaceToDrawingElement(FaceRef), &LayerRef);
+		FLayerIDType LayerId = DatasmithSketchUpUtils::GetEntityID(SULayerToEntity(LayerRef));
+		Geometry.Layers.Add(LayerId);
+
+		bool bFaceHidden = false;
+		SUDrawingElementGetHidden(SUFaceToDrawingElement(FaceRef), &bFaceHidden);
+
+		if (!bFaceHidden && Context.Layers.IsLayerVisible(LayerRef))
+		{
+			ExtractedMesh.AddFace(Context, FaceRef, LayerId);
+		}
+	}
+
+	OnNewExtractedMesh(ExtractedMeshPtr);
+}
