[SketchUp] Static Mesh name uses Persistent Id instead of Guid

- for SU versions above 2019(since 2020.1)

diff --git a/Engine/Source/Programs/Enterprise/Datasmith/DatasmithSketchUpRubyExporter/Private/DatasmithSketchUpComponent.cpp b/Engine/Source/Programs/Enterprise/Datasmith/DatasmithSketchUpRubyExporter/Private/DatasmithSketchUpComponent.cpp
--- a/Engine/Source/Programs/Enterprise/Datasmith/DatasmithSketchUpRubyExporter/Private/DatasmithSketchUpComponent.cpp
+++ b/Engine/Source/Programs/Enterprise/Datasmith/DatasmithSketchUpRubyExporter/Private/DatasmithSketchUpComponent.cpp
@@ -14,6 +14,7 @@
 #include "SketchUpAPI/model/component_definition.h"
 #include "SketchUpAPI/model/drawing_element.h"
 #include "SketchUpAPI/model/entities.h"
+#include <SketchUpAPI/model/entity.h>
 #include "SketchUpAPI/model/group.h"
 #include "SketchUpAPI/model/layer.h"
 #include "SketchUpAPI/model/model.h"
@@ -394,6 +395,11 @@ FString FModelDefinition::GetSketchupSourceName()
 	return SketchupSourceName;
 }
 
+FString FModelDefinition::GetSketchupSourceId()
+{
+	return GetSketchupSourceGUID();
+}
+
 SUTransformation FModelDefinition::GetMeshBakedTransform()
 {
 	SUTransformation Transform;
@@ -544,6 +550,26 @@ FString FComponentDefinition::GetSketchupSourceName()
 	return SuGetString(SUComponentDefinitionGetName, ComponentDefinitionRef);
 }
 
+FString FComponentDefinition::GetSketchupSourceId()
+{
+// Although implemented since SUEntityGetPersistentID since SU 2017 is returns valid Id for Definitions
+// only since SU 2020.1 (even though SUEntityGetPersistentID docs states SUComponentDefinitionRef 'supported' from 2017)
+// see [REDACTED_PROFILE_URL]/api-issue-tracker/issues/314
+#ifndef SKP_SDK_2019
+	// Use Entity PersistentID - this one is persistent(between sessions)  for model file and doesn't change when definition is modified(e.g. geometry edited)
+	int64_t EntityPid = 0;
+	if (SUEntityGetPersistentID(SUComponentDefinitionToEntity(ComponentDefinitionRef), &EntityPid) == SU_ERROR_NONE)
+	{
+		if (ensure(EntityPid != 0))
+		{
+			return FString::Printf(TEXT("%llx"), EntityPid);
+		}
+	}
+#endif
+
+	return FMD5::HashAnsiString(*GetSketchupSourceGUID());
+}
+
 // Bake transform into the mesh only when
 // - component has no child components, 
 // - only SINGLE instance
@@ -979,6 +1005,10 @@ void FComponentInstance::FillOccurrenceActorMetadata(FNodeOccurence& Node)
 	DefinitionName->SetValue(*GetDefinition()->GetSketchupSourceName());
 	Node.DatasmithMetadataElement->AddProperty(DefinitionName);
 
+	TSharedPtr<IDatasmithKeyValueProperty> DefinitionIdName = FDatasmithSceneFactory::CreateKeyValueProperty(TEXT("DefinitionIdName"));
+	DefinitionIdName->SetValue(*GetDefinition()->GetSketchupSourceId());
+	Node.DatasmithMetadataElement->AddProperty(DefinitionIdName);
+
 	// Add instance metadata
 	if (ParsedMetadata)
 	{
diff --git a/Engine/Source/Programs/Enterprise/Datasmith/DatasmithSketchUpRubyExporter/Private/DatasmithSketchUpComponent.h b/Engine/Source/Programs/Enterprise/Datasmith/DatasmithSketchUpRubyExporter/Private/DatasmithSketchUpComponent.h
--- a/Engine/Source/Programs/Enterprise/Datasmith/DatasmithSketchUpRubyExporter/Private/DatasmithSketchUpComponent.h
+++ b/Engine/Source/Programs/Enterprise/Datasmith/DatasmithSketchUpRubyExporter/Private/DatasmithSketchUpComponent.h
@@ -136,8 +136,9 @@ namespace DatasmithSketchUp
 		virtual void InvalidateInstancesMetadata(FExportContext& Context) = 0; // Mark that all instances(and their occurrences) needed to be updated
 		virtual void FillOccurrenceActorMetadata(FNodeOccurence& Node) = 0;
 		
-		virtual FString GetSketchupSourceGUID() = 0;
-		virtual FString GetSketchupSourceName() = 0;
+		virtual FString GetSketchupSourceGUID() = 0; // Guid of the definition, SketchUp hashes definition contents into guid
+		virtual FString GetSketchupSourceName() = 0; // Name used for label
+		virtual FString GetSketchupSourceId() = 0; // Unique name identifier
 
 		virtual SUTransformation GetMeshBakedTransform() = 0; // Transform that may be used to bake geometry
 
@@ -186,6 +187,7 @@ namespace DatasmithSketchUp
 
 		virtual FString GetSketchupSourceGUID() override;
 		virtual FString GetSketchupSourceName() override;
+		virtual FString GetSketchupSourceId() override;
 
 		virtual SUTransformation GetMeshBakedTransform() override;
 		// End FDefinition
@@ -236,6 +238,7 @@ namespace DatasmithSketchUp
 
 		virtual FString GetSketchupSourceGUID() override;
 		virtual FString GetSketchupSourceName()  override;
+		virtual FString GetSketchupSourceId() override;
 
 		virtual SUTransformation GetMeshBakedTransform() override;
 		// End FDefinition
@@ -258,8 +261,6 @@ namespace DatasmithSketchUp
 		void AddMeshesToDatasmithScene(FExportContext& Context);
 		void RemoveMeshesFromDatasmithScene(FExportContext& Context);
 
-		TSharedPtr<IDatasmithMeshElement> CreateMeshElement(FExportContext& Context, FDatasmithMesh& DatasmithMesh);
-
 		TArray<SUGroupRef> GetGroups();
 
 		TArray<SUComponentInstanceRef> GetComponentInstances();
diff --git a/Engine/Source/Programs/Enterprise/Datasmith/DatasmithSketchUpRubyExporter/Private/DatasmithSketchUpMesh.cpp b/Engine/Source/Programs/Enterprise/Datasmith/DatasmithSketchUpRubyExporter/Private/DatasmithSketchUpMesh.cpp
--- a/Engine/Source/Programs/Enterprise/Datasmith/DatasmithSketchUpRubyExporter/Private/DatasmithSketchUpMesh.cpp
+++ b/Engine/Source/Programs/Enterprise/Datasmith/DatasmithSketchUpRubyExporter/Private/DatasmithSketchUpMesh.cpp
@@ -14,6 +14,7 @@
 
 // SketchUp SDK.
 #include "DatasmithSketchUpSDKBegins.h"
+#include <SketchUpAPI/model/component_definition.h>
 #include "SketchUpAPI/model/drawing_element.h"
 #include "SketchUpAPI/model/edge.h"
 #include "SketchUpAPI/model/face.h"
@@ -426,7 +427,7 @@ void FEntities::UpdateGeometry(FExportContext& Context)
 	{
 		if (ExtractedMeshPtr->ContainsGeometry())
 		{
-			FString MeshElementName = FString::Printf(TEXT("M%ls_%d"), *FMD5::HashAnsiString(*Definition.GetSketchupSourceGUID()), MeshCount + 1); // Count meshes from 1
+			FString MeshElementName = FString::Printf(TEXT("M%ls_%d"), *Definition.GetSketchupSourceId(), MeshCount + 1); // Count meshes from 1
 			FString MeshLabel = FDatasmithUtils::SanitizeObjectName(Definition.GetSketchupSourceName());
 
 			TSharedPtr<FDatasmithInstantiatedMesh> Mesh;
@@ -531,21 +532,6 @@ void FEntities::RemoveMeshesFromDatasmithScene(FExportContext& Context)
 	}
 }
 
-TSharedPtr<IDatasmithMeshElement> FEntities::CreateMeshElement(FExportContext& Context, FDatasmithMesh& DatasmithMesh)
-{
-	FString MeshElementName = FString::Printf(TEXT("M%ls_%d"), *FMD5::HashAnsiString(*Definition.GetSketchupSourceGUID()), EntitiesGeometry->Meshes.Num() + 1); // Count meshes from 1
-	FString MeshLabel = FDatasmithUtils::SanitizeObjectName(Definition.GetSketchupSourceName());
-
-	FDatasmithMeshExporter DatasmithMeshExporter;
-	TSharedPtr<IDatasmithMeshElement> MeshElementPtr = DatasmithMeshExporter.ExportToUObject(Context.GetAssetsOutputPath(), *MeshElementName, DatasmithMesh, nullptr, FDatasmithExportOptions::LightmapUV);
-
-	// Set the mesh element label used in the Unreal UI.
-	MeshElementPtr->SetLabel(*MeshLabel);
-
-	Context.DatasmithScene->AddMesh(MeshElementPtr);
-	return MeshElementPtr;
-}
-
 TArray<SUGroupRef> FEntities::GetGroups()
 {
 	// Get the number of groups in the SketchUp model entities.
