[SketchUp] Classification Types exported to metadata

diff --git a/Engine/Source/Programs/Enterprise/Datasmith/DatasmithSketchUpRubyExporter/Private/DatasmithSketchUpMetadata.cpp b/Engine/Source/Programs/Enterprise/Datasmith/DatasmithSketchUpRubyExporter/Private/DatasmithSketchUpMetadata.cpp
--- a/Engine/Source/Programs/Enterprise/Datasmith/DatasmithSketchUpRubyExporter/Private/DatasmithSketchUpMetadata.cpp
+++ b/Engine/Source/Programs/Enterprise/Datasmith/DatasmithSketchUpRubyExporter/Private/DatasmithSketchUpMetadata.cpp
@@ -92,8 +92,19 @@ FMetadata::FMetadata(
 			size_t SSchemaCount = 0;
 			SUClassificationInfoGetNumSchemas(SClassificationInfoRef, &SSchemaCount);
 
+			TArray<FString> ClassificationSchemaTypes;
+			ClassificationSchemaTypes.Reserve(SSchemaCount);
+
 			for (int32 SSchemaNo = 0; SSchemaNo < SSchemaCount; SSchemaNo++)
 			{
+				SUStringRef ShemaTypeStringRef = SU_INVALID;
+				SUStringCreate(&ShemaTypeStringRef);
+				SUClassificationInfoGetSchemaType(SClassificationInfoRef, SSchemaNo, &ShemaTypeStringRef);
+				FString SchemaType = SuConvertString(ShemaTypeStringRef);
+				SUStringRelease(&ShemaTypeStringRef);
+				ClassificationSchemaTypes.Add(SchemaType);
+
+
 				// Retrieve the SketchUp classification schema attribute.
 				SUClassificationAttributeRef SSchemaAttributeRef = SU_INVALID;
 				SUClassificationInfoGetSchemaAttribute(SClassificationInfoRef, SSchemaNo, &SSchemaAttributeRef); // we can ignore the returned SU_RESULT
@@ -102,6 +113,11 @@ FMetadata::FMetadata(
 				ScanClassificationSchema(SSchemaAttributeRef);
 			}
 
+			if (SSchemaCount > 0)
+			{
+				MetadataKeyValueMap.Add(TEXT("ClassificationSchemaTypes"), FString::Join(ClassificationSchemaTypes, TEXT(",")));
+			}
+
 			// Release the SketchUp component classification info.
 			SUClassificationInfoRelease(&SClassificationInfoRef); // we can ignore the returned SU_RESULT
 		}
