[SketchUp] Image Entity

- Image entity exported as a separate actor
- layer visibility working(as well as regular visibility)
- Image dimensions baked into static mesh
- Image entities with the same image contents(bitmap data) share Datasmith material
- Sync of image parameters - transform and dimensions

diff --git a/Engine/Source/Programs/Enterprise/Datasmith/DatasmithSketchUpRubyExporter/Private/DatasmithSketchUpComponent.cpp b/Engine/Source/Programs/Enterprise/Datasmith/DatasmithSketchUpRubyExporter/Private/DatasmithSketchUpComponent.cpp
--- a/Engine/Source/Programs/Enterprise/Datasmith/DatasmithSketchUpRubyExporter/Private/DatasmithSketchUpComponent.cpp
+++ b/Engine/Source/Programs/Enterprise/Datasmith/DatasmithSketchUpRubyExporter/Private/DatasmithSketchUpComponent.cpp
@@ -20,6 +20,7 @@
 #include "SketchUpAPI/model/model.h"
 #include "SketchUpAPI/model/component_instance.h"
 #include "SketchUpAPI/geometry/transformation.h"
+#include "SketchUpAPI/geometry/vector3d.h"
 
 #if !defined(SKP_SDK_2019) && !defined(SKP_SDK_2020)
 #include "SketchUpAPI/model/layer_folder.h"
@@ -31,88 +32,65 @@
 #include "DatasmithSceneFactory.h"
 #include "DatasmithUtils.h"
 
-#define REMOVE_MESHES_WHEN_INVISIBLE
+#include "Misc/Paths.h"
 
-using namespace DatasmithSketchUp;
 
-void FNodeOccurence::AddChildOccurrence(FExportContext& Context, FComponentInstance& ChildComponentInstance)
-{
-	FNodeOccurence& ChildNode = ChildComponentInstance.CreateNodeOccurrence(Context, /*Parent*/ *this);
+#define REMOVE_MESHES_WHEN_INVISIBLE
 
-	// todo: rename ToDatasmith - this just creates hierarchy
-	ChildNode.ToDatasmith(Context);
-}
+using namespace DatasmithSketchUp;
 
-void FNodeOccurence::ToDatasmith(FExportContext& Context)
+void FDefinition::ParseNode(FExportContext& Context, FNodeOccurence& Node)
 {
-	FDefinition* EntityDefinition = Entity.GetDefinition();
-
-	if (!EntityDefinition)
-	{
-		return;
-	}
-
 	// Process child nodes
-	FEntities& Entities = EntityDefinition->GetEntities();
+	
 	// Convert the SketchUp normal component instances into sub-hierarchies of Datasmith actors.
-	for (SUComponentInstanceRef SComponentInstanceRef : Entities.GetComponentInstances())
+	for (SUComponentInstanceRef SComponentInstanceRef : GetEntities().GetComponentInstances())
 	{
-		TSharedPtr<FComponentInstance> ComponentInstance = Context.ComponentInstances.AddComponentInstance(*EntityDefinition, SComponentInstanceRef);
+		TSharedPtr<FComponentInstance> ComponentInstance = Context.ComponentInstances.AddComponentInstance(*this, SComponentInstanceRef);
 		if (ComponentInstance.IsValid())
 		{
-			AddChildOccurrence(Context, *ComponentInstance);
+			FNodeOccurence& ChildNode = ComponentInstance->CreateNodeOccurrence(Context, Node);
+			ComponentInstance->ParseNode(Context, ChildNode);
 		}
 	}
 
 	// Convert the SketchUp group component instances into sub-hierarchies of Datasmith actors.
-	for (SUGroupRef SGroupRef : Entities.GetGroups())
+	for (SUGroupRef SGroupRef : GetEntities().GetGroups())
 	{
 		SUComponentInstanceRef SComponentInstanceRef = SUGroupToComponentInstance(SGroupRef);
 
-		TSharedPtr<FComponentInstance> ComponentInstance = Context.ComponentInstances.AddComponentInstance(*EntityDefinition, SComponentInstanceRef);
+		TSharedPtr<FComponentInstance> ComponentInstance = Context.ComponentInstances.AddComponentInstance(*this, SComponentInstanceRef);
 		if (ComponentInstance.IsValid())
 		{
-			AddChildOccurrence(Context, *ComponentInstance);
+			FNodeOccurence& ChildNode = ComponentInstance->CreateNodeOccurrence(Context, Node);
+			ComponentInstance->ParseNode(Context, ChildNode);
 		}
 	}
-}
 
-void FNodeOccurence::UpdateMeshActors(FExportContext& Context)
-{
-	// Remove old mesh actors
-	// todo: reuse old mesh actors (also can keep instances when removing due to say hidden)
-	if (DatasmithActorElement.IsValid())
-	{
-		for (TSharedPtr<IDatasmithMeshActorElement> MeshActor : MeshActors)
-		{
-			DatasmithActorElement->RemoveChild(MeshActor);
-		}
-	}
-	else
+	for (SUImageRef ImageRef : GetEntities().GetImages())
 	{
-		for (TSharedPtr<IDatasmithMeshActorElement> MeshActor : MeshActors)
+		TSharedPtr<FImage> Image = Context.Images.AddImage(*this, ImageRef);
+		if (Image.IsValid())
 		{
-			Context.DatasmithScene->RemoveActor(MeshActor, EDatasmithActorRemovalRule::RemoveChildren);
+			Image->CreateNodeOccurrence(Context, Node);
 		}
 	}
-	MeshActors.Reset();
-
-	if (!bVisible) // no mesh actors for invisible node
-	{
-		return;
-	}
+}
 
-	FDefinition* EntityDefinition = Entity.GetDefinition();
+// Update mesh actors of an entiry with Entities(that is model or component instance)
+bool FEntityWithEntities::UpdateOccurrenceMeshActors(FExportContext& Context, FNodeOccurence& Node)
+{
+	FDefinition* EntityDefinition = GetDefinition();
 	DatasmithSketchUp::FEntitiesGeometry* EntitiesGeometry = EntityDefinition->GetEntities().EntitiesGeometry.Get();
 
 	if (!EntitiesGeometry)
 	{
-		return;
+		return true;
 	}
 
-	MeshActors.Reset(EntitiesGeometry->GetMeshCount());
+	Node.MeshActors.Reset(EntitiesGeometry->GetMeshCount());
 
-	FString ComponentActorName = GetActorName();
+	FString ComponentActorName = Node.GetActorName();
 	
 	for (int32 MeshIndex = 0; MeshIndex < EntitiesGeometry->GetMeshCount(); ++MeshIndex)
 	{
@@ -121,11 +99,11 @@ void FNodeOccurence::UpdateMeshActors(FExportContext& Context)
 		// Create a Datasmith mesh actor for the Datasmith mesh element.
 		TSharedPtr<IDatasmithMeshActorElement> DMeshActorPtr = FDatasmithSceneFactory::CreateMeshActor(*MeshActorName);
 
-		MeshActors.Add(DMeshActorPtr);
+		Node.MeshActors.Add(DMeshActorPtr);
 
 		// Add the Datasmith actor component depth tag.
 		// We use component depth + 1 to factor in the added Datasmith scene root once imported in Unreal.
-		FString ComponentDepthTag = FString::Printf(TEXT("SU.DEPTH.%d"), Depth + 1);
+		FString ComponentDepthTag = FString::Printf(TEXT("SU.DEPTH.%d"), Node.Depth + 1);
 		DMeshActorPtr->AddTag(*ComponentDepthTag);
 
 		// Add the Datasmith actor component definition GUID tag.
@@ -137,9 +115,9 @@ void FNodeOccurence::UpdateMeshActors(FExportContext& Context)
 		DMeshActorPtr->AddTag(*InstancePathTag);
 
 		// Add the mesh actor to our component Datasmith actor hierarchy.
-		if (DatasmithActorElement.IsValid()) 
+		if (Node.DatasmithActorElement.IsValid()) 
 		{
-			DatasmithActorElement->AddChild(DMeshActorPtr);
+			Node.DatasmithActorElement->AddChild(DMeshActorPtr);
 		}
 		else
 		{
@@ -151,6 +129,35 @@ void FNodeOccurence::UpdateMeshActors(FExportContext& Context)
 		// Set the Datasmith mesh element used by the mesh actor.
 		DMeshActorPtr->SetStaticMeshPathName(EntitiesGeometry->GetMeshElementName(MeshIndex));
 	}
+	return false;
+}
+
+void FNodeOccurence::UpdateMeshActors(FExportContext& Context)
+{
+	// Remove old mesh actors
+	// todo: reuse old mesh actors (also can keep instances when removing due to say hidden)
+	if (DatasmithActorElement.IsValid())
+	{
+		for (TSharedPtr<IDatasmithMeshActorElement> MeshActor : MeshActors)
+		{
+			DatasmithActorElement->RemoveChild(MeshActor);
+		}
+	}
+	else
+	{
+		for (TSharedPtr<IDatasmithMeshActorElement> MeshActor : MeshActors)
+		{
+			Context.DatasmithScene->RemoveActor(MeshActor, EDatasmithActorRemovalRule::RemoveChildren);
+		}
+	}
+	MeshActors.Reset();
+
+	if (!bVisible) // no mesh actors for invisible node
+	{
+		return;
+	}
+
+	Entity.UpdateOccurrenceMeshActors(Context, *this);
 }
 
 void FNodeOccurence::UpdateVisibility(FExportContext& Context)
@@ -237,19 +244,27 @@ void FNodeOccurence::RemoveOccurrence(FExportContext& Context)
 {
 	// RemoveOccurrence is called from Entity only(i.e. it doesn't remove occurrence from the Entity itself, it's done there)
 
-	Entity.EntityOccurrenceVisible(this, false);
-
 	if (MaterialOverride)
 	{
 		MaterialOverride->UnregisterInstance(Context, this);
 	}
 
-	for (FNodeOccurence* Child : Children)
+	if (ParentNode)
+	{
+		ParentNode->Children.Remove(this);
+	}
+
+	// Usually child component instances are removed in proper order but grouping entities
+	// has this weird behavior that containing component removed without cleaning. Test case:
+	// Group an instance and some other entity(e.g. face) and then convert group to component
+	// then converted group is removed without other events for its children
+	// Probably this is because those 'children' are actually entities in the group's definition
+	// And Group entity itself is just an instance of its definition(just like ComponentInstance)
+	// So that Group definition content is not changed and that definition just receives another instance
+	for (FNodeOccurence* Child : Children.Array())
 	{
 		Child->RemoveOccurrence(Context);
-		Child->Entity.DeleteOccurrence(Context, Child);
 	}
-	Children.Reset(); // don't need this is node it deallocated?
 
 	for (const TSharedPtr<IDatasmithMeshActorElement>& MeshActor : MeshActors) 
 	{
@@ -274,6 +289,8 @@ void FNodeOccurence::RemoveOccurrence(FExportContext& Context)
 			Context.DatasmithScene->RemoveActor(DatasmithActorElement, EDatasmithActorRemovalRule::RemoveChildren);
 		}
 	}
+
+	Entity.DeleteOccurrence(Context, this);
 }
 
 void FNodeOccurence::ResetMetadataElement(FExportContext& Context)
@@ -298,7 +315,6 @@ void FNodeOccurence::ResetMetadataElement(FExportContext& Context)
 void FNodeOccurence::SetVisibility(bool bValue)
 {
 	bVisible = bValue;
-	Entity.EntityOccurrenceVisible(this, bVisible);
 }
 
 void FNodeOccurence::RemoveDatasmithActorHierarchy(FExportContext& Context)
@@ -414,15 +430,26 @@ FString FModelDefinition::GetSketchupSourceGUID()
 
 void FModelDefinition::AddInstance(FExportContext& Context, TSharedPtr<FComponentInstance> Instance)
 {
-	Context.RootNode->AddChildOccurrence(Context, *Instance);
+	FNodeOccurence& ChildNode = Instance->CreateNodeOccurrence(Context, *Context.RootNode);
+	Instance->ParseNode(Context, ChildNode);
 }
 
+void FModelDefinition::AddImage(FExportContext& Context, TSharedPtr<FImage> Image)
+{
+	Image->CreateNodeOccurrence(Context, *Context.RootNode);
+}
+
+
 FComponentDefinition::FComponentDefinition(
 	SUComponentDefinitionRef InComponentDefinitionRef)
 	: ComponentDefinitionRef(InComponentDefinitionRef)
 {
 }
 
+FComponentDefinition::~FComponentDefinition()
+{
+}
+
 void FComponentDefinition::Parse(FExportContext& Context)
 {
 	SUEntitiesRef EntitiesRef = SU_INVALID;
@@ -499,7 +526,7 @@ void FComponentDefinition::BuildNodeNames(FNodeOccurence& Node)
 	Node.DatasmithActorName = FString::Printf(TEXT("%ls_%lld"), *Node.ParentNode->GetActorName(), SketchupPersistentID);
 
 	FString EntityName = Node.Entity.GetName();
-	Node.DatasmithActorLabel = FDatasmithUtils::SanitizeObjectName(EntityName.IsEmpty() ? Node.Entity.GetDefinition()->GetSketchupSourceName() : EntityName);
+	Node.DatasmithActorLabel = FDatasmithUtils::SanitizeObjectName(EntityName.IsEmpty() ? GetSketchupSourceName() : EntityName);
 }
 
 void FComponentDefinition::InvalidateInstancesGeometry(FExportContext& Context)
@@ -552,7 +579,7 @@ FString FComponentDefinition::GetSketchupSourceName()
 
 FString FComponentDefinition::GetSketchupSourceId()
 {
-// Although implemented since SUEntityGetPersistentID since SU 2017 is returns valid Id for Definitions
+// Although SUEntityGetPersistentID implemented since SU 2017 it returns valid Id for ComponentDefinitions
 // only since SU 2020.1 (even though SUEntityGetPersistentID docs states SUComponentDefinitionRef 'supported' from 2017)
 // see [REDACTED_PROFILE_URL]/api-issue-tracker/issues/314
 #ifndef SKP_SDK_2019
@@ -729,12 +756,25 @@ void FComponentDefinition::AddInstance(FExportContext& Context, TSharedPtr<FComp
 	{
 		for (FNodeOccurence* ParentOccurrence : ParentInstance->Occurrences)
 		{
-			ParentOccurrence->AddChildOccurrence(Context, *Instance);
+			FNodeOccurence& ChildNode = Instance->CreateNodeOccurrence(Context, *ParentOccurrence);
+			Instance->ParseNode(Context, ChildNode);
 		}
 	}
 }
 
-void FEntity::UpdateOccurrence(FExportContext& Context, FNodeOccurence& Node)
+void FComponentDefinition::AddImage(FExportContext& Context, TSharedPtr<FImage> Image)
+{
+	for (FComponentInstance* ParentInstance : Instances)
+	{
+		for (FNodeOccurence* ParentOccurrence : ParentInstance->Occurrences)
+		{
+			// todo: remove djupcacation with FComponentDefinition::AddInstance
+			Image->CreateNodeOccurrence(Context, *ParentOccurrence);
+		}
+	}
+}
+
+void FEntityWithEntities::UpdateOccurrence(FExportContext& Context, FNodeOccurence& Node)
 {
 	if (Node.MaterialOverride)
 	{
@@ -769,6 +809,30 @@ void FEntity::UpdateOccurrence(FExportContext& Context, FNodeOccurence& Node)
 	}
 }
 
+FNodeOccurence& FEntity::CreateNodeOccurrence(FExportContext& Context, FNodeOccurence& ParentNode)
+{
+	FNodeOccurence* Occurrence = new FNodeOccurence(&ParentNode, *this);
+	ParentNode.Children.Add(Occurrence);
+	Occurrences.Add(Occurrence);
+	return *Occurrence;
+}
+
+void FEntity::DeleteOccurrence(FExportContext& Context, FNodeOccurence* Node)
+{
+	EntityOccurrenceVisible(Node, false);
+	Occurrences.Remove(Node);
+	delete Node;
+}
+
+void FEntity::RemoveOccurrences(FExportContext& Context)
+{
+	TArray<FNodeOccurence*> OccurencesCopy = Occurrences; // Copy RemoveOccurrence modifies the array
+	for (FNodeOccurence* Occurrence : OccurencesCopy)
+	{
+		Occurrence->RemoveOccurrence(Context);
+	}
+}
+
 void FEntity::UpdateEntityGeometry(FExportContext& Context)
 {
 	if (bGeometryInvalidated)
@@ -805,10 +869,35 @@ void FEntity::EntityOccurrenceVisible(FNodeOccurence* Node, bool bVisible)
 			VisibleNodes.Remove(Node);
 		}
 	}
+}
+
+void FEntity::SetParentDefinition(FExportContext& Context, FDefinition* InParent)
+{
+	if (!IsParentDefinition(InParent)) // Changing parent
+	{
+		// If we are re-parenting(i.e. entity was previously owned by another Definition - this happens
+		// when say a ComponentInstance was selected in UI and "Make Group" was performed.
+		if (Parent)
+		{
+			RemoveOccurrences(Context);
+			Occurrences.Reset();  // Clear occurrences - RemoveOccurrences doesn't do it(not needed during ComponentInstance removal)
+		}
+
+		Parent = InParent;
+	}
+}
+
+void FEntityWithEntities::EntityOccurrenceVisible(FNodeOccurence* Node, bool bVisible)
+{
+	Super::EntityOccurrenceVisible(Node, bVisible);
 
 	GetDefinition()->EntityVisible(this, VisibleNodes.Num() > 0);
 }
 
+FComponentInstance::~FComponentInstance()
+{
+}
+
 FDefinition* FComponentInstance::GetDefinition()
 {
 	return &Definition;
@@ -893,20 +982,6 @@ FString FComponentInstance::GetName()
 	return SuGetString(SUComponentInstanceGetName, InSComponentInstanceRef);
 }
 
-FNodeOccurence& FComponentInstance::CreateNodeOccurrence(FExportContext& Context, FNodeOccurence& ParentNode)
-{
-	FNodeOccurence* Occurrence = new FNodeOccurence(&ParentNode, *this);
-	ParentNode.Children.Add(Occurrence);
-	Occurrences.Add(Occurrence);
-	return *Occurrence;
-}
-
-void FComponentInstance::DeleteOccurrence(FExportContext& Context, FNodeOccurence* Node)
-{
-	Occurrences.Remove(Node);
-	delete Node;
-}
-
 void FComponentInstance::UpdateMetadata(FExportContext& Context)
 {
 	ParsedMetadata = MakeUnique<FMetadata>(SUComponentInstanceToEntity(GetComponentInstanceRef()));
@@ -922,6 +997,18 @@ void FComponentInstance::UpdateEntityProperties(FExportContext& Context)
 	FEntity::UpdateEntityProperties(Context);
 }
 
+void FComponentInstance::ParseNode(FExportContext& Context, FNodeOccurence& Node)
+{
+	FDefinition* EntityDefinition = GetDefinition();
+
+	if (!EntityDefinition)
+	{
+		return;
+	}
+
+	EntityDefinition->ParseNode(Context, Node);
+}
+
 void FComponentInstance::InvalidateOccurrencesGeometry(FExportContext& Context)
 {
 	for (FNodeOccurence* Node : Occurrences)
@@ -971,15 +1058,6 @@ FComponentInstanceIDType FComponentInstance::GetComponentInstanceId()
 	return DatasmithSketchUpUtils::GetComponentInstanceID(GetComponentInstanceRef());
 }
 
-void FComponentInstance::RemoveOccurrences(FExportContext& Context)
-{
-	for (FNodeOccurence* Occurrence : Occurrences)
-	{
-		Occurrence->RemoveOccurrence(Context);
-		Occurrence->ParentNode->Children.Remove(Occurrence);
-	}
-}
-
 SUComponentInstanceRef FComponentInstance::GetComponentInstanceRef()
 {
 	return SUComponentInstanceFromEntity(EntityRef);
@@ -1019,11 +1097,26 @@ void FComponentInstance::FillOccurrenceActorMetadata(FNodeOccurence& Node)
 	GetDefinition()->FillOccurrenceActorMetadata(Node);
 }
 
+
+void FImageCollection::LayerModified(FEntityIDType LayerId)
+{
+	for (const TPair<FEntityIDType, TSharedPtr<FImage>>& IdImage : Images)
+	{
+		TSharedPtr<FImage> Image = IdImage.Value;
+		if (SUIsValid(Image->LayerRef) && (LayerId == DatasmithSketchUpUtils::GetEntityID(SULayerToEntity(Image->LayerRef))))
+		{
+			Image->InvalidateEntityProperties();
+		}
+	}
+}
+
 void FComponentInstance::UpdateOccurrenceVisibility(FExportContext& Context, FNodeOccurence& Node)
 {
 	// Parent node, component instance and layer - all should be visible to have node visible
 	Node.SetVisibility(Node.ParentNode->bVisible && !bHidden && bLayerVisible);
 
+	EntityOccurrenceVisible(&Node, Node.bVisible);
+
 	if (Node.bVisible)
 	{
 		if (!Node.DatasmithActorElement)
@@ -1032,7 +1125,6 @@ void FComponentInstance::UpdateOccurrenceVisibility(FExportContext& Context, FNo
 			EntityDefinition->CreateActor(Context, Node);
 
 			// Invalidate actor-dependent 
-			Node.InvalidateMeshActors();
 			Node.InvalidateProperties();
 		}
 		Node.InvalidateMeshActors();
@@ -1065,22 +1157,7 @@ void FComponentInstance::RemoveComponentInstance(FExportContext& Context)
 	}
 }
 
-void FComponentInstance::SetParentDefinition(FExportContext& Context, FDefinition* InParent)
-{
-	if (!IsParentDefinition(InParent)) // Changing parent
-	{
-		// If we are re-parenting(i.e. entity was previously owned by another Definition - this happens
-		// when say a ComponentInstance was selected in UI and "Make Group" was performed.
-		if (Parent)
-		{
-			RemoveOccurrences(Context);
-		}
-
-		Parent = InParent;
-	}
-}
-
-FModel::FModel(FModelDefinition& InDefinition) : Definition(InDefinition)
+FModel::FModel(FModelDefinition& InDefinition) : Super(SU_INVALID), Definition(InDefinition)
 {
 
 }
@@ -1120,17 +1197,22 @@ void FModel::InvalidateOccurrencesProperties(FExportContext& Context)
 void FModel::UpdateOccurrenceVisibility(FExportContext& Context, FNodeOccurence& Node)
 {
 	Node.SetVisibility(true);
-}
 
-void FModel::DeleteOccurrence(FExportContext& Context, FNodeOccurence* Node)
-{
-	// Model occurrence is not deleted by any parent
+	EntityOccurrenceVisible(&Node, true);
 }
 
 void FModel::UpdateMetadata(FExportContext& Context)
 {
 }
 
+FDefinition::FDefinition(): bMeshesAdded(false)
+                            , bGeometryInvalidated(true)
+                            , bPropertiesInvalidated(true)
+{}
+
+FDefinition::~FDefinition()
+{}
+
 void FDefinition::EntityVisible(FEntity* Entity, bool bVisible)
 {
 	if (bVisible)
@@ -1161,6 +1243,9 @@ void FDefinition::UpdateDefinition(FExportContext& Context)
 
 		if (bPropertiesInvalidated)
 		{
+			// Currently SketchUp has no Observer for Component Definition attributes.
+			// So this code is only executed on export
+			// todo: implement attributes sync once api is available
 			UpdateMetadata(Context);
 			InvalidateInstancesMetadata(Context); // Make sure instances keep up with definition changes
 
diff --git a/Engine/Source/Programs/Enterprise/Datasmith/DatasmithSketchUpRubyExporter/Private/DatasmithSketchUpComponent.h b/Engine/Source/Programs/Enterprise/Datasmith/DatasmithSketchUpRubyExporter/Private/DatasmithSketchUpComponent.h
--- a/Engine/Source/Programs/Enterprise/Datasmith/DatasmithSketchUpRubyExporter/Private/DatasmithSketchUpComponent.h
+++ b/Engine/Source/Programs/Enterprise/Datasmith/DatasmithSketchUpRubyExporter/Private/DatasmithSketchUpComponent.h
@@ -19,6 +19,9 @@
 	class FCamera;
 	class FComponentDefinition;
 	class FComponentInstance;
+	class FDefinition;
+	class FImage;
+	class FImageMaterial;
 	class FEntities;
 	class FEntitiesGeometry;
 	class FEntity;
@@ -57,8 +60,6 @@
 		void UpdateVisibility(FExportContext& Context);
 		void Update(FExportContext& Context);
 
-		void AddChildOccurrence(FExportContext& Context, FComponentInstance& ComponentInstance);
-
 		void UpdateMeshActors(FExportContext& Context); // Create/Free Datasmith Mesh Actors for meshes in the node's Entities
 
 		void InvalidateProperties(); // Invalidate name and transform. Invalidate propagates down the hierarchy - child transforms depend on the parent
@@ -67,8 +68,6 @@
 
 		void RemoveDatasmithActorHierarchy(FExportContext& Context);
 
-		void ToDatasmith(FExportContext& Context); // Build actor hierarchy 
-
 		FString GetActorName();
 
 		FString GetActorLabel();
@@ -113,14 +112,9 @@
 	{
 	public:
 
-		FDefinition()
-			: bMeshesAdded(false)
-			, bGeometryInvalidated(true)
-			, bPropertiesInvalidated(true)
-		{}
+		FDefinition();
 
+		virtual ~FDefinition();
 
-		virtual ~FDefinition() {}
-
 		virtual void Parse(FExportContext& Context) = 0;
 		virtual void CreateActor(FExportContext& Context, FNodeOccurence& Node) = 0; // Create Datasmith actor for node occurrence
@@ -132,6 +126,8 @@
 
 		// Modification methods
 		virtual void AddInstance(FExportContext& Context, TSharedPtr<FComponentInstance> Instance) = 0; // Register child CompoenntInstance Entity of Definition's Entities
+		virtual void AddImage(FExportContext& Context, TSharedPtr<FImage> Image) = 0;
+
 		virtual void InvalidateInstancesGeometry(FExportContext& Context) = 0; // Mark that all instances(and their occurrences) needed to be updated
 		virtual void InvalidateInstancesMetadata(FExportContext& Context) = 0; // Mark that all instances(and their occurrences) needed to be updated
 		virtual void FillOccurrenceActorMetadata(FNodeOccurence& Node) = 0;
@@ -153,6 +149,8 @@
 		}
 
 		void UpdateDefinition(FExportContext& Context);
+		void ParseNode(FExportContext& Context, FNodeOccurence& Node);  // Parse hierarchy of a child node('child' meaning 'from this definition's entities')
+		void ApplyOverrideMaterialToNode(FNodeOccurence& Node, FMaterialOccurrence& Material);
 
 	protected:
 
@@ -173,6 +171,8 @@
 			SUComponentDefinitionRef InComponentDefinitionRef // source SketchUp component definition
 		);
 
+		virtual ~FComponentDefinition() override;
+
 		// Begin FDefinition
 		virtual void Parse(FExportContext& Context) override;
 		virtual void CreateActor(FExportContext& Context, FNodeOccurence& Node) override;
@@ -181,6 +181,8 @@
 		virtual void UpdateMetadata(FExportContext& Context) override;
 
 		virtual void AddInstance(FExportContext& Context, TSharedPtr<FComponentInstance> Instance) override;
+		virtual void AddImage(FExportContext& Context, TSharedPtr<FImage> Image) override;
+
 		virtual void InvalidateInstancesGeometry(FExportContext& Context) override;
 		virtual void InvalidateInstancesMetadata(FExportContext& Context) override;
 		virtual void FillOccurrenceActorMetadata(FNodeOccurence& Node) override;
@@ -232,6 +234,8 @@
 		virtual void UpdateMetadata(FExportContext& Context) override;
 
 		virtual void AddInstance(FExportContext& Context, TSharedPtr<FComponentInstance> Instance) override;
+		virtual void AddImage(FExportContext& Context, TSharedPtr<FImage> Image) override;
+
 		virtual void InvalidateInstancesGeometry(FExportContext& Context) override;
 		virtual void InvalidateInstancesMetadata(FExportContext& Context) override;
 		virtual void FillOccurrenceActorMetadata(FNodeOccurence& Node) override;
@@ -262,8 +266,8 @@
 		void RemoveMeshesFromDatasmithScene(FExportContext& Context);
 
 		TArray<SUGroupRef> GetGroups();
-
 		TArray<SUComponentInstanceRef> GetComponentInstances();
+		TArray<SUImageRef> GetImages();
 
 		// Source SketchUp component entities.
 		SUEntitiesRef EntitiesRef = SU_INVALID;
@@ -302,39 +306,45 @@
 			bDefaultMaterialUsed = true;
 		}
 
-
 		// todo: update reusing datasmith elements? 
-		// todo: merge ALL faces that are present in Entities into single mesh? do we really need separate mesh for every isolated set of faces?
-		// todo: occurrences using these entities must ne referenced
 		// todo: update occurrences that used this entities - MeshActors need to be refreshed in accordance to OR this could be done on a level higher?
 	};
 
-	// Interface to implement SketchUp Entity Node(i.e. an instance of a ComponentDefinition - ComponentInstance or Group) access
-	// todo: rename this to NodeEntity? This class represents on any entity but only those that build scene hierarchy(Model, ComponentInstance, Group)
+	// Interface to implement SketchUp Entity Node - i.e. an instance of a ComponentDefinition(ComponentInstance or Group), Model, Image
 	class FEntity : FNoncopyable
 	{
 	public:
 
-		FEntity() 
-			: bGeometryInvalidated(true)
+		FEntity(SUEntityRef InEntityRef)
+			: EntityRef(InEntityRef)
+			, bGeometryInvalidated(true)
 			, bPropertiesInvalidated(true)
 		{}
 
 		virtual ~FEntity() {}
 
-		virtual FDefinition* GetDefinition() = 0;
-		virtual bool GetAssignedMaterial(FMaterialIDType& MaterialId) = 0; // Get material of this entity
-		virtual void InvalidateOccurrencesGeometry(FExportContext& Context) = 0;
-		virtual void UpdateOccurrence(FExportContext& Context, FNodeOccurence& Node); // Update occurrence of this entity
-		virtual void InvalidateOccurrencesProperties(FExportContext& Context) = 0;
 		virtual int64 GetPersistentId() = 0;
 		virtual FString GetName() = 0;
+
+		virtual void ApplyOverrideMaterialToNode(FNodeOccurence& Node, FMaterialOccurrence& Material) = 0;
+
+		virtual void UpdateOccurrence(FExportContext& Context, FNodeOccurence& Node) = 0; // Update occurrence of this entity
+		// Create an occurrence of this ComponentInstance (component instance can appear multiple times in SketchUp hierarchy)
 		virtual void UpdateOccurrenceVisibility(FExportContext& Context, FNodeOccurence&) = 0;
-		virtual void DeleteOccurrence(FExportContext& Context, FNodeOccurence* Node) = 0;
+		virtual bool UpdateOccurrenceMeshActors(FExportContext& Context, FNodeOccurence& Node) = 0;
+
+		virtual void InvalidateOccurrencesGeometry(FExportContext& Context) = 0;
+		virtual void InvalidateOccurrencesProperties(FExportContext& Context) = 0;
+
+		virtual void EntityOccurrenceVisible(FNodeOccurence* Node, bool bUses);
+
+		virtual void UpdateEntityProperties(FExportContext& Context);
 		virtual void UpdateMetadata(FExportContext& Context) = 0;
 
-		void EntityOccurrenceVisible(FNodeOccurence* Node, bool bUses);
+		FNodeOccurence& CreateNodeOccurrence(FExportContext& Context, FNodeOccurence& ParentNode);
+		void DeleteOccurrence(FExportContext& Context, FNodeOccurence* Node);
+		void RemoveOccurrences(FExportContext& Context);
+		/////
 
-
 		// Invalidates transform, name
 		void InvalidateEntityProperties()
@@ -347,27 +357,58 @@
 			bGeometryInvalidated = true;
 		}
 
-		virtual void UpdateEntityProperties(FExportContext& Context);
 		void UpdateEntityGeometry(FExportContext& Context);
 
-		TSet<FNodeOccurence*> VisibleNodes;
+		void SetParentDefinition(FExportContext& Context, FDefinition* InParent);
+		bool IsParentDefinition(FDefinition* InParent)
+		{
+			return Parent == InParent;
+		}
 
+		FDefinition* Parent = nullptr;
+
+		SUEntityRef EntityRef = SU_INVALID;
+
+		TArray<FNodeOccurence*> Occurrences; // All occurrencies of this Entity in Model hierarchy
+		TSet<FNodeOccurence*> VisibleNodes; // Occurrences currently fully visible
+
 		uint8 bGeometryInvalidated:1;
 		uint8 bPropertiesInvalidated:1;
 	};
 
-	class FComponentInstance : public FEntity
+	// Entity that has SUEntities children (Component or model)
+	class FEntityWithEntities: public FEntity
 	{
 		using Super = FEntity;
 	public:
-		SUEntityRef EntityRef = SU_INVALID;
+
+		FEntityWithEntities(SUEntityRef InEntityRef)
+			: Super(InEntityRef)
+		{}
+
+		virtual void UpdateOccurrence(FExportContext& Context, FNodeOccurence& Node); // Update occurrence of this entity
+		
+		virtual bool UpdateOccurrenceMeshActors(FExportContext& Context, FNodeOccurence& Node);
+
+		virtual void EntityOccurrenceVisible(FNodeOccurence* Node, bool bUses) override;
+
+		virtual FDefinition* GetDefinition() = 0;
+		virtual bool GetAssignedMaterial(FMaterialIDType& MaterialId) = 0; // Get material of this entity
+	};
+
+	class FComponentInstance : public FEntityWithEntities
+	{
+		using Super = FEntityWithEntities;
+	public:
 		FComponentDefinition& Definition;
 
 		FComponentInstance(SUEntityRef InEntityRef, FComponentDefinition& InDefinition)
-			: EntityRef(InEntityRef)
+			: Super(InEntityRef)
 			, Definition(InDefinition)
 		{}
 
+		virtual ~FComponentInstance() override;
+
 		// >>> FEntity
 		virtual FDefinition* GetDefinition() override;
 		virtual bool GetAssignedMaterial(FMaterialIDType& MaterialId) override;
@@ -377,26 +418,14 @@
 		virtual int64 GetPersistentId() override;
 		virtual FString GetName() override;
 		virtual void UpdateOccurrenceVisibility(FExportContext& Context, FNodeOccurence&) override;
-		virtual void DeleteOccurrence(FExportContext& Context, FNodeOccurence* Node) override;
 		virtual void UpdateMetadata(FExportContext& Context) override;
 		virtual void UpdateEntityProperties(FExportContext& Context) override;
+		void ApplyOverrideMaterialToNode(FNodeOccurence& Node, FMaterialOccurrence& Material);
 		// <<< FEntity
 
-
-
-		// Set Definition which Entities contains this entity
-		void SetParentDefinition(FExportContext& Context, FDefinition* InParent);
-		bool IsParentDefinition(FDefinition* InParent)
-		{
-			return Parent == InParent;
-		}
-
+		void ParseNode(FExportContext& Context, FNodeOccurence&);
 		void RemoveComponentInstance(FExportContext& Context);
 
-		// Create an occurrence of this ComponentInstance (component instance can appear multiple times in SketchUp hierarchy)
-		FNodeOccurence& CreateNodeOccurrence(FExportContext& Context, FNodeOccurence& ParentNode);
-		void RemoveOccurrences(FExportContext& Context);
-
 		FComponentInstanceIDType GetComponentInstanceId();
 		SUComponentInstanceRef GetComponentInstanceRef(); 
 
@@ -406,15 +435,64 @@
 		SULayerRef LayerRef = SU_INVALID;
 		bool bLayerVisible = true;
 
-		TArray<FNodeOccurence*> Occurrences;
-
-		FDefinition* Parent = nullptr;
-
 		TUniquePtr<FMetadata> ParsedMetadata;
 	};
 
-	class FModel : public FEntity
+	class FImage : public FEntity
 	{
+		using Super = FEntity;
 	public:
+
+		FImage(SUImageRef InEntityRef);
+
+		virtual ~FImage() override
+		{
+		}
+
+		virtual int64 GetPersistentId() override;
+		virtual FString GetName() override;
+		virtual void ApplyOverrideMaterialToNode(FNodeOccurence& Node, FMaterialOccurrence& Material) override;
+		virtual void UpdateOccurrence(FExportContext& Context, FNodeOccurence& Node) override;
+		void RemoveImageFromDatasmithScene(FExportContext& Context);
+		virtual void UpdateOccurrenceVisibility(FExportContext& Context, FNodeOccurence&) override;
+		virtual bool UpdateOccurrenceMeshActors(FExportContext& Context, FNodeOccurence& Node) override;
+		virtual void InvalidateOccurrencesGeometry(FExportContext& Context) override;
+		virtual void InvalidateOccurrencesProperties(FExportContext& Context) override;
+		virtual void UpdateMetadata(FExportContext& Context) override;
+		virtual void EntityOccurrenceVisible(FNodeOccurence* Node, bool bUses) override;
+
+		void UpdateEntityProperties(FExportContext& Context);
+
+		const TCHAR* GetMeshElementName();
+		void Update(FExportContext& Context);
+		void UpdateGeometry(FExportContext& Context);
+		const TCHAR* GetDatasmithTextureElementName();
+
+		FString GetFileName();
+		void InvalidateImage();
+
+		FString MeshElementName;
+
+		void BuildNodeNames(FNodeOccurence& Node);
+		void CreateActor(FExportContext& Context, FNodeOccurence& Node);
+
+		void RemoveImage(FExportContext& Context);
+
+		bool bHidden = false;
+		SULayerRef LayerRef = SU_INVALID;
+		bool bLayerVisible = true;
+
+		TUniquePtr<FMetadata> ParsedMetadata;
+
+		TSharedPtr<IDatasmithMeshElement> DatasmithMeshElement;
+
+		FImageMaterial* ImageMaterial = nullptr;
+
+	};
+
+	class FModel : public FEntityWithEntities
+	{
+		using Super = FEntityWithEntities;
+	public:
 		FModel(class FModelDefinition& InDefinition);
 
@@ -427,9 +505,11 @@
 		int64 GetPersistentId() override;
 		FString GetName() override;
 		void UpdateOccurrenceVisibility(FExportContext& Context, FNodeOccurence&) override;
-		void DeleteOccurrence(FExportContext& Context, FNodeOccurence* Node) override;
 		void UpdateMetadata(FExportContext& Context) override;
+		void ApplyOverrideMaterialToNode(FNodeOccurence& Node, FMaterialOccurrence& Material);
+
 		// <<< FEntity
+
 	private:
 		FModelDefinition& Definition;
 	};
diff --git a/Engine/Source/Programs/Enterprise/Datasmith/DatasmithSketchUpRubyExporter/Private/DatasmithSketchUpExportContext.cpp b/Engine/Source/Programs/Enterprise/Datasmith/DatasmithSketchUpRubyExporter/Private/DatasmithSketchUpExportContext.cpp
--- a/Engine/Source/Programs/Enterprise/Datasmith/DatasmithSketchUpRubyExporter/Private/DatasmithSketchUpExportContext.cpp
+++ b/Engine/Source/Programs/Enterprise/Datasmith/DatasmithSketchUpRubyExporter/Private/DatasmithSketchUpExportContext.cpp
@@ -29,6 +29,8 @@
 #include "SketchUpAPI/model/camera.h"
 #include "SketchUpAPI/model/component_definition.h"
 #include "SketchUpAPI/model/component_instance.h"
+#include <SketchUpAPI/model/image.h>
+
 #include "SketchUpAPI/model/layer.h"
 #include "SketchUpAPI/model/model.h"
 #include <SketchUpAPI/model/rendering_options.h>
@@ -51,6 +53,8 @@ FExportContext::FExportContext()
 	, Scenes(*this)
 	, Textures(*this)
 	, Layers(*this)
+	, Images(*this)
+	, ImageFiles(*this)
 {
 }
 
@@ -94,7 +98,7 @@ void FExportContext::Populate()
 	Scenes.PopulateFromModel(ModelRef);
 	ComponentDefinitions.PopulateFromModel(ModelRef);
 
-	RootNode->ToDatasmith(*this);
+	ModelDefinition->ParseNode(*this, *RootNode); // Create node hierarchy
 }
 
 // ColorByLayer(or ColorByTag) changes material display on every entity by useinf either regular materials or colors(materials) set up for layers
@@ -168,7 +172,8 @@ bool FExportContext::Update(bool bModifiedHint)
 
 	// Invalidate occurrences for changed instances first
 	Model->UpdateEntityProperties(*this);
-	ComponentInstances.UpdateProperties(); 
+	ComponentInstances.UpdateProperties();
+	Images.UpdateProperties();
 
 	// Update occurrences visibility(before updating meshes to make sure to skip updating unused meshes)
 	RootNode->UpdateVisibility(*this);
@@ -176,6 +181,7 @@ bool FExportContext::Update(bool bModifiedHint)
 	// Update Datasmith Meshes after their usage was refreshed(in visibility update) and before node hierarchy update(where Mesh Actors are updated for meshes)
 	ModelDefinition->UpdateDefinition(*this);
 	ComponentDefinitions.Update();
+	Images.Update();
 
 	// ComponentInstances will invalidate occurrences 
 	Model->UpdateEntityGeometry(*this);
@@ -242,6 +248,59 @@ void FComponentDefinitionCollection::Update()
 	}
 }
 
+void FImageCollection::Update()
+{
+	for (const auto& IdValue : Images)
+	{
+		IdValue.Value->Update(Context);
+	}
+}
+
+void FImageCollection::UpdateProperties()
+{
+	for (const auto& IdValue : Images)
+	{
+		IdValue.Value->UpdateEntityProperties(Context);
+	}
+}
+
+bool FImageCollection::InvalidateImage(const FEntityIDType& EntityID)
+{
+	if (TSharedPtr<FImage>* Found = Images.Find(EntityID))
+	{
+		(*Found)->InvalidateImage();
+		return true;
+	}
+	return false;
+}
+
+bool FImageCollection::RemoveImage(FComponentInstanceIDType ParentEntityId, FEntityIDType ImageId)
+{
+	TSharedPtr<FImage>* Found = Images.Find(ImageId);
+	if (!Found)
+	{
+		return false;
+	}
+	const TSharedPtr<FImage>& Image = *Found;
+
+	FDefinition* ParentDefinition =  Context.GetDefinition(ParentEntityId);
+
+	// Remove ComponentInstance for good only if incoming ParentDefinition is current Instance's parent. 
+	//
+	// Details:
+	// ComponentInstance which removal is notified could have been relocated to another Definition 
+	// This happens when Make Group is done - first new Group is added, containing existing ComponentInstance
+	// And only after that event about removal from previous owning Definition is received
+	if (Image->IsParentDefinition(ParentDefinition))
+	{
+		Image->RemoveImage(Context);
+		Images.Remove(ImageId);
+	}
+
+	return true;
+
+}
+
 void FSceneCollection::PopulateFromModel(SUModelRef InModelRef)
 {
 	// Get the number of scenes in the SketchUp model.
@@ -402,6 +461,26 @@ TSharedPtr<FComponentDefinition>* FComponentDefinitionCollection::FindComponentD
 	return ComponentDefinitionMap.Find(ComponentDefinitionID);
 }
 
+TSharedPtr<FImage> FImageCollection::AddImage(FDefinition& ParentDefinition, const SUImageRef& ImageRef)
+{
+	FEntityIDType ImageId = DatasmithSketchUpUtils::GetEntityID(SUImageToEntity(ImageRef));
+
+	TSharedPtr<FImage> Image;
+	if (TSharedPtr<FImage>* Found = Images.Find(ImageId))
+	{
+		Image = *Found; 
+	}
+	else
+	{
+		Image = MakeShared<FImage>(ImageRef);
+		Images.Add(ImageId, Image);
+	}
+
+	Image->SetParentDefinition(Context, &ParentDefinition);
+
+	return Image;
+}
+
 void FEntitiesObjectCollection::RegisterEntities(DatasmithSketchUp::FEntities& Entities)
 {
 	for (int32 FaceId : Entities.EntitiesGeometry->FaceIds)
@@ -436,11 +515,13 @@ void FEntitiesObjectCollection::UnregisterEntities(DatasmithSketchUp::FEntities&
 	}
 }
 
+
 TSharedPtr<FEntities> FEntitiesObjectCollection::AddEntities(FDefinition& InDefinition, SUEntitiesRef EntitiesRef)
 {
 	TSharedPtr<FEntities> Entities = MakeShared<FEntities>(InDefinition);
 
 	Entities->EntitiesRef = EntitiesRef;
+
 	return Entities;
 }
 
diff --git a/Engine/Source/Programs/Enterprise/Datasmith/DatasmithSketchUpRubyExporter/Private/DatasmithSketchUpExportContext.h b/Engine/Source/Programs/Enterprise/Datasmith/DatasmithSketchUpRubyExporter/Private/DatasmithSketchUpExportContext.h
--- a/Engine/Source/Programs/Enterprise/Datasmith/DatasmithSketchUpRubyExporter/Private/DatasmithSketchUpExportContext.h
+++ b/Engine/Source/Programs/Enterprise/Datasmith/DatasmithSketchUpRubyExporter/Private/DatasmithSketchUpExportContext.h
@@ -50,6 +50,9 @@ namespace DatasmithSketchUp
 	class FModelDefinition;
 	class FTexture;
 	class FTextureImageFile;
+	class FImage;
+	class FImageFile;
+	class FImageMaterial;
 
 	class FComponentInstanceCollection
 	{
@@ -99,6 +102,28 @@ namespace DatasmithSketchUp
 		FExportContext& Context;
 	};
 
+	class FImageCollection
+	{
+	public:
+		FImageCollection(FExportContext& InContext) : Context(InContext) {}
+
+		TSharedPtr<FImage> AddImage(FDefinition& ParentDefinition, const SUImageRef& ImageRef);
+		void Update();
+		bool InvalidateImage(const FEntityIDType& EntityID);
+		bool RemoveImage(FComponentInstanceIDType ParentEntityId, FEntityIDType ImageId);
+		void UpdateProperties();
+		void ReleaseMaterial(FImage& Image);
+		const TCHAR* AcquireMaterial(FImage& Image);
+		void LayerModified(FEntityIDType LayerId);
+
+	private:
+		FExportContext& Context;
+
+		TMap<FEntityIDType, TSharedPtr<FImage>> Images;
+
+		TMap<FMD5Hash, TSharedPtr<FImageMaterial>> ImageMaterials;  // One material for same images(even for different entities - their appearance onlly depends in the image content)
+	};
+
 	class FTextureCollection
 	{
 	public:
@@ -131,6 +156,23 @@ namespace DatasmithSketchUp
 		TMap<FMD5Hash,  TSharedPtr<FTextureImageFile>> Images; // set of images using the same name in SketchUp materials
 	};
 
+	// Holds all image files that are exported
+	class FImageFileCollection
+	{
+	public:
+		FImageFileCollection(FExportContext& Context): Context(Context) {}
+
+		TSharedPtr<FImageFile> AddImage(SUImageRepRef ImageRep, FString FileName);
+
+		const TCHAR* GetImageFilePath(FImageFile&);
+		FMD5Hash GetImageFileHash(FImageFile&);
+
+	private:
+		FExportContext& Context;
+		TMap<FMD5Hash, TSharedPtr<FImageFile>> ImageFiles;
+
+	};
+
 	class FEntitiesObjectCollection
 	{
 	public:
@@ -318,6 +360,8 @@ namespace DatasmithSketchUp
 		FSceneCollection Scenes;
 		FTextureCollection Textures;
 		FLayerCollection Layers;
+		FImageCollection Images;
+		FImageFileCollection ImageFiles;
 
 		bool bColorByLayer = false;
 		bool bColorByLayerInvaliated = true;
diff --git a/Engine/Source/Programs/Enterprise/Datasmith/DatasmithSketchUpRubyExporter/Private/DatasmithSketchUpExporter.cpp b/Engine/Source/Programs/Enterprise/Datasmith/DatasmithSketchUpRubyExporter/Private/DatasmithSketchUpExporter.cpp
--- a/Engine/Source/Programs/Enterprise/Datasmith/DatasmithSketchUpRubyExporter/Private/DatasmithSketchUpExporter.cpp
+++ b/Engine/Source/Programs/Enterprise/Datasmith/DatasmithSketchUpRubyExporter/Private/DatasmithSketchUpExporter.cpp
@@ -372,6 +372,11 @@ public:
 			return true;
 		}
 
+		if (Context.Images.RemoveImage(ParentEntityId, EntityId))
+		{
+			return true;
+		}
+
 		// Try Material
 		// Doesn't seem like material removal event ever comes through properly
 		// MaterialsObserver::onMaterialRemove simply has deleted material entity(meaning there's not way to retrieve correct reference/id for it) 
@@ -406,6 +411,11 @@ public:
 			return true;
 		}
 
+		if (Context.Images.InvalidateImage(EntityId))
+		{
+			return true;
+		}
+
 		return false;
 	}
 
@@ -483,6 +493,16 @@ public:
 
 			break;
 		}
+		case SURefType_Image:
+		{
+			DatasmithSketchUp::FDefinition* DefinitionPtr = Context.GetDefinition(EntityParent);
+			if (ensure(DefinitionPtr)) // Parent definition expected to already exist when new entity being added
+			{
+				DefinitionPtr->AddImage(Context, Context.Images.AddImage(*DefinitionPtr, SUImageFromEntity(Entity)));
+			}
+
+			break;
+		}
 		case SURefType_Face:
 		{
 			Context.GetDefinition(EntityParent)->InvalidateDefinitionGeometry();
@@ -513,6 +533,7 @@ public:
 		Context.EntitiesObjects.LayerModified(LayerId);
 		Context.Layers.UpdateLayer(SULayerFromEntity(Entity));
 		Context.Materials.LayerMaterials.UpdateLayer(SULayerFromEntity(Entity));
+		Context.Images.LayerModified(LayerId);
 
 		return true;
 	}
@@ -739,8 +760,6 @@ VALUE DatasmithSketchUpDirectLinkExporter_on_entity_modified_by_id(VALUE self, V
 	int32 EntityId = FIX2LONG(ruby_entity_id);
 	// Done converting args
 
-	;
-
 	return Ptr->OnEntityModified(DatasmithSketchUp::FEntityIDType(EntityId)) ? Qtrue : Qfalse;
 }
 
diff --git a/Engine/Source/Programs/Enterprise/Datasmith/DatasmithSketchUpRubyExporter/Private/DatasmithSketchUpImage.cpp b/Engine/Source/Programs/Enterprise/Datasmith/DatasmithSketchUpRubyExporter/Private/DatasmithSketchUpImage.cpp
new file mode 100644
--- /dev/null
+++ b/Engine/Source/Programs/Enterprise/Datasmith/DatasmithSketchUpRubyExporter/Private/DatasmithSketchUpImage.cpp
@@ -0,0 +1,442 @@
+// Copyright Epic Games, Inc. All Rights Reserved.
+
+#include "DatasmithSketchUpComponent.h"
+
+#include "DatasmithSketchUpExportContext.h"
+#include "DatasmithSketchUpMaterial.h"
+#include "DatasmithSketchUpMetadata.h"
+#include "DatasmithSketchUpString.h"
+#include "DatasmithSketchUpSummary.h"
+#include "DatasmithSketchUpUtils.h"
+
+// SketchUp SDK.
+#include "DatasmithSketchUpSDKBegins.h"
+#include "SketchUpAPI/model/drawing_element.h"
+#include <SketchUpAPI/model/entity.h>
+#include <SketchUpAPI/model/image.h>
+#include <SketchUpAPI/model/image_rep.h>
+#include "SketchUpAPI/model/layer.h"
+#include "SketchUpAPI/geometry/transformation.h"
+#include "SketchUpAPI/geometry/vector3d.h"
+
+#include "DatasmithSketchUpSDKCeases.h"
+
+#include "IDatasmithSceneElements.h"
+#include "DatasmithSceneFactory.h"
+#include "DatasmithUtils.h"
+
+#include "Misc/Paths.h"
+
+
+using namespace DatasmithSketchUp;
+
+FImage::FImage(SUImageRef InEntityRef)
+	: Super(SUImageToEntity(InEntityRef))
+{
+	FString IdStr = FString::Printf(TEXT("%llx"), GetPersistentId());
+	MeshElementName = FString::Printf(TEXT("M%s"), *IdStr);
+}
+
+void FImage::RemoveImage(FExportContext& Context)
+{
+	RemoveOccurrences(Context);
+	RemoveImageFromDatasmithScene(Context);
+}
+
+int64 FImage::GetPersistentId()
+{
+	int64 PersistentId = 0;
+	SUEntityGetPersistentID(EntityRef, &PersistentId);
+
+	return PersistentId;
+}
+
+FString FImage::GetFileName()
+{
+	// Get file name to use in name(label) computation. Not used to export file image!
+
+	// Note: 
+	// Not using SUImageGetName - returns empty string:
+	// [REDACTED_PROFILE_URL]/api-issue-tracker/issues/303
+
+	// Also, newer SU has access to Definition of an Image
+	// 	SUComponentDefinitionRef Definition = SU_INVALID;
+	// 	SUResult ImageGetDefinitionResult = SUImageGetDefinition(SUImageFromEntity(EntityRef), &Definition);
+	// 	ensure(SU_ERROR_NONE == ImageGetDefinitionResult);
+
+	// So for SU 2019+ file name is the only way to get Image's 'name' 
+
+	FString FileName = SuGetString(SUImageGetFileName, SUImageFromEntity(EntityRef));
+	FString Name = FPaths::GetCleanFilename(FileName);
+
+	if (!Name.IsEmpty())
+	{
+		return Name;
+	}
+
+	return FString::Printf(TEXT("image_%llx.png"), GetPersistentId());
+}
+
+void FImage::InvalidateImage()
+{
+	InvalidateEntityProperties();
+	InvalidateEntityGeometry();
+}
+
+
+FString FImage::GetName()
+{
+	return FDatasmithUtils::SanitizeObjectName(GetFileName());
+}
+
+void FImage::ApplyOverrideMaterialToNode(FNodeOccurence& Node, FMaterialOccurrence& Material)
+{
+	// Image entity is not affected by override materials on components
+}
+
+void FImage::UpdateOccurrence(FExportContext& Context, FNodeOccurence& Node)
+{
+	if (!Node.bVisible)
+	{
+		return;
+	}
+	
+	Node.EffectiveLayerRef = DatasmithSketchUpUtils::GetEffectiveLayer( SUImageToDrawingElement(SUImageFromEntity(EntityRef)), Node.ParentNode->EffectiveLayerRef);
+
+	BuildNodeNames(Node);
+
+	FString EffectiveLayerName = FDatasmithUtils::SanitizeObjectName(SuGetString(SULayerGetName, Node.EffectiveLayerRef));
+	Node.DatasmithActorElement->SetLayer(*EffectiveLayerName);
+	Node.DatasmithActorElement->SetLabel(*Node.DatasmithActorLabel);
+
+	// Compute the world transform of the SketchUp component instance.
+	SUTransformation LocalTransform;
+
+	// note: For Image, transformation handling is a bit than for Component(see further)
+	SUImageGetTransform(SUImageFromEntity(EntityRef), &LocalTransform);
+
+	// Remove scaling from local Image's transform
+	// SU embeds image dimensions into its transform's scaling and we don't need it since we construct geometry using Image entity's width/height
+	SUVector3D XAxis{1, 0, 0};
+	SUTransformationGetXAxis(&LocalTransform, &XAxis);
+	SUVector3D YAxis{0, 1, 0};
+	SUTransformationGetYAxis(&LocalTransform, &YAxis);
+	SUVector3D ZAxis{0, 0, 1};
+	SUTransformationGetZAxis(&LocalTransform, &ZAxis);
+
+	SUVector3DNormalize(&XAxis);
+	SUVector3DNormalize(&YAxis);
+	SUVector3DNormalize(&ZAxis);
+
+	SUPoint3D Origin{0, 0, 0};
+	SUTransformationGetOrigin(&LocalTransform, &Origin);
+
+	// SUTransformationSetFromPointAndAxes point parameter is not 'Origin' of transformation(or in other words translation of transform from local to world)
+	// but rather inverse of translation
+	SUPoint3D Point = {-Origin.x, -Origin.y, -Origin.z}; 
+	SUTransformationSetFromPointAndAxes(&LocalTransform, &Point, &XAxis, &YAxis,  &ZAxis);
+
+	SUTransformation WorldTransform;
+	SUTransformationMultiply(&Node.ParentNode->WorldTransform, &LocalTransform, &WorldTransform);
+	Node.WorldTransform = WorldTransform; // Store world transform to be used by children to compute its
+
+	// Set the Datasmith actor world transform.
+	DatasmithSketchUpUtils::SetActorTransform(Node.DatasmithActorElement, Node.WorldTransform);
+
+	Node.ResetMetadataElement(Context);
+
+	// Update Datasmith Mesh Actors
+	for (int32 MeshIndex = 0; MeshIndex < Node.MeshActors.Num(); ++MeshIndex)
+	{
+		const TSharedPtr<IDatasmithMeshActorElement>& MeshActor = Node.MeshActors[MeshIndex];
+
+		// Set mesh actor transform after node transform
+		MeshActor->SetScale(Node.DatasmithActorElement->GetScale());
+		MeshActor->SetRotation(Node.DatasmithActorElement->GetRotation());
+		MeshActor->SetTranslation(Node.DatasmithActorElement->GetTranslation());
+	}
+
+	FString MeshActorLabel = Node.GetActorLabel();
+
+	// Update Datasmith Mesh Actors
+	for (int32 MeshIndex = 0; MeshIndex < Node.MeshActors.Num(); ++MeshIndex)
+	{
+		const TSharedPtr<IDatasmithMeshActorElement>& MeshActor = Node.MeshActors[MeshIndex];
+		MeshActor->SetLabel(*MeshActorLabel);
+		MeshActor->SetLayer(*EffectiveLayerName);
+	}
+}
+
+void FImage::RemoveImageFromDatasmithScene(FExportContext& Context)
+{
+	Context.Images.ReleaseMaterial(*this);
+
+	if (DatasmithMeshElement)
+	{
+		Context.DatasmithScene->RemoveMesh(DatasmithMeshElement);
+		DatasmithMeshElement.Reset();
+	}
+}
+
+void FImage::EntityOccurrenceVisible(FNodeOccurence* Node, bool bUses)
+{
+	bool bNodesInvisibleBefore = VisibleNodes.IsEmpty();
+	FEntity::EntityOccurrenceVisible(Node, bUses);
+	if (bNodesInvisibleBefore != VisibleNodes.IsEmpty())
+	{
+		// Just invalidate geometry to create/remove geometry depending on existing occurrences
+		InvalidateEntityGeometry(); 
+	}
+
+}
+
+const TCHAR* FImage::GetMeshElementName()
+{
+	return *MeshElementName;
+}
+
+void FImage::Update(FExportContext& Context)
+{
+	if (bGeometryInvalidated)
+	{
+		RemoveImageFromDatasmithScene(Context);
+
+		if (!VisibleNodes.IsEmpty())
+		{
+			UpdateEntityGeometry(Context);
+			UpdateGeometry(Context);
+
+			DatasmithMeshElement->SetMaterial(Context.Images.AcquireMaterial(*this), 0);
+		}
+	}
+}
+
+void FImage::BuildNodeNames(FNodeOccurence& Node)
+{
+	int64 SketchupPersistentID = Node.Entity.GetPersistentId();
+	Node.DatasmithActorName = FString::Printf(TEXT("%ls_%lld"), *Node.ParentNode->GetActorName(), SketchupPersistentID);
+
+	FString EntityName = Node.Entity.GetName();
+	Node.DatasmithActorLabel = FDatasmithUtils::SanitizeObjectName(EntityName.IsEmpty() ? GetName() : EntityName);
+}
+
+void FImage::CreateActor(FExportContext& Context, FNodeOccurence& Node)
+{
+	BuildNodeNames(Node);
+
+	// Create a Datasmith actor for the component instance.
+	Node.DatasmithActorElement = FDatasmithSceneFactory::CreateActor(*Node.DatasmithActorName); // a EDatasmithElementType::Actor
+
+	{
+		// Add the Datasmith actor component depth tag.
+		// We use component depth + 1 to factor in the added Datasmith scene root once imported in Unreal.
+		FString ComponentDepthTag = FString::Printf(TEXT("SU.DEPTH.%d"), Node.Depth);
+		Node.DatasmithActorElement->AddTag(*ComponentDepthTag);
+
+		// Add the Datasmith actor component instance path tag.
+		FString InstancePathTag = Node.DatasmithActorName.Replace(TEXT("SU"), TEXT("SU.PATH.0")).Replace(TEXT("_"), TEXT("."));
+		Node.DatasmithActorElement->AddTag(*InstancePathTag);
+	}
+
+	if (Node.ParentNode->DatasmithActorElement)
+	{
+		Node.ParentNode->DatasmithActorElement->AddChild(Node.DatasmithActorElement);
+	}
+	else
+	{
+		Context.DatasmithScene->AddActor(Node.DatasmithActorElement);
+	}
+}
+
+
+void FImage::UpdateOccurrenceVisibility(FExportContext& Context, FNodeOccurence& Node)
+{
+	// Parent node, component instance and layer - all should be visible to have node visible
+	Node.SetVisibility(Node.ParentNode->bVisible && !bHidden && bLayerVisible);
+
+	EntityOccurrenceVisible(&Node, Node.bVisible);
+
+	if (Node.bVisible)
+	{
+		if (!Node.DatasmithActorElement)
+		{
+			CreateActor(Context, Node);
+
+			// Invalidate actor-dependent 
+			Node.InvalidateProperties();
+		}
+		Node.InvalidateMeshActors();
+	}
+	else
+	{
+		Node.RemoveDatasmithActorHierarchy(Context);
+	}
+
+	for (FNodeOccurence* ChildNode : Node.Children)
+	{
+		ChildNode->bVisibilityInvalidated = true;
+	}
+}
+
+bool FImage::UpdateOccurrenceMeshActors(FExportContext& Context, FNodeOccurence& Node)
+{
+	Node.MeshActors.Reset(1);
+
+	FString ComponentActorName = Node.GetActorName();
+	
+	{
+		FString MeshActorName = FString::Printf(TEXT("%ls_%d"), *ComponentActorName); // Count meshes/mesh actors from 1
+
+		// Create a Datasmith mesh actor for the Datasmith mesh element.
+		TSharedPtr<IDatasmithMeshActorElement> DMeshActorPtr = FDatasmithSceneFactory::CreateMeshActor(*MeshActorName);
+
+		Node.MeshActors.Add(DMeshActorPtr);
+
+		// Add the Datasmith actor component depth tag.
+		// We use component depth + 1 to factor in the added Datasmith scene root once imported in Unreal.
+		FString ComponentDepthTag = FString::Printf(TEXT("SU.DEPTH.%d"), Node.Depth + 1);
+		DMeshActorPtr->AddTag(*ComponentDepthTag);
+
+		// Add the Datasmith actor component instance path tag.
+		FString InstancePathTag = ComponentActorName.Replace(TEXT("SU"), TEXT("SU.PATH.0")).Replace(TEXT("_"), TEXT("."));
+		DMeshActorPtr->AddTag(*InstancePathTag);
+
+		// Add the mesh actor to our component Datasmith actor hierarchy.
+		if (Node.DatasmithActorElement.IsValid()) 
+		{
+			Node.DatasmithActorElement->AddChild(DMeshActorPtr);
+		}
+		else
+		{
+			Context.DatasmithScene->AddActor(DMeshActorPtr);
+		}
+
+		// ADD_TRACE_LINE(TEXT("Actor %ls: %ls %ls %ls"), *MeshActorLabel, *ComponentDepthTag, *DefinitionGUIDTag, *InstancePathTag);
+
+		// Set the Datasmith mesh element used by the mesh actor.
+		DMeshActorPtr->SetStaticMeshPathName(GetMeshElementName());
+	}
+	return false;
+}
+
+void FImage::InvalidateOccurrencesGeometry(FExportContext& Context)
+{
+	for (FNodeOccurence* Node : Occurrences)
+	{
+		Node->InvalidateMeshActors();
+		Node->InvalidateProperties();
+	}
+}
+
+void FImage::InvalidateOccurrencesProperties(FExportContext& Context)
+{
+	// When ComponentInstance is modified we need to determine if its visibility might have changed foremost
+	// because this determines whether corresponding node would exist in the Datasmith scene 
+	// Two things affect this - Hidden instance flag and layer(tag):
+
+	bool bNewHidden = false;
+	SUDrawingElementRef DrawingElementRef = SUImageToDrawingElement(SUImageFromEntity(EntityRef));
+	SUDrawingElementGetHidden(DrawingElementRef, &bNewHidden);
+
+	SUDrawingElementGetLayer(DrawingElementRef, &LayerRef);
+	bool bNewLayerVisible = Context.Layers.IsLayerVisible(LayerRef);
+
+	if (bHidden != bNewHidden || bLayerVisible != bNewLayerVisible)
+	{
+		bHidden = bNewHidden;
+		bLayerVisible = bNewLayerVisible;
+		for (FNodeOccurence* Node : Occurrences)
+		{
+			Node->bVisibilityInvalidated = true;
+		}
+	}
+
+	for (FNodeOccurence* Node : Occurrences)
+	{
+		Node->InvalidateProperties();
+	}
+}
+
+void FImage::UpdateMetadata(FExportContext& Context)
+{
+	ParsedMetadata = MakeUnique<FMetadata>(EntityRef);;
+}
+
+void FImage::UpdateEntityProperties(FExportContext& Context)
+{
+	if (bPropertiesInvalidated)
+	{
+		InvalidateEntityGeometry();
+	}
+
+	FEntity::UpdateEntityProperties(Context);
+}
+
+namespace DatasmithSketchUp
+{
+	class FImageFile
+	{
+	public:
+		struct FOptions
+		{
+			bool bColorized = false;
+		};
+
+		SUImageRepRef ImageRep;
+		FOptions Options;
+
+		FString FilePath;
+
+		FImageFile(SUImageRepRef InImageRep)
+			: ImageRep(InImageRep)
+		{
+			size_t Width, Height;
+			SUImageRepGetPixelDimensions(ImageRep, &Width, &Height);
+
+			size_t DataSize, Bpp;
+			SUImageRepGetDataSize(ImageRep, &DataSize, &Bpp);
+
+			TArray<SUByte> Data;
+			Data.SetNumUninitialized(DataSize);
+			SUImageRepGetData(ImageRep, DataSize, Data.GetData());
+
+			FMD5 MD5;
+			MD5.Update(reinterpret_cast<const uint8*>(&Width), sizeof(Width));
+			MD5.Update(reinterpret_cast<const uint8*>(&Height), sizeof(Height));
+			MD5.Update(reinterpret_cast<const uint8*>(&Bpp), sizeof(Bpp));
+			MD5.Update(reinterpret_cast<const uint8*>(&Options.bColorized), sizeof(Options.bColorized)); // Colorized flag affects resulting texture image
+			MD5.Update(Data.GetData(), Data.Num());
+
+			ImageHash.Set(MD5);
+		}
+
+		FMD5Hash ImageHash;
+
+	};
+
+}
+
+TSharedPtr<FImageFile> FImageFileCollection::AddImage(SUImageRepRef ImageRep, FString FileName)
+{
+	TSharedPtr<FImageFile> ImageFile = MakeShared<FImageFile>(ImageRep);
+	if (TSharedPtr<FImageFile>* Found = ImageFiles.Find(ImageFile->ImageHash))
+	{
+		return *Found;
+	}
+
+	ImageFile->FilePath = FPaths::Combine(Context.GetAssetsOutputPath(), TCHAR_TO_UTF8(*FileName));
+	SUResult SaveToFileResult = SUImageRepSaveToFile(ImageRep, TCHAR_TO_UTF8(*ImageFile->FilePath));
+	ensure(SaveToFileResult == SU_ERROR_NONE);
+
+	return ImageFiles.Add(ImageFile->ImageHash, ImageFile);
+}
+
+const TCHAR* FImageFileCollection::GetImageFilePath(FImageFile& ImageFile)
+{
+	return *ImageFile.FilePath;
+}
+
+FMD5Hash FImageFileCollection::GetImageFileHash(FImageFile& ImageFile)
+{
+	return ImageFile.ImageHash;
+}
diff --git a/Engine/Source/Programs/Enterprise/Datasmith/DatasmithSketchUpRubyExporter/Private/DatasmithSketchUpMaterial.cpp b/Engine/Source/Programs/Enterprise/Datasmith/DatasmithSketchUpRubyExporter/Private/DatasmithSketchUpMaterial.cpp
--- a/Engine/Source/Programs/Enterprise/Datasmith/DatasmithSketchUpRubyExporter/Private/DatasmithSketchUpMaterial.cpp
+++ b/Engine/Source/Programs/Enterprise/Datasmith/DatasmithSketchUpRubyExporter/Private/DatasmithSketchUpMaterial.cpp
@@ -15,6 +15,8 @@
 // SketchUp SDK.
 #include "DatasmithSketchUpSDKBegins.h"
 #include "SketchUpAPI/model/component_instance.h"
+#include <SketchUpAPI/model/image.h>
+#include <SketchUpAPI/model/image_rep.h>
 #include <SketchUpAPI/model/layer.h>
 #include "SketchUpAPI/model/model.h"
 #include "SketchUpAPI/model/texture.h"
@@ -194,7 +196,6 @@ namespace DatasmithSketchUp
 				IDatasmithMaterialExpressionGeneric* ExpressionOpacity = DatasmithMaterialElementPtr->AddMaterialExpression<IDatasmithMaterialExpressionGeneric>();
 				ExpressionOpacity->SetExpressionName(TEXT("OneMinus"));
 
-
 				ExpressionTexture->ConnectExpression(*ExpressionOpacity->GetInput(0), 3);
 
 				ExpressionOpacity->ConnectExpression(DatasmithMaterialElementPtr->GetOpacity());
@@ -228,6 +229,36 @@ namespace DatasmithSketchUp
 		// Return the Datasmith material element.
 		return DatasmithMaterialElementPtr;
 	}
+
+	void FModel::ApplyOverrideMaterialToNode(FNodeOccurence& Node, FMaterialOccurrence& Material)
+	{
+		GetDefinition()->ApplyOverrideMaterialToNode(Node, Material);
+	}
+
+	void FComponentInstance::ApplyOverrideMaterialToNode(FNodeOccurence& Node, FMaterialOccurrence& Material)
+	{
+		GetDefinition()->ApplyOverrideMaterialToNode(Node, Material);
+	}
+
+	void FDefinition::ApplyOverrideMaterialToNode(FNodeOccurence& Node, FMaterialOccurrence& Material)
+	{
+		FEntitiesGeometry& EntitiesGeometry = *GetEntities().EntitiesGeometry;
+
+		for (int32 MeshIndex = 0; MeshIndex < Node.MeshActors.Num(); ++MeshIndex)
+		{
+			// Update Override(Inherited)  Material
+			// todo: set inherited material only on mesh actors that have faces with default material, right now setting on every mesh, hot harmful but excessive
+			if (EntitiesGeometry.IsMeshUsingInheritedMaterial(MeshIndex))
+			{
+				const TSharedPtr<IDatasmithMeshActorElement>& MeshActor = Node.MeshActors[MeshIndex];
+
+				// SketchUp has 'material override' only for single('Default') material. 
+				// So we reset overrides on the actor to remove this single override(if it was set) and re-add new override
+				MeshActor->ResetMaterialOverrides(); // Clear previous override if was set
+				MeshActor->AddMaterialOverride(Material.GetName(), EntitiesGeometry.GetInheritedMaterialOverrideSlotId());
+			}
+		}
+	}
 }
 
 FMaterial::FMaterial(SUMaterialRef InMaterialRef)
@@ -454,24 +485,7 @@ void FMaterialOccurrence::ApplyRegularMaterial(FMaterialIDType MaterialId)
 	// Apply material to mesh actors
 	for (FNodeOccurence* NodePtr : NodesMaterialInheritedBy)
 	{
-		FNodeOccurence& Node = *NodePtr;
-		FDefinition* EntityDefinition = Node.Entity.GetDefinition();
-		DatasmithSketchUp::FEntitiesGeometry& EntitiesGeometry = *EntityDefinition->GetEntities().EntitiesGeometry;
-
-		for (int32 MeshIndex = 0; MeshIndex < Node.MeshActors.Num(); ++MeshIndex)
-		{
-			// Update Override(Inherited)  Material
-			// todo: set inherited material only on mesh actors that have faces with default material, right now setting on every mesh, hot harmful but excessive
-			if (EntitiesGeometry.IsMeshUsingInheritedMaterial(MeshIndex))
-			{
-				const TSharedPtr<IDatasmithMeshActorElement>& MeshActor = Node.MeshActors[MeshIndex];
-
-				// SketchUp has 'material override' only for single('Default') material. 
-				// So we reset overrides on the actor to remove this single override(if it was set) and re-add new override
-				MeshActor->ResetMaterialOverrides(); // Clear previous override if was set
-				MeshActor->AddMaterialOverride(GetName(), EntitiesGeometry.GetInheritedMaterialOverrideSlotId());
-			}
-		}
+		NodePtr->Entity.ApplyOverrideMaterialToNode(*NodePtr, *this);
 	}
 }
 
@@ -495,24 +509,7 @@ void FMaterialOccurrence::ApplyLayerMaterial(FLayerIDType LayerId)
 	// Apply material to mesh actors
 	for (FNodeOccurence* NodePtr : NodesMaterialInheritedBy)
 	{
-		FNodeOccurence& Node = *NodePtr;
-		FDefinition* EntityDefinition = Node.Entity.GetDefinition();
-		DatasmithSketchUp::FEntitiesGeometry& EntitiesGeometry = *EntityDefinition->GetEntities().EntitiesGeometry;
-
-		for (int32 MeshIndex = 0; MeshIndex < Node.MeshActors.Num(); ++MeshIndex)
-		{
-			// Update Override(Inherited)  Material
-			// todo: set inherited material only on mesh actors that have faces with default material, right now setting on every mesh, hot harmful but excessive
-			if (EntitiesGeometry.IsMeshUsingInheritedMaterial(MeshIndex))
-			{
-				const TSharedPtr<IDatasmithMeshActorElement>& MeshActor = Node.MeshActors[MeshIndex];
-
-				// SketchUp has 'material override' only for single('Default') material. 
-				// So we reset overrides on the actor to remove this single override(if it was set) and re-add new override
-				MeshActor->ResetMaterialOverrides(); // Clear previous override if was set
-				MeshActor->AddMaterialOverride(GetName(), EntitiesGeometry.GetInheritedMaterialOverrideSlotId());
-			}
-		}
+		NodePtr->Entity.ApplyOverrideMaterialToNode(*NodePtr, *this);
 	}
 }
 
@@ -744,3 +741,192 @@ void FMaterialCollection::Update()
 
 	RegularMaterials.UpdateDefaultMaterial();
 }
+
+namespace DatasmithSketchUp
+{
+
+class FImageMaterial
+{
+public:
+
+	struct FOptions
+	{
+		bool bScaleTexture = false;
+		FVector2D TextureScale;
+		bool bUseAlphaChannel = false;
+		bool bTranslucent = false;
+	};
+
+	FImageMaterial(FExportContext& Context, FImage& Image, FImageFile& ImageFile)
+	{
+		FString NameBase = Image.GetName();
+
+		FString MaterialName = FString::Printf(TEXT("image_material_%s"), *NameBase);  
+		FString MaterialLabel = NameBase;
+
+		FString TextureName = FString::Printf(TEXT("image_texture_%s"), *NameBase);  
+
+		TextureElement = FDatasmithSceneFactory::CreateTexture(*TextureName);
+		TextureElement->SetSRGB(EDatasmithColorSpace::sRGB);
+		TextureElement->SetFile(Context.ImageFiles.GetImageFilePath(ImageFile));
+		Context.DatasmithScene->AddTexture(TextureElement);
+
+		FImageMaterial::FOptions MaterialOptions;
+		CreateMaterialElement(Context, *MaterialName, *MaterialLabel, TextureElement->GetName(), MaterialOptions);
+
+		Hash = Context.ImageFiles.GetImageFileHash(ImageFile);
+	}
+
+	void Release(FExportContext& Context)
+	{
+		if (DatasmithMaterialElement)
+		{
+			Context.DatasmithScene->RemoveMaterial(DatasmithMaterialElement);
+			DatasmithMaterialElement.Reset();
+		}
+
+		if (TextureElement)
+		{
+			Context.DatasmithScene->RemoveTexture(TextureElement);
+			TextureElement.Reset();
+		}
+	}
+
+	const TCHAR* GetMaterialElementName()
+	{
+		return DatasmithMaterialElement->GetName();
+	}
+
+	TSharedPtr<IDatasmithBaseMaterialElement> CreateMaterialElement(
+		FExportContext& Context,
+		const TCHAR* InMaterialName,
+		const TCHAR* InMaterialLabel,
+		const TCHAR* InDatasmithTextureElementName,
+		const FOptions& Options
+	)
+	{
+		// Create a Datasmith material element for the material definition.
+		DatasmithMaterialElement = FDatasmithSceneFactory::CreateUEPbrMaterial(InMaterialName);
+
+		// Set the material element label used in the Unreal UI.
+		FString MaterialLabel = FDatasmithUtils::SanitizeObjectName(InMaterialLabel);
+		DatasmithMaterialElement->SetLabel(*MaterialLabel);
+
+		DatasmithMaterialElement->SetTwoSided(true); // imported Image in SketchUp is displayed as two-sided with no way to controls it, so 
+
+		// todo: inspect texture to see if there's alpha. Also consider masked opacity(might be better fit for SketchUp images?)
+		bool bTranslucent = false; 
+
+		{
+			IDatasmithMaterialExpressionTexture* ExpressionTexture = DatasmithMaterialElement->AddMaterialExpression< IDatasmithMaterialExpressionTexture >();
+			ExpressionTexture->SetName(TEXT("Texture"));
+			ExpressionTexture->SetTexturePathName(InDatasmithTextureElementName);
+
+			// Apply texture scaling
+			if (Options.bScaleTexture && !Options.TextureScale.Equals(FVector2D::UnitVector))
+			{
+				IDatasmithMaterialExpressionFunctionCall* UVEditExpression = DatasmithMaterialElement->AddMaterialExpression< IDatasmithMaterialExpressionFunctionCall >();
+				UVEditExpression->SetFunctionPathName(TEXT("/DatasmithContent/Materials/UVEdit.UVEdit"));
+
+				UVEditExpression->ConnectExpression(ExpressionTexture->GetInputCoordinate());
+
+				// Tiling
+				IDatasmithMaterialExpressionColor* TilingValue = DatasmithMaterialElement->AddMaterialExpression< IDatasmithMaterialExpressionColor >();
+				TilingValue->SetName(TEXT("UV Tiling"));
+				TilingValue->GetColor() = FLinearColor(Options.TextureScale.X, Options.TextureScale.Y, 0.f);
+
+				TilingValue->ConnectExpression(*UVEditExpression->GetInput(2));
+
+				IDatasmithMaterialExpressionTextureCoordinate* TextureCoordinateExpression = DatasmithMaterialElement->AddMaterialExpression< IDatasmithMaterialExpressionTextureCoordinate >();
+				TextureCoordinateExpression->SetCoordinateIndex(0);
+				TextureCoordinateExpression->ConnectExpression(*UVEditExpression->GetInput(0));
+			}
+
+			ExpressionTexture->ConnectExpression(DatasmithMaterialElement->GetBaseColor());
+
+			bTranslucent = Options.bTranslucent || Options.bUseAlphaChannel;
+
+			// Set the Datasmith material element opacity.
+			if (Options.bUseAlphaChannel)
+			{
+				// Invert texture translarency to get Unreal opacity
+				IDatasmithMaterialExpressionGeneric* ExpressionOpacity = DatasmithMaterialElement->AddMaterialExpression<IDatasmithMaterialExpressionGeneric>();
+				ExpressionOpacity->SetExpressionName(TEXT("OneMinus"));
+
+
+				ExpressionTexture->ConnectExpression(*ExpressionOpacity->GetInput(0), 3);
+
+				ExpressionOpacity->ConnectExpression(DatasmithMaterialElement->GetOpacity());
+			}
+		}
+
+		if (bTranslucent)
+		{
+			DatasmithMaterialElement->SetBlendMode(/*EBlendMode::BLEND_Translucent*/2);
+		}
+
+		Context.DatasmithScene->AddMaterial(DatasmithMaterialElement);
+
+		// Return the Datasmith material element.
+		return DatasmithMaterialElement;
+	}
+
+	TSet<FImage*> ImagesForMaterial;
+	FMD5Hash Hash;
+
+private:
+	TSharedPtr<IDatasmithUEPbrMaterialElement> DatasmithMaterialElement;
+	TSharedPtr<IDatasmithTextureElement> TextureElement;
+};
+
+}
+
+void FImageCollection::ReleaseMaterial(FImage& Image)
+{
+	if (Image.ImageMaterial)
+	{
+		Image.ImageMaterial->ImagesForMaterial.Remove(&Image);
+
+		if (Image.ImageMaterial->ImagesForMaterial.IsEmpty())
+		{
+			Image.ImageMaterial->Release(Context); //  Clean up material from the datasmith scene before deallocating
+			ImageMaterials.Remove(Image.ImageMaterial->Hash);
+		}
+
+		Image.ImageMaterial = nullptr;
+	}
+}
+
+const TCHAR* FImageCollection::AcquireMaterial(FImage& Image)
+{
+	if (!ensure(!Image.ImageMaterial))
+	{
+		ReleaseMaterial(Image);
+	}
+
+	SUImageRepRef ImageRep = SU_INVALID;
+	ensure(SU_ERROR_NONE == SUImageRepCreate(&ImageRep));
+	SUResult ImageGetImageRepResult = SUImageGetImageRep(SUImageFromEntity(Image.EntityRef), &ImageRep);
+	ensure(ImageGetImageRepResult == SU_ERROR_NONE);
+
+	TSharedPtr<FImageFile> ImageFile = Context.ImageFiles.AddImage(ImageRep, Image.GetFileName());
+	SUResult ImageRepReleaseResult = SUImageRepRelease(&ImageRep);
+	ensure(SU_ERROR_NONE == ImageRepReleaseResult);
+
+	TSharedPtr<FImageMaterial> Material;
+
+	if (TSharedPtr<FImageMaterial>* Found = ImageMaterials.Find(Context.ImageFiles.GetImageFileHash(*ImageFile)))
+	{
+		Material = (*Found);
+	}
+	else
+	{
+		Material = MakeShared<FImageMaterial>(Context, Image, *ImageFile);
+		ImageMaterials.Add(Material->Hash, Material);
+	}
+
+	Material->ImagesForMaterial.Add(&Image);
+	Image.ImageMaterial = Material.Get();
+
+	return Material->GetMaterialElementName();
+}
diff --git a/Engine/Source/Programs/Enterprise/Datasmith/DatasmithSketchUpRubyExporter/Private/DatasmithSketchUpMesh.cpp b/Engine/Source/Programs/Enterprise/Datasmith/DatasmithSketchUpRubyExporter/Private/DatasmithSketchUpMesh.cpp
--- a/Engine/Source/Programs/Enterprise/Datasmith/DatasmithSketchUpRubyExporter/Private/DatasmithSketchUpMesh.cpp
+++ b/Engine/Source/Programs/Enterprise/Datasmith/DatasmithSketchUpRubyExporter/Private/DatasmithSketchUpMesh.cpp
@@ -20,13 +20,16 @@
 #include "SketchUpAPI/model/face.h"
 #include "SketchUpAPI/model/entities.h"
 #include "SketchUpAPI/model/entity.h"
+#include "SketchUpAPI/model/image.h"
 #include "SketchUpAPI/model/layer.h"
 #include "SketchUpAPI/model/mesh_helper.h"
 #include "SketchUpAPI/model/uv_helper.h"
 #include "SketchUpAPI/geometry/point3d.h"
+#include <SketchUpAPI/geometry/transformation.h>
 #include "DatasmithSketchUpSDKCeases.h"
 
 // Datasmith SDK.
+
 #include "DatasmithMesh.h"
 #include "DatasmithMeshExporter.h"
 #include "DatasmithSceneExporter.h"
@@ -385,6 +388,107 @@ namespace DatasmithSketchUp
 
 }
 
+namespace DatasmithSketchUp
+{
+	// Parses SU Image entity geometry
+	class FImageParser
+	{
+	public:
+		FImageParser(SUEntityRef InEntityRef): EntityRef(InEntityRef)
+		{
+		}
+
+		int32 ImageId = 0;
+		FLayerIDType LayerId;
+
+		void Parse(FExportContext& Context, FDatasmithSketchUpMesh& ExtractedMesh)
+		{
+			SUEntityGetID(EntityRef, &ImageId);
+
+			// Record every face's layer(even for invisible faces!). When face layer visibility changes 
+			// this geometry needs to be rebuilt
+			SULayerRef LayerRef = SU_INVALID;
+			SUImageRef ImageRef = SUImageFromEntity(EntityRef);
+			SUDrawingElementGetLayer(SUImageToDrawingElement(ImageRef), &LayerRef);
+			LayerId = DatasmithSketchUpUtils::GetEntityID(SULayerToEntity(LayerRef));
+
+			bool bHidden = false;
+			SUDrawingElementGetHidden(SUImageToDrawingElement(ImageRef), &bHidden);
+
+			if (!bHidden && Context.Layers.IsLayerVisible(LayerRef))
+			{
+				TArray<SUPoint3D>& MeshVertexPoints = ExtractedMesh.MeshVertexPoints;
+				TArray<SUVector3D>& MeshVertexNormals = ExtractedMesh.MeshVertexNormals;
+				TArray<SUUVQ>& MeshVertexUVQs = ExtractedMesh.MeshVertexUVQs;
+				TArray<SMeshTriangleIndices>& MeshTriangleIndices = ExtractedMesh.MeshTriangleIndices;
+				TArray<int32>& MeshTriangleSlotIds = ExtractedMesh.MeshTriangleSlotIds;
+
+				// Get the SketchUp triangle mesh vertex offset into the combined mesh vertex vector.
+				int32 MeshVertexIndexOffset = MeshVertexPoints.Num();
+
+				const int32 VertexPointCount = 4;
+				const int32 VertexNormalCount = 4;
+				const int32 TriangleCount = 2;
+				const int32 TriangleVertexIndexCount = 6;
+
+				MeshVertexNormals.Reserve(MeshVertexNormals.Num() + VertexNormalCount);
+				MeshVertexUVQs.Reserve(MeshVertexUVQs.Num() + VertexPointCount);
+				MeshTriangleIndices.Reserve(MeshTriangleIndices.Num() + TriangleVertexIndexCount);
+
+				SUVector3D Normal{0, 0, 1}; // todo: fix
+
+				double WidthInch = 0;
+				double HeightInch = 0;
+				if (SUImageGetDimensions(ImageRef, &WidthInch, &HeightInch) != SU_ERROR_NONE)
+				{
+					return;
+				}
+
+				for (size_t VertexPointIndex = 0; VertexPointIndex < VertexPointCount; VertexPointIndex++)
+				{
+
+					MeshVertexNormals.Add(Normal);
+
+					double X = (VertexPointIndex % 2);
+					double Y = (VertexPointIndex / 2);
+
+					// Put 1 in Q as we divide later by it
+					MeshVertexUVQs.Add(SUUVQ{X, Y, 1});
+					MeshVertexPoints.Add(SUPoint3D{X * WidthInch, Y * HeightInch, 0});
+				}
+
+				const size_t TriangleVertexIndices[TriangleVertexIndexCount] = 
+				{
+					0, 1, 2,
+					3, 2, 1
+				};
+
+				// Combine the mesh front-facing triangle vertex indices.
+				for (size_t Index = 0; Index < TriangleVertexIndexCount;)
+				{
+					size_t IndexA = MeshVertexIndexOffset + TriangleVertexIndices[Index++];
+					size_t IndexB = MeshVertexIndexOffset + TriangleVertexIndices[Index++];
+					size_t IndexC = MeshVertexIndexOffset + TriangleVertexIndices[Index++];
+
+					SMeshTriangleIndices TriangleIndices = { IndexA, IndexB, IndexC };
+
+					MeshTriangleIndices.Add(TriangleIndices);
+				}
+
+				MeshTriangleSlotIds.Reserve(MeshTriangleSlotIds.Num() + TriangleCount);
+					
+				for(int32 TriangleIndex = 0; TriangleIndex < TriangleCount; ++TriangleIndex)
+				{
+					MeshTriangleSlotIds.Add(0);
+				}
+			}
+		}
+
+		SUEntityRef EntityRef;
+	};
+}
+
+
 bool FEntitiesGeometry::IsMeshUsingInheritedMaterial(int32 MeshIndex)
 {
 	return Meshes[MeshIndex]->bIsUsingInheritedMaterial;
@@ -495,7 +599,8 @@ void FEntities::UpdateGeometry(FExportContext& Context)
 					FDatasmithMesh DatasmithMesh;
 					ExtractedMeshPtr->ConvertMeshToDatasmith(Context, Transform, DatasmithMesh);
 
-					FGCScopeGuard GCGuard; // Prevent GC from running while UDatasmithMesh is created in ExportToUObject. 
+					FGCScopeGuard GCGuard; // Prevent GC from running while UDatasmithMesh is created in ExportToUObject.
+
 					return DatasmithMeshExporter.ExportToUObject(Mesh->DatasmithMesh, Context.GetAssetsOutputPath(), DatasmithMesh, nullptr, FDatasmithExportOptions::LightmapUV);
 				}
 			));
@@ -559,6 +664,18 @@ TArray<SUComponentInstanceRef> FEntities::GetComponentInstances()
 	return MoveTemp(SComponentInstances);
 }
 
+TArray<SUImageRef> FEntities::GetImages()
+{
+	size_t ImageCount = 0;
+	SUEntitiesGetNumImages(EntitiesRef, &ImageCount);
+
+	TArray<SUImageRef> ImageRefs;
+	ImageRefs.Init(SU_INVALID, ImageCount);
+	SUEntitiesGetImages(EntitiesRef, ImageCount, ImageRefs.GetData(), &ImageCount);
+	ImageRefs.SetNum(ImageCount);
+
+	return MoveTemp(ImageRefs);
+}
 
 void ScanSketchUpEntitiesFaces(FExportContext& Context, SUEntitiesRef EntitiesRef, FEntitiesGeometry& Geometry, TFunctionRef<void(TSharedPtr<FDatasmithSketchUpMesh> ExtractedMesh)> OnNewExtractedMesh)
 {
@@ -732,3 +849,31 @@ void CombineSketchUpEntitiesFaces(FExportContext& Context, SUEntitiesRef Entitie
 
 	OnNewExtractedMesh(ExtractedMeshPtr);
 }
+
+
+void FImage::UpdateGeometry(FExportContext& Context)
+{
+	FImageParser ImageParser(EntityRef);
+	FDatasmithSketchUpMesh ExtractedMesh;
+	ImageParser.Parse(Context, ExtractedMesh);
+
+	FDatasmithMeshExporter DatasmithMeshExporter;
+	FDatasmithMesh DatasmithMesh;
+
+	SUTransformation Transform;
+	SUTransformationScale(&Transform, 1.0);
+
+	ExtractedMesh.ConvertMeshToDatasmith(Context, Transform, DatasmithMesh);
+
+	DatasmithMeshElement = FDatasmithSceneFactory::CreateMesh(TEXT(""));
+
+	DatasmithMeshElement->SetName(GetMeshElementName());
+	DatasmithMeshElement->SetLabel(*GetName());
+
+	FGCScopeGuard GCGuard; // Prevent GC from running while UDatasmithMesh is created in ExportToUObject. 
+	bool bResult = DatasmithMeshExporter.ExportToUObject(DatasmithMeshElement, Context.GetAssetsOutputPath(), DatasmithMesh, nullptr, FDatasmithExportOptions::LightmapUV);
+
+	// todo: something like AddMeshesToDatasmithScene
+	Context.DatasmithScene->AddMesh(DatasmithMeshElement);
+}
+
diff --git a/Engine/Source/Programs/Enterprise/Datasmith/DatasmithSketchUpRubyExporter/Private/DatasmithSketchUpTexture.cpp b/Engine/Source/Programs/Enterprise/Datasmith/DatasmithSketchUpRubyExporter/Private/DatasmithSketchUpTexture.cpp
--- a/Engine/Source/Programs/Enterprise/Datasmith/DatasmithSketchUpRubyExporter/Private/DatasmithSketchUpTexture.cpp
+++ b/Engine/Source/Programs/Enterprise/Datasmith/DatasmithSketchUpRubyExporter/Private/DatasmithSketchUpTexture.cpp
@@ -100,7 +100,7 @@ bool FTexture::GetTextureUseAlphaChannel()
 void FTexture::WriteImageFile(FExportContext& Context, const FString& TextureFilePath)
 {
 	// Write the SketchUp texture into a file when required.
-	SUResult SResult = SResult = SUTextureWriteToFile(TextureRef, TCHAR_TO_UTF8(*TextureFilePath));
+	SUResult SResult = SUTextureWriteToFile(TextureRef, TCHAR_TO_UTF8(*TextureFilePath));
 	if (SResult == SU_ERROR_SERIALIZATION)
 	{
 		// TODO: Append an error message to the export summary.
diff --git a/Engine/Source/Programs/Enterprise/Datasmith/DatasmithSketchUpRubyExporter/Private/DatasmithSketchUpUtils.cpp b/Engine/Source/Programs/Enterprise/Datasmith/DatasmithSketchUpRubyExporter/Private/DatasmithSketchUpUtils.cpp
--- a/Engine/Source/Programs/Enterprise/Datasmith/DatasmithSketchUpRubyExporter/Private/DatasmithSketchUpUtils.cpp
+++ b/Engine/Source/Programs/Enterprise/Datasmith/DatasmithSketchUpRubyExporter/Private/DatasmithSketchUpUtils.cpp
@@ -51,10 +51,17 @@ SULayerRef DatasmithSketchUpUtils::GetEffectiveLayer(
 	SUComponentInstanceRef InComponentInstanceRef, // valid SketckUp component instance
 	SULayerRef             InInheritedLayerRef
 )
+{
+	SUDrawingElementRef DrawingElementRef = SUComponentInstanceToDrawingElement(InComponentInstanceRef);
+
+	return GetEffectiveLayer(DrawingElementRef, InInheritedLayerRef);
+}
+
+SULayerRef DatasmithSketchUpUtils::GetEffectiveLayer(SUDrawingElementRef DrawingElementRef, SULayerRef InInheritedLayerRef)
 {
 	// Retrieve the SketckUp component instance layer.
 	SULayerRef SComponentInstanceLayerRef = SU_INVALID;
-	SUDrawingElementGetLayer(SUComponentInstanceToDrawingElement(InComponentInstanceRef), &SComponentInstanceLayerRef); // we can ignore the returned SU_RESULT
+	SUDrawingElementGetLayer(DrawingElementRef, &SComponentInstanceLayerRef); // we can ignore the returned SU_RESULT
 
 	// Retrieve the SketchUp component instance layer name.
 	FString SComponentInstanceLayerName;
@@ -64,6 +71,8 @@ SULayerRef DatasmithSketchUpUtils::GetEffectiveLayer(
 	return SComponentInstanceLayerName.Equals(TEXT("Layer0")) ? InInheritedLayerRef : SComponentInstanceLayerRef;
 }
 
+
+
 FComponentDefinitionIDType DatasmithSketchUpUtils::GetComponentID(
 	SUComponentDefinitionRef InComponentDefinitionRef
 )
@@ -85,8 +94,6 @@ FComponentInstanceIDType DatasmithSketchUpUtils::GetGroupID(
 	return GetEntityID(SUGroupToEntity(InGroupRef));
 }
 
-
-
 int64 DatasmithSketchUpUtils::GetComponentPID(
 	SUComponentInstanceRef InComponentInstanceRef
 )
diff --git a/Engine/Source/Programs/Enterprise/Datasmith/DatasmithSketchUpRubyExporter/Private/DatasmithSketchUpUtils.h b/Engine/Source/Programs/Enterprise/Datasmith/DatasmithSketchUpRubyExporter/Private/DatasmithSketchUpUtils.h
--- a/Engine/Source/Programs/Enterprise/Datasmith/DatasmithSketchUpRubyExporter/Private/DatasmithSketchUpUtils.h
+++ b/Engine/Source/Programs/Enterprise/Datasmith/DatasmithSketchUpRubyExporter/Private/DatasmithSketchUpUtils.h
@@ -78,6 +78,7 @@ namespace DatasmithSketchUpUtils
 		SUComponentInstanceRef InSComponentInstanceRef, // valid SketckUp component instance
 		SULayerRef             InSInheritedLayerRef     // SketchUp inherited layer
 	);
+	SULayerRef GetEffectiveLayer(SUDrawingElementRef DrawingElementRef, SULayerRef InInheritedLayerRef);
 
 	// Return whether or not a SketckUp component instance is visible in the current SketchUp scene.
 	bool IsVisible(
