[SketchUp] Material Transparency with Texture

- fixed case when Sketchup transparent material using texture without alpha appearad opaque
- fixed combining scalar opacity and texture alpha(now they are modulated)
- fixed VERY old issue with texture opacity source - INVERTED BLUE channel of material texture was used as opacity source

diff --git a/Engine/Source/Programs/Enterprise/Datasmith/DatasmithSketchUpRubyExporter/Private/DatasmithSketchUpMaterial.cpp b/Engine/Source/Programs/Enterprise/Datasmith/DatasmithSketchUpRubyExporter/Private/DatasmithSketchUpMaterial.cpp
--- a/Engine/Source/Programs/Enterprise/Datasmith/DatasmithSketchUpRubyExporter/Private/DatasmithSketchUpMaterial.cpp
+++ b/Engine/Source/Programs/Enterprise/Datasmith/DatasmithSketchUpRubyExporter/Private/DatasmithSketchUpMaterial.cpp
@@ -153,7 +153,9 @@ namespace DatasmithSketchUp
 
 		DatasmithMaterialElementPtr->SetTwoSided(false);// todo: consider this
 
-		bool bTranslucent = InMaterial.bSourceColorAlphaUsed;
+		IDatasmithMaterialExpression* ExpressionOpacityFromTexture = nullptr;
+		IDatasmithMaterialExpressionScalar* ExpressionOpacityScalar = nullptr;
+
 		if (Texture)
 		{
 			IDatasmithMaterialExpressionTexture* ExpressionTexture = DatasmithMaterialElementPtr->AddMaterialExpression< IDatasmithMaterialExpressionTexture >();
@@ -187,18 +189,10 @@ namespace DatasmithSketchUp
 
 			ExpressionTexture->ConnectExpression(DatasmithMaterialElementPtr->GetBaseColor());
 
-			bTranslucent = bTranslucent || Texture->GetTextureUseAlphaChannel();
-
 			// Set the Datasmith material element opacity.
 			if (Texture->GetTextureUseAlphaChannel())
 			{
-				// Invert texture translarency to get Unreal opacity
-				IDatasmithMaterialExpressionGeneric* ExpressionOpacity = DatasmithMaterialElementPtr->AddMaterialExpression<IDatasmithMaterialExpressionGeneric>();
-				ExpressionOpacity->SetExpressionName(TEXT("OneMinus"));
-
-				ExpressionTexture->ConnectExpression(*ExpressionOpacity->GetInput(0), 3);
-
-				ExpressionOpacity->ConnectExpression(DatasmithMaterialElementPtr->GetOpacity());
+				ExpressionOpacityFromTexture = ExpressionTexture;
 			}
 		}
 		else
@@ -207,20 +201,37 @@ namespace DatasmithSketchUp
 			ExpressionColor->SetName(TEXT("Base Color"));
 			ExpressionColor->GetColor() = LinearColor;
 			ExpressionColor->ConnectExpression(DatasmithMaterialElementPtr->GetBaseColor());
+		}
 
-			// Set the Datasmith material element opacity.
-			if (InMaterial.bSourceColorAlphaUsed)
-			{
-				IDatasmithMaterialExpressionScalar* ExpressionOpacity = DatasmithMaterialElementPtr->AddMaterialExpression<IDatasmithMaterialExpressionScalar>();
-				ExpressionOpacity->SetName(TEXT("Opacity"));
-				ExpressionOpacity->GetScalar() = float(InMaterial.SourceColor.alpha) / float(255);
-				ExpressionOpacity->ConnectExpression(DatasmithMaterialElementPtr->GetOpacity());
-			}
+		// Create scalar parameter if opacity is set in material or texture has opacity channel recognized by SU
+		if (InMaterial.bSourceColorAlphaUsed || ExpressionOpacityFromTexture)
+		{
+			ExpressionOpacityScalar = DatasmithMaterialElementPtr->AddMaterialExpression<IDatasmithMaterialExpressionScalar>();
+			ExpressionOpacityScalar->SetName(TEXT("Opacity"));
+			ExpressionOpacityScalar->GetScalar() = InMaterial.bSourceColorAlphaUsed ? (float(InMaterial.SourceColor.alpha) / float(255)) : 1.f;
+			ExpressionOpacityScalar->ConnectExpression(DatasmithMaterialElementPtr->GetOpacity());
 		}
 
+		IDatasmithMaterialExpression* ExpressionOpacity = nullptr;
 
-		if (bTranslucent)
+		if (ExpressionOpacityFromTexture)
+		{
+			IDatasmithMaterialExpressionGeneric* Multiply = DatasmithMaterialElementPtr->AddMaterialExpression<IDatasmithMaterialExpressionGeneric>();
+			Multiply->SetExpressionName(TEXT("Multiply"));
+
+			ExpressionOpacityFromTexture->ConnectExpression(*Multiply->GetInput(0), 4);  // Use Alpha output from Texture sampler
+			ExpressionOpacityScalar->ConnectExpression(*Multiply->GetInput(1));
+
+			ExpressionOpacity = Multiply;
+		}
+		else
+		{
+			ExpressionOpacity = ExpressionOpacityScalar;
+		}
+
+		if (ExpressionOpacity)
 		{
+			ExpressionOpacity->ConnectExpression(DatasmithMaterialElementPtr->GetOpacity());
 			DatasmithMaterialElementPtr->SetBlendMode(/*EBlendMode::BLEND_Translucent*/2);
 		}
 
