[3ds Max] Hide By [REDACTED_PERSON] support

diff --git a/Engine/Source/Programs/Enterprise/Datasmith/DatasmithMaxExporter/Private/DatasmithMaxDirectLinkNotifications.cpp b/Engine/Source/Programs/Enterprise/Datasmith/DatasmithMaxExporter/Private/DatasmithMaxDirectLinkNotifications.cpp
--- a/Engine/Source/Programs/Enterprise/Datasmith/DatasmithMaxExporter/Private/DatasmithMaxDirectLinkNotifications.cpp
+++ b/Engine/Source/Programs/Enterprise/Datasmith/DatasmithMaxExporter/Private/DatasmithMaxDirectLinkNotifications.cpp
@@ -458,6 +458,10 @@ public:
 			SceneTracker.NodeMaterialGraphModified(&Node);
 		break;
 
+		// Sadly, this event only sent for node when it's being hidden by using Display/Hide Byt Category. Unhiding it using that option doesn't send meaningful message per node
+		// At least NOTIFY_BY_CATEGORY_DISPLAY_FILTER_CHANGED is sent(handled elsewhere)
+		// case REFMSG_NODE_DISPLAY_PROP_CHANGED: break;
+
 		case REFMSG_SUBANIM_STRUCTURE_CHANGED:
 			if (Node.GetMtl() == TargetHandle)
 			{
@@ -689,6 +693,10 @@ void FNotifications::On3dsMaxNotification(void* param, NotifyInfo* info)
 		ISceneTracker& SceneTracker = Exporter->GetSceneTracker();
 		switch (info->intcode)
 		{
+		case NOTIFY_BY_CATEGORY_DISPLAY_FILTER_CHANGED:
+			SceneTracker.HideByCategoryChanged();
+			break;
+
 		case NOTIFY_NODE_POST_MTL:
 			// Event - node got a new material
 			break;
diff --git a/Engine/Source/Programs/Enterprise/Datasmith/DatasmithMaxExporter/Private/DatasmithMaxExporterWithDirectLink.cpp b/Engine/Source/Programs/Enterprise/Datasmith/DatasmithMaxExporter/Private/DatasmithMaxExporterWithDirectLink.cpp
--- a/Engine/Source/Programs/Enterprise/Datasmith/DatasmithMaxExporter/Private/DatasmithMaxExporterWithDirectLink.cpp
+++ b/Engine/Source/Programs/Enterprise/Datasmith/DatasmithMaxExporter/Private/DatasmithMaxExporterWithDirectLink.cpp
@@ -1302,6 +1302,53 @@ public:
 		}
 	}
 
+	void UpdateHideByCategory(bool& bChangeEncountered)
+	{
+		DWORD HideByCategoryFlagsNew = GetCOREInterface()->GetHideByCategoryFlags();
+		if (bHideByCategoryChanged)
+		{
+			DWORD HideByCategoryFlagsChanged = HideByCategoryFlags ^ HideByCategoryFlagsNew;
+
+			for (TPair<FNodeKey, FNodeTrackerHandle> NodeKeyAndNodeTracker: NodeTrackers)
+			{
+				// Excluded ProgressCounter/Cancel handling from this loop - the code is simple, running quickly even for lots of nodes.
+
+				bool bNeedInvalidate = false;
+				FNodeTracker* NodeTracker = NodeKeyAndNodeTracker.Value.GetNodeTracker();
+				switch (NodeTracker->GetConverterType())
+				{
+				case FNodeConverter::Unknown:
+					bNeedInvalidate = true;
+					break;
+				case FNodeConverter::MeshNode:
+					bNeedInvalidate = (HideByCategoryFlagsChanged & (HIDE_OBJECTS|HIDE_SHAPES)) != 0;
+					break;
+				case FNodeConverter::HismNode: 
+					bNeedInvalidate = (HideByCategoryFlagsChanged & HIDE_OBJECTS) != 0;
+					break;
+				case FNodeConverter::LightNode:
+					bNeedInvalidate = (HideByCategoryFlagsChanged & HIDE_LIGHTS) != 0;
+					break;
+				case FNodeConverter::CameraNode: 
+					bNeedInvalidate = (HideByCategoryFlagsChanged & HIDE_CAMERAS) != 0;
+					break;
+				case FNodeConverter::HelperNode: 
+					bNeedInvalidate = (HideByCategoryFlagsChanged & HIDE_HELPERS) != 0;
+					break;
+				default: ;
+				}
+				if (bNeedInvalidate)
+				{
+					InvalidateNode(*NodeTracker, false);
+					bChangeEncountered = true;
+				}
+			}
+			
+			bHideByCategoryChanged = false;
+		}
+		HideByCategoryFlags = HideByCategoryFlagsNew;
+	}
+
 	 // returns if any change in scene was encountered and scene update completed(i.e. DirectLink Sync can be run)
 	bool UpdateInternal(FUpdateProgress::FStage& MainStage)
 	{
@@ -1335,6 +1382,11 @@ public:
 			bChangeEncountered |= InvalidatedNodeTrackers.PurgeDeletedNodeTrackers(*this);
 		}
 
+		{
+			PROGRESS_STAGE("Check Hide By Category")
+			UpdateHideByCategory(bChangeEncountered);
+		}
+
 		{
 			PROGRESS_STAGE("Check Nodes Validity")
 			PROGRESS_STAGE_RESULT_DEFERRED
@@ -3051,11 +3103,19 @@ public:
 		}
 	}
 
+	virtual void HideByCategoryChanged()
+	{
+		bHideByCategoryChanged = true;
+	}
+
 	virtual void NodeNameChanged(FNodeKey NodeKey) override
 	{
 		if (FNodeTracker* NodeTracker =  GetNodeTracker(NodeKey))
 		{
-			NodeTrackersNames.Update(*NodeTracker);
+			if (!NodeTracker->bDeleted)
+			{
+				NodeTrackersNames.Update(*NodeTracker);
+			}
 			InvalidateNode(*NodeTracker);
 		}
 	}
@@ -3141,6 +3201,9 @@ public:
 
 	TMap<FNodeTracker*, TSharedPtr<IDatasmithMetaDataElement>> NodeDatasmithMetadata; // All scene nodes
 
+	bool bHideByCategoryChanged = false;
+	DWORD HideByCategoryFlags = 0;
+
 	// Nodes/instances that need to be rebuilt
 	FInvalidatedNodeTrackers InvalidatedNodeTrackers;
 	TSet<FInstances*> InvalidatedInstances;
diff --git a/Engine/Source/Programs/Enterprise/Datasmith/DatasmithMaxExporter/Public/DatasmithMaxDirectLink.h b/Engine/Source/Programs/Enterprise/Datasmith/DatasmithMaxExporter/Public/DatasmithMaxDirectLink.h
--- a/Engine/Source/Programs/Enterprise/Datasmith/DatasmithMaxExporter/Public/DatasmithMaxDirectLink.h
+++ b/Engine/Source/Programs/Enterprise/Datasmith/DatasmithMaxExporter/Public/DatasmithMaxDirectLink.h
@@ -383,6 +383,8 @@ public:
 
 	virtual void MaterialGraphModified(Mtl* Material) = 0;
 
+	virtual void HideByCategoryChanged() = 0;
+
 	virtual bool IsUpdateInProgress() = 0;
 
 	// Scene modification
@@ -394,7 +396,7 @@ public:
 	virtual void RemoveTexture(const TSharedPtr<IDatasmithTextureElement>&) = 0;
 	virtual void NodeXRefMerged(INode* Node) = 0;
 	virtual void RemapConvertedMaterialUVChannels(Mtl* ActualMaterial, const TSharedPtr<IDatasmithBaseMaterialElement>& DatasmithMaterial) = 0;
-
+	
 	// Sync/Update
 	FSyncPoint CurrentSyncPoint;
 
