[3ds Max] Fixed camera parameters export for standard cameras

 - note: this was not a regression

diff --git a/Engine/Source/Programs/Enterprise/Datasmith/DatasmithMaxExporter/Private/DatasmithMaxCameraExporter.cpp b/Engine/Source/Programs/Enterprise/Datasmith/DatasmithMaxExporter/Private/DatasmithMaxCameraExporter.cpp
--- a/Engine/Source/Programs/Enterprise/Datasmith/DatasmithMaxExporter/Private/DatasmithMaxCameraExporter.cpp
+++ b/Engine/Source/Programs/Enterprise/Datasmith/DatasmithMaxExporter/Private/DatasmithMaxCameraExporter.cpp
@@ -1,6 +1,7 @@
 // Copyright Epic Games, Inc. All Rights Reserved.
 #include "DatasmithMaxCameraExporter.h"
 
+#include "DatasmithMaxDirectLink.h"
 #include "DatasmithMaxExporterDefines.h"
 #include "DatasmithMaxHelper.h"
 #include "DatasmithMaxSceneParser.h"
@@ -392,7 +393,7 @@ namespace DatasmithMaxCameraExporterImpl
 	}
 }
 
-bool FDatasmithMaxCameraExporter::ExportCamera(INode& Node, TSharedRef< IDatasmithCameraActorElement > Camera)
+bool FDatasmithMaxCameraExporter::ExportCamera(TimeValue CurrentTime, INode& Node, TSharedRef< IDatasmithCameraActorElement > Camera)
 {
 	ObjectState ObjState = Node.EvalWorldState(0);
 	CameraObject* CamObj = (CameraObject*)ObjState.obj;
@@ -403,21 +404,34 @@ bool FDatasmithMaxCameraExporter::ExportCamera(INode& Node, TSharedRef< IDatasmi
 
 	Camera->SetSensorWidth( GetCOREInterface()->GetRendApertureWidth() );
 	Camera->SetSensorAspectRatio( (float)GetCOREInterface()->GetRendWidth() / (float)GetCOREInterface()->GetRendHeight() );
-	Camera->SetFocalLength( GetCOREInterface()->GetRendApertureWidth() / (2.f * FMath::Tan(CamObj->GetFOV(GetCOREInterface()->GetTime()) * 0.5f)) );
+	Camera->SetFocalLength( GetCOREInterface()->GetRendApertureWidth() / (2.f * FMath::Tan(CamObj->GetFOV(CurrentTime) * 0.5f)) );
 
 	int NumParamBlocks = CamObj->NumParamBlocks();
 	float CustomUnit = FMath::Abs((float)GetSystemUnitScale(UNITS_CENTIMETERS));
 
-	DatasmithMaxCameraExporterImpl::FMaxCameraParameters MaxCameraParameters = DatasmithMaxCameraExporterImpl::ParseCameraParameters( *CamObj );
+	Class_ID ClassID = CamObj->ClassID();
 
-	Camera->SetFocusDistance( MaxCameraParameters.TargetDistance );
-	Camera->SetSensorWidth( MaxCameraParameters.FilmWidth );
-	Camera->SetFocalLength( MaxCameraParameters.FocalLength );
-	Camera->SetEnableDepthOfField( MaxCameraParameters.bUseDepthOfField );
+	if ((ClassID == Class_ID(LOOKAT_CAM_CLASS_ID, 0)) || (ClassID == Class_ID(SIMPLE_CAM_CLASS_ID, 0)))
+	{
+		GenCamera* GenCam = (GenCamera*)CamObj;
 
-	if ( MaxCameraParameters.bUseVignette )
+		float FocusDistance = GenCam->GetTDist(CurrentTime) * DatasmithMaxDirectLink::FDatasmithConverter().UnitToCentimeter;
+		Camera->SetFocusDistance( FocusDistance );
+		Camera->SetEnableDepthOfField( GenCam->GetDOFEnable(CurrentTime) );
+		// SensorWidth defaulted from GetRendApertureWidth above
+	}
+	else
 	{
-		Camera->GetPostProcess()->SetVignette( MaxCameraParameters.VignettingAmount );
+		DatasmithMaxCameraExporterImpl::FMaxCameraParameters MaxCameraParameters = DatasmithMaxCameraExporterImpl::ParseCameraParameters( *CamObj );
+		Camera->SetFocusDistance( MaxCameraParameters.TargetDistance );
+		Camera->SetSensorWidth( MaxCameraParameters.FilmWidth );
+		Camera->SetFocalLength( MaxCameraParameters.FocalLength );
+		Camera->SetEnableDepthOfField( MaxCameraParameters.bUseDepthOfField );
+
+		if ( MaxCameraParameters.bUseVignette )
+		{
+			Camera->GetPostProcess()->SetVignette( MaxCameraParameters.VignettingAmount );
+		}
 	}
 
 	// It's sufficient to check for the node's target to see if the camera has a look-at behavior
@@ -427,11 +441,11 @@ bool FDatasmithMaxCameraExporter::ExportCamera(INode& Node, TSharedRef< IDatasmi
 		Camera->SetLookAtActor( *FString::FromInt( CameraTargetNode->GetHandle() ) );
 	}
 
-	if ( CamObj->ClassID() == PHYSICALCAMERA_CLASS )
+	if ( ClassID == PHYSICALCAMERA_CLASS )
 	{
 		ExportPhysicalCamera( *CamObj, Camera );
 	}
-	else if ( CamObj->ClassID() == VRAY_PHYSICALCAMERA_CLASS )
+	else if ( ClassID == VRAY_PHYSICALCAMERA_CLASS )
 	{
 		ExportVRayPhysicalCamera( *CamObj, Camera );
 	}
diff --git a/Engine/Source/Programs/Enterprise/Datasmith/DatasmithMaxExporter/Private/DatasmithMaxConverters.cpp b/Engine/Source/Programs/Enterprise/Datasmith/DatasmithMaxExporter/Private/DatasmithMaxConverters.cpp
--- a/Engine/Source/Programs/Enterprise/Datasmith/DatasmithMaxExporter/Private/DatasmithMaxConverters.cpp
+++ b/Engine/Source/Programs/Enterprise/Datasmith/DatasmithMaxExporter/Private/DatasmithMaxConverters.cpp
@@ -49,7 +49,7 @@ void FCameraNodeConverter::ConvertToDatasmith(ISceneTracker& SceneTracker, FNode
 	FString UniqueName = FString::FromInt(NodeTracker.Node->GetHandle());
 	NodeTracker.CreateConverted().DatasmithActorElement = FDatasmithSceneFactory::CreateCameraActor((const TCHAR*)*UniqueName);
 
-	FDatasmithMaxCameraExporter::ExportCamera(*NodeTracker.Node, StaticCastSharedPtr<IDatasmithCameraActorElement>(NodeTracker.GetConverted().DatasmithActorElement).ToSharedRef());
+	FDatasmithMaxCameraExporter::ExportCamera(SceneTracker.CurrentSyncPoint.Time, *NodeTracker.Node, StaticCastSharedPtr<IDatasmithCameraActorElement>(NodeTracker.GetConverted().DatasmithActorElement).ToSharedRef());
 
 	SceneTracker.SetupActor(NodeTracker);
 
diff --git a/Engine/Source/Programs/Enterprise/Datasmith/DatasmithMaxExporter/Private/DatasmithMaxSceneExporter.cpp b/Engine/Source/Programs/Enterprise/Datasmith/DatasmithMaxExporter/Private/DatasmithMaxSceneExporter.cpp
--- a/Engine/Source/Programs/Enterprise/Datasmith/DatasmithMaxExporter/Private/DatasmithMaxSceneExporter.cpp
+++ b/Engine/Source/Programs/Enterprise/Datasmith/DatasmithMaxExporter/Private/DatasmithMaxSceneExporter.cpp
@@ -2146,7 +2146,7 @@ bool FDatasmithMaxSceneExporter::ExportCameraActor(TSharedRef< IDatasmithScene >
 
 	TSharedRef< IDatasmithCameraActorElement > CameraActor = FDatasmithSceneFactory::CreateCameraActor(Name);
 
-	if ( !FDatasmithMaxCameraExporter::ExportCamera( *Parent, CameraActor ) )
+	if ( !FDatasmithMaxCameraExporter::ExportCamera( GetCOREInterface()->GetTime(), *Parent, CameraActor ) )
 	{
 		return false;
 	}
diff --git a/Engine/Source/Programs/Enterprise/Datasmith/DatasmithMaxExporter/Public/DatasmithMaxCameraExporter.h b/Engine/Source/Programs/Enterprise/Datasmith/DatasmithMaxExporter/Public/DatasmithMaxCameraExporter.h
--- a/Engine/Source/Programs/Enterprise/Datasmith/DatasmithMaxExporter/Public/DatasmithMaxCameraExporter.h
+++ b/Engine/Source/Programs/Enterprise/Datasmith/DatasmithMaxExporter/Public/DatasmithMaxCameraExporter.h
@@ -5,6 +5,15 @@
 #include "CoreMinimal.h"
 #include "Templates/SharedPointer.h"
 
+
+#include "DatasmithMaxHelper.h"
+
+#include "Windows/AllowWindowsPlatformTypes.h"
+MAX_INCLUDES_START
+	#include "maxtypes.h"
+MAX_INCLUDES_END
+#include "Windows/HideWindowsPlatformTypes.h"
+
 class CameraObject;
 class IDatasmithCameraActorElement;
 class IDatasmithPostProcessVolumeElement;
@@ -17,7 +26,7 @@ class ToneOperator;
 class FDatasmithMaxCameraExporter
 {
 public:
-	static bool ExportCamera( INode& Node, TSharedRef< IDatasmithCameraActorElement > CameraElement );
+	static bool ExportCamera( TimeValue CurrentTime, INode& Node, TSharedRef< IDatasmithCameraActorElement > CameraElement );
 	static bool ExportToneOperator( ToneOperator& ToneOp, TSharedRef< IDatasmithPostProcessVolumeElement > PostProcessVolumeElement );
 
 private:
