[SketchUp] Fixed normals when baking transform into mesh

- normals are transformed with inverse-transposed orientation component of the bake transform

diff --git a/Engine/Source/Programs/Enterprise/Datasmith/DatasmithSketchUpRubyExporter/Private/DatasmithSketchUpMesh.cpp b/Engine/Source/Programs/Enterprise/Datasmith/DatasmithSketchUpRubyExporter/Private/DatasmithSketchUpMesh.cpp
--- a/Engine/Source/Programs/Enterprise/Datasmith/DatasmithSketchUpRubyExporter/Private/DatasmithSketchUpMesh.cpp
+++ b/Engine/Source/Programs/Enterprise/Datasmith/DatasmithSketchUpRubyExporter/Private/DatasmithSketchUpMesh.cpp
@@ -25,6 +25,7 @@
 #include "SketchUpAPI/model/mesh_helper.h"
 #include "SketchUpAPI/model/uv_helper.h"
 #include "SketchUpAPI/geometry/point3d.h"
+#include "SketchUpAPI/geometry/vector3d.h"
 #include <SketchUpAPI/geometry/transformation.h>
 #include "DatasmithSketchUpSDKCeases.h"
 
@@ -361,6 +362,33 @@ namespace DatasmithSketchUp
 		// Set the number of triangles of the exported Datasmith mesh.
 		OutDMesh.SetFacesCount(TriangleCount);
 
+		SUTransformation TransformNoTranslation = Transform;
+		TransformNoTranslation.values[12] = 0;
+		TransformNoTranslation.values[13] = 0;
+		TransformNoTranslation.values[14] = 0;
+
+		// Make inverse-transposed transform for normals from orientation component of source transform
+		SUTransformation TransformInverse;
+		SUTransformationGetInverse(&TransformNoTranslation, &TransformInverse);
+		SUTransformation TransformForNormals;
+
+		// Transpose
+		for (int32 X = 0; X < 4; ++X)
+		{
+			for (int32 Y = 0; Y < 4; ++Y)
+			{
+				TransformForNormals.values[X + Y*4] = TransformNoTranslation.values[Y + X*4];
+			}
+		}
+
+		TArray<SUVector3D> MeshVertexNormalsBaked;
+		MeshVertexNormalsBaked.Reserve(MeshVertexNormals.Num());
+		for (SUVector3D Normal : MeshVertexNormals)
+		{
+			SUVector3DTransform(&TransformForNormals, &Normal);
+			MeshVertexNormalsBaked.Add(Normal);
+		}
+
 		// Convert triangle vertex indices and normals from SketchUp right-handed Z-up coordinates to Unreal left-handed Z-up coordinates.
 		// To avoid perturbating X, which is forward in Unreal, the handedness conversion is done by flipping the side vector Y.
 		for (int32 TriangleNo = 0, NormalNo = 0; TriangleNo < TriangleCount; TriangleNo++)
@@ -374,9 +402,9 @@ namespace DatasmithSketchUp
 			OutDMesh.SetFace(TriangleNo, int32(TriangleIndices.IndexA), int32(TriangleIndices.IndexB), int32(TriangleIndices.IndexC), MeshTriangleSlotIds[TriangleNo]);
 
 			// Set the triangle vertex normals in the exported Datasmith mesh.
-			SMeshTriangleNormals TriangleNormals = { MeshVertexNormals[TriangleIndices.IndexA],
-													 MeshVertexNormals[TriangleIndices.IndexB],
-													 MeshVertexNormals[TriangleIndices.IndexC] };
+			SMeshTriangleNormals TriangleNormals = { MeshVertexNormalsBaked[TriangleIndices.IndexA],
+													 MeshVertexNormalsBaked[TriangleIndices.IndexB],
+													 MeshVertexNormalsBaked[TriangleIndices.IndexC] };
 			OutDMesh.SetNormal(NormalNo++, float(TriangleNormals.NormalA.x), float(-TriangleNormals.NormalA.y), float(TriangleNormals.NormalA.z));
 			OutDMesh.SetNormal(NormalNo++, float(TriangleNormals.NormalB.x), float(-TriangleNormals.NormalB.y), float(TriangleNormals.NormalB.z));
 			OutDMesh.SetNormal(NormalNo++, float(TriangleNormals.NormalC.x), float(-TriangleNormals.NormalC.y), float(TriangleNormals.NormalC.z));
