[3ds Max] More stable static mesh id

- using smallest handle value among instance nodes

diff --git a/Engine/Source/Programs/Enterprise/Datasmith/DatasmithMaxExporter/Private/DatasmithMaxExporterWithDirectLink.cpp b/Engine/Source/Programs/Enterprise/Datasmith/DatasmithMaxExporter/Private/DatasmithMaxExporterWithDirectLink.cpp
--- a/Engine/Source/Programs/Enterprise/Datasmith/DatasmithMaxExporter/Private/DatasmithMaxExporterWithDirectLink.cpp
+++ b/Engine/Source/Programs/Enterprise/Datasmith/DatasmithMaxExporter/Private/DatasmithMaxExporterWithDirectLink.cpp
@@ -2429,11 +2429,26 @@ public:
 		Instances.bShouldBakePivot = bShouldBakePivot;
 		Instances.BakePivot = BakePivot;
 
-		// Export static mesh using first lucky node
-		// todo: possible optimization to reuse previous node. Somehow. 
-		for (FNodeTracker* NodeTrackerPtr : Instances.NodeTrackers)
+		// Select node to export geometry from
+		// Use Node with smallest value of Handle to make change of StaticMesh Name less likely
+		// Re-exporting scene won't change StaticMesh even when nodes are reordered in hierarchy.
+		FNodeTracker* SelectedNodeTracker = nullptr;
+		ULONG SelectedHandle = 0;
+
+		for (FNodeTracker* NodeTracker : Instances.NodeTrackers)
 		{
-			FNodeTracker& NodeTracker = *NodeTrackerPtr;
+			ULONG Handle = NodeTracker->Node->GetHandle();
+			if (!SelectedNodeTracker || (Handle < SelectedHandle))
+			{
+				SelectedHandle = Handle;
+				SelectedNodeTracker = NodeTracker;
+			}
+		}
+
+
+		// Update geometry using selected Node
+		{
+			FNodeTracker& NodeTracker = *SelectedNodeTracker;
 
 			// todo: use single EnumProc instance to enumerate all nodes during update to:
 			//    - have single call to BeginEnumeration and EndEnumeration
@@ -2442,13 +2457,13 @@ public:
 			SCENE_UPDATE_STAT_INC(UpdateInstances, GeometryUpdated);
 			UpdateInstancesGeometry(Instances, NodeTracker);
 
+
 			// assign materials to static mesh for the first instance(others will use override on mesh actors)
 			static_cast<FMeshNodeConverter&>(NodeTracker.GetConverter()).bMaterialsAssignedToStaticMesh = true;
 			if (Mtl* Material = UpdateGeometryNodeMaterial(*this, Instances, NodeTracker))
 			{
 				Instances.AssignMaterialToStaticMesh(MaterialsCollectionTracker, Material);
 			}
-			break;
 		}
 
 		for (FNodeTracker* NodeTrackerPtr : Instances.NodeTrackers)
