[SketchUp] Optimize Actor counts

- Component Instance without children doesn't create extra actor(only mesh actor)
- Image creates only Mesh
- fixed medatada removal when component is removed

diff --git a/Engine/Source/Programs/Enterprise/Datasmith/DatasmithSketchUpRubyExporter/Private/DatasmithSketchUpComponent.cpp b/Engine/Source/Programs/Enterprise/Datasmith/DatasmithSketchUpRubyExporter/Private/DatasmithSketchUpComponent.cpp
--- a/Engine/Source/Programs/Enterprise/Datasmith/DatasmithSketchUpRubyExporter/Private/DatasmithSketchUpComponent.cpp
+++ b/Engine/Source/Programs/Enterprise/Datasmith/DatasmithSketchUpRubyExporter/Private/DatasmithSketchUpComponent.cpp
@@ -78,14 +78,14 @@ void FDefinition::ParseNode(FExportContext& Context, FNodeOccurence& Node)
 }
 
 // Update mesh actors of an entiry with Entities(that is model or component instance)
-bool FEntityWithEntities::UpdateOccurrenceMeshActors(FExportContext& Context, FNodeOccurence& Node)
+void FEntityWithEntities::UpdateOccurrenceMeshActors(FExportContext& Context, FNodeOccurence& Node)
 {
 	FDefinition* EntityDefinition = GetDefinition();
 	DatasmithSketchUp::FEntitiesGeometry* EntitiesGeometry = EntityDefinition->GetEntities().EntitiesGeometry.Get();
 
 	if (!EntitiesGeometry)
 	{
-		return true;
+		return;
 	}
 
 	Node.MeshActors.Reset(EntitiesGeometry->GetMeshCount());
@@ -114,50 +114,11 @@ bool FEntityWithEntities::UpdateOccurrenceMeshActors(FExportContext& Context, FN
 		FString InstancePathTag = ComponentActorName.Replace(TEXT("SU"), TEXT("SU.PATH.0")).Replace(TEXT("_"), TEXT("."));
 		DMeshActorPtr->AddTag(*InstancePathTag);
 
-		// Add the mesh actor to our component Datasmith actor hierarchy.
-		if (Node.DatasmithActorElement.IsValid()) 
-		{
-			Node.DatasmithActorElement->AddChild(DMeshActorPtr);
-		}
-		else
-		{
-			Context.DatasmithScene->AddActor(DMeshActorPtr);
-		}
-
 		// ADD_TRACE_LINE(TEXT("Actor %ls: %ls %ls %ls"), *MeshActorLabel, *ComponentDepthTag, *DefinitionGUIDTag, *InstancePathTag);
 
 		// Set the Datasmith mesh element used by the mesh actor.
 		DMeshActorPtr->SetStaticMeshPathName(EntitiesGeometry->GetMeshElementName(MeshIndex));
 	}
-	return false;
-}
-
-void FNodeOccurence::UpdateMeshActors(FExportContext& Context)
-{
-	// Remove old mesh actors
-	// todo: reuse old mesh actors (also can keep instances when removing due to say hidden)
-	if (DatasmithActorElement.IsValid())
-	{
-		for (TSharedPtr<IDatasmithMeshActorElement> MeshActor : MeshActors)
-		{
-			DatasmithActorElement->RemoveChild(MeshActor);
-		}
-	}
-	else
-	{
-		for (TSharedPtr<IDatasmithMeshActorElement> MeshActor : MeshActors)
-		{
-			Context.DatasmithScene->RemoveActor(MeshActor, EDatasmithActorRemovalRule::RemoveChildren);
-		}
-	}
-	MeshActors.Reset();
-
-	if (!bVisible) // no mesh actors for invisible node
-	{
-		return;
-	}
-
-	Entity.UpdateOccurrenceMeshActors(Context, *this);
 }
 
 void FNodeOccurence::UpdateVisibility(FExportContext& Context)
@@ -191,13 +152,22 @@ void FNodeOccurence::Update(FExportContext& Context)
 
 	if (bMeshActorsInvalidated)
 	{
-		UpdateMeshActors(Context);
+		Entity.ResetOccurrenceActors(Context, *this);
+
+		if (bVisible)
+		{
+			Entity.UpdateOccurrenceMeshActors(Context, *this);
+		}
 		bMeshActorsInvalidated = false;
 	}
 
 	if (bPropertiesInvalidated)
 	{
-		Entity.UpdateOccurrence(Context, *this);
+		if (bVisible)
+		{
+			Entity.UpdateOccurrence(Context, *this);
+		}
+
 		bPropertiesInvalidated = false;
 	}
 
@@ -254,8 +224,8 @@ void FNodeOccurence::RemoveOccurrence(FExportContext& Context)
 		ParentNode->Children.Remove(this);
 	}
 
-	// Usually child component instances are removed in proper order but grouping entities
-	// has this weird behavior that containing component removed without cleaning. Test case:
+	// Usually child component instances are removed in proper order - children first. Nut grouping entities
+	// has this weird behavior that containing component removed without cleaning its children. Test case:
 	// Group an instance and some other entity(e.g. face) and then convert group to component
 	// then converted group is removed without other events for its children
 	// Probably this is because those 'children' are actually entities in the group's definition
@@ -266,30 +236,7 @@ void FNodeOccurence::RemoveOccurrence(FExportContext& Context)
 		Child->RemoveOccurrence(Context);
 	}
 
-	for (const TSharedPtr<IDatasmithMeshActorElement>& MeshActor : MeshActors) 
-	{
-		if (const TSharedPtr<IDatasmithActorElement>& ParentActor = MeshActor->GetParentActor())
-		{
-			ParentActor->RemoveChild(MeshActor);
-		}
-		else
-		{
-			Context.DatasmithScene->RemoveActor(MeshActor, EDatasmithActorRemovalRule::RemoveChildren);
-		}
-	}
-
-	if (DatasmithActorElement)
-	{
-		if (const TSharedPtr<IDatasmithActorElement>& ParentActor = DatasmithActorElement->GetParentActor())
-		{
-			ParentActor->RemoveChild(DatasmithActorElement);
-		}
-		else
-		{
-			Context.DatasmithScene->RemoveActor(DatasmithActorElement, EDatasmithActorRemovalRule::RemoveChildren);
-		}
-	}
-
+	Entity.ResetOccurrenceActors(Context, *this);
 	Entity.DeleteOccurrence(Context, this);
 }
 
@@ -308,7 +255,7 @@ void FNodeOccurence::ResetMetadataElement(FExportContext& Context)
 	{
 		DatasmithMetadataElement->SetName(*MetadataElementName);
 	}
-	DatasmithMetadataElement->SetLabel(*DatasmithActorLabel);
+	DatasmithMetadataElement->SetLabel(*GetActorLabel());
 	DatasmithMetadataElement->ResetProperties();
 }
 
@@ -321,6 +268,7 @@ void FNodeOccurence::RemoveDatasmithActorHierarchy(FExportContext& Context)
 {
 	if (!DatasmithActorElement)
 	{
+		// Hierarchy already removed(or wasn't created)
 		return;
 	}
 
@@ -330,27 +278,52 @@ void FNodeOccurence::RemoveDatasmithActorHierarchy(FExportContext& Context)
 		ChildNode->RemoveDatasmithActorHierarchy(Context);
 	}
 
-	for (TSharedPtr<IDatasmithMeshActorElement> MeshActor : MeshActors)
-	{
-		DatasmithActorElement->RemoveChild(MeshActor);
-	}
-	MeshActors.Reset();
+	Entity.ResetOccurrenceActors(Context, *this);
+}
 
-	if (const TSharedPtr<IDatasmithActorElement>& ParentActor = DatasmithActorElement->GetParentActor())
-	{
-		ParentActor->RemoveChild(DatasmithActorElement);
-	}
-	else
-	{
-		Context.DatasmithScene->RemoveActor(DatasmithActorElement, EDatasmithActorRemovalRule::RemoveChildren);
-	}
-	DatasmithActorElement.Reset();
+void FNodeOccurence::ResetNodeActors(FExportContext& Context)
+{
+	FNodeOccurence& Node = *this;
 
-	if (DatasmithMetadataElement)
+	// Remove old mesh actors
+	// todo: reuse old mesh actors (also can keep instances when removing due to say hidden)
+	if (Node.DatasmithActorElement)
 	{
-		Context.DatasmithScene->RemoveMetaData(DatasmithMetadataElement);
+		// Check if component used an actor to combine mesh and child nodes under it
+		// todo: just add flag for code clearness?
+		bool bHasActor = Node.MeshActors.IsEmpty() || (Node.DatasmithActorElement != Node.MeshActors[0]);
+
+		if (bHasActor)  
+		{
+			// In this case detach all the children before removing actor from the parent/scene
+			// note: DatasmithScene::RemoveActor has only two ways to remove children - relocating then to Scene root or deleting hierarchy
+
+			int32 ChildCount = Node.DatasmithActorElement->GetChildrenCount();
+			// Remove last child each time to optimize array elements relocation
+			for(int32 ChildIndex = ChildCount-1; ChildIndex >= 0; --ChildIndex)
+			{
+				Node.DatasmithActorElement->RemoveChild(Node.DatasmithActorElement->GetChild(ChildIndex));
+			}
+		}
+
+		if (const TSharedPtr<IDatasmithActorElement>& ParentActor = Node.DatasmithActorElement->GetParentActor())
+		{
+			ParentActor->RemoveChild(Node.DatasmithActorElement);
+		}
+		else
+		{
+			Context.DatasmithScene->RemoveActor(Node.DatasmithActorElement, EDatasmithActorRemovalRule::RemoveChildren);
+		}
+		Node.DatasmithActorElement.Reset();
+
+		if (DatasmithMetadataElement)
+		{
+			Context.DatasmithScene->RemoveMetaData(DatasmithMetadataElement);
+		}
+		DatasmithMetadataElement.Reset();
+
 	}
-	DatasmithMetadataElement.Reset();
+	Node.MeshActors.Reset();
 }
 
 FModelDefinition::FModelDefinition(SUModelRef InModel) : Model(InModel)
@@ -387,20 +360,6 @@ void FModelDefinition::FillOccurrenceActorMetadata(FNodeOccurence& Node)
 {
 }
 
-void FModelDefinition::CreateActor(FExportContext& Context, FNodeOccurence& Node)
-{
-	// Don't create single root Actor for model
-}
-
-void FModelDefinition::BuildNodeNames(FNodeOccurence& Node)
-{
-	// Get the SketckUp component instance persistent ID.
-	int64 SketchupPersistentID = Node.Entity.GetPersistentId();
-	Node.DatasmithActorName = FString::Printf(TEXT("%ls_%lld"), *GetSketchupSourceName(), SketchupPersistentID);
-
-	Node.DatasmithActorLabel = GetSketchupSourceName();
-}
-
 FString FModelDefinition::GetSketchupSourceName()
 {
 	FString SketchupSourceName = SuGetString(SUModelGetName, Model);
@@ -469,35 +428,29 @@ void FComponentDefinition::Parse(FExportContext& Context)
 	bSketchupSourceFaceCamera = SComponentBehavior.component_always_face_camera;
 }
 
-void FComponentDefinition::CreateActor(FExportContext& Context, FNodeOccurence& Node)
+void FComponentInstance::SetupActor(FExportContext& Context, FNodeOccurence& Node)
 {
-	BuildNodeNames(Node);
+	FComponentDefinition& EntityDefinition = *(FComponentDefinition *)GetDefinition();
 
-	// Create a Datasmith actor for the component instance.
-	Node.DatasmithActorElement = FDatasmithSceneFactory::CreateActor(*Node.DatasmithActorName); // a EDatasmithElementType::Actor
+	// Add the Datasmith actor component depth tag.
+	// We use component depth + 1 to factor in the added Datasmith scene root once imported in Unreal.
+	FString ComponentDepthTag = FString::Printf(TEXT("SU.DEPTH.%d"), Node.Depth);
+	Node.DatasmithActorElement->AddTag(*ComponentDepthTag);
 
-	{
-		// Add the Datasmith actor component depth tag.
-		// We use component depth + 1 to factor in the added Datasmith scene root once imported in Unreal.
-		FString ComponentDepthTag = FString::Printf(TEXT("SU.DEPTH.%d"), Node.Depth);
-		Node.DatasmithActorElement->AddTag(*ComponentDepthTag);
+	// Add the Datasmith actor component definition GUID tag.
+	FString DefinitionGUIDTag = FString::Printf(TEXT("SU.GUID.%ls"), *GetDefinition()->GetSketchupSourceGUID());
+	Node.DatasmithActorElement->AddTag(*DefinitionGUIDTag);
 
-		// Add the Datasmith actor component definition GUID tag.
-		FString DefinitionGUIDTag = FString::Printf(TEXT("SU.GUID.%ls"), *GetSketchupSourceGUID());
-		Node.DatasmithActorElement->AddTag(*DefinitionGUIDTag);
+	// Add the Datasmith actor component instance path tag.
+	FString InstancePathTag = Node.GetActorName().Replace(TEXT("SU"), TEXT("SU.PATH.0")).Replace(TEXT("_"), TEXT("."));
+	Node.DatasmithActorElement->AddTag(*InstancePathTag);
 
-		// Add the Datasmith actor component instance path tag.
-		FString InstancePathTag = Node.DatasmithActorName.Replace(TEXT("SU"), TEXT("SU.PATH.0")).Replace(TEXT("_"), TEXT("."));
-		Node.DatasmithActorElement->AddTag(*InstancePathTag);
-
-		// Add the Datasmith actor component instance face camera tag when required.
-		if (bSketchupSourceFaceCamera)
-		{
-			Node.DatasmithActorElement->AddTag(TEXT("SU.BEHAVIOR.FaceCamera"));
-		}
+	// Add the Datasmith actor component instance face camera tag when required.
+	if (EntityDefinition.bSketchupSourceFaceCamera)
+	{
+		Node.DatasmithActorElement->AddTag(TEXT("SU.BEHAVIOR.FaceCamera"));
 	}
 
-	// ComponentInstance occurrence always has parent node(Model it at top)
 	if (Node.ParentNode->DatasmithActorElement)
 	{
 		Node.ParentNode->DatasmithActorElement->AddChild(Node.DatasmithActorElement);
@@ -519,14 +472,14 @@ void FComponentDefinition::UpdateMetadata(FExportContext& Context)
 	ParsedMetadata = MakeUnique<FMetadata>(SUComponentDefinitionToEntity(ComponentDefinitionRef));;
 }
 
-void FComponentDefinition::BuildNodeNames(FNodeOccurence& Node)
+void FComponentInstance::BuildNodeNames(FNodeOccurence& Node)
 {
 	// Get the SketckUp component instance persistent ID.
 	int64 SketchupPersistentID = Node.Entity.GetPersistentId();
 	Node.DatasmithActorName = FString::Printf(TEXT("%ls_%lld"), *Node.ParentNode->GetActorName(), SketchupPersistentID);
 
 	FString EntityName = Node.Entity.GetName();
-	Node.DatasmithActorLabel = FDatasmithUtils::SanitizeObjectName(EntityName.IsEmpty() ? GetSketchupSourceName() : EntityName);
+	Node.DatasmithActorLabel = FDatasmithUtils::SanitizeObjectName(EntityName.IsEmpty() ? GetDefinition()->GetSketchupSourceName() : EntityName);
 }
 
 void FComponentDefinition::InvalidateInstancesGeometry(FExportContext& Context)
@@ -798,7 +751,7 @@ void FEntityWithEntities::UpdateOccurrence(FExportContext& Context, FNodeOccuren
 	{
 		const TSharedPtr<IDatasmithMeshActorElement>& MeshActor = Node.MeshActors[MeshIndex];
 		MeshActor->SetLabel(*MeshActorLabel);
-		MeshActor->SetLayer(*EffectiveLayerName);
+		MeshActor->SetLayer(*FDatasmithUtils::SanitizeObjectName(EffectiveLayerName));
 
 		// Update Override(Inherited)  Material
 		// todo: set inherited material only on mesh actors that have faces with default material, right now setting on every mesh, hot harmful but excessive
@@ -920,20 +873,15 @@ bool FComponentInstance::GetAssignedMaterial(FMaterialIDType& MaterialId)
 
 void FComponentInstance::UpdateOccurrence(FExportContext& Context, FNodeOccurence& Node)
 {
-	if (!Node.bVisible)
-	{
-		return;
-	}
-
 	Node.EffectiveLayerRef = DatasmithSketchUpUtils::GetEffectiveLayer(GetComponentInstanceRef(), Node.ParentNode->EffectiveLayerRef);
 
 	if (FDefinition* EntityDefinition = GetDefinition())
 	{
-		EntityDefinition->BuildNodeNames(Node);
+		BuildNodeNames(Node);
 	}
 
 	// Set the actor label used in the Unreal UI.
-	Node.DatasmithActorElement->SetLabel(*Node.DatasmithActorLabel);
+	Node.DatasmithActorElement->SetLabel(*Node.GetActorLabel());
 
 	// Retrieve the SketchUp component instance effective layer name.
 	FString SEffectiveLayerName;
@@ -997,6 +945,46 @@ void FComponentInstance::UpdateEntityProperties(FExportContext& Context)
 	FEntity::UpdateEntityProperties(Context);
 }
 
+void FComponentInstance::UpdateOccurrenceMeshActors(FExportContext& Context, FNodeOccurence& Node)
+{
+	BuildNodeNames(Node);
+
+	FEntityWithEntities::UpdateOccurrenceMeshActors(Context, Node);
+
+	if (Node.Children.IsEmpty() && !Node.MeshActors.IsEmpty()) // Don't make extra actor when geometry node has no children 
+	{
+		TSharedPtr<IDatasmithMeshActorElement> MeshActor = Node.MeshActors[0];
+
+		Node.DatasmithActorElement = MeshActor;
+
+		SetupActor(Context, Node);
+	}
+	else
+	{
+		Node.DatasmithActorElement = FDatasmithSceneFactory::CreateActor(*Node.GetActorName());
+		SetupActor(Context, Node);
+
+		for (TSharedPtr<IDatasmithMeshActorElement> MeshActor : Node.MeshActors)
+		{
+			Node.DatasmithActorElement->AddChild(MeshActor);
+		}
+
+		for (FNodeOccurence* Child: Node.Children)
+		{
+			if (Child->DatasmithActorElement)
+			{
+				Node.DatasmithActorElement->AddChild(Child->DatasmithActorElement);
+			}
+		}
+	}
+
+}
+
+void FComponentInstance::ResetOccurrenceActors(FExportContext& Context, FNodeOccurence& Node)
+{
+	Node.ResetNodeActors(Context);
+}
+
 void FComponentInstance::ParseNode(FExportContext& Context, FNodeOccurence& Node)
 {
 	FDefinition* EntityDefinition = GetDefinition();
@@ -1119,14 +1107,7 @@ void FComponentInstance::UpdateOccurrenceVisibility(FExportContext& Context, FNo
 
 	if (Node.bVisible)
 	{
-		if (!Node.DatasmithActorElement)
-		{
-			FDefinition* EntityDefinition = GetDefinition();
-			EntityDefinition->CreateActor(Context, Node);
-
-			// Invalidate actor-dependent 
-			Node.InvalidateProperties();
-		}
+		Node.InvalidateProperties();
 		Node.InvalidateMeshActors();
 	}
 	else
@@ -1205,6 +1186,25 @@ void FModel::UpdateMetadata(FExportContext& Context)
 {
 }
 
+void FModel::UpdateOccurrenceMeshActors(FExportContext& Context, FNodeOccurence& Node)
+{
+	FEntityWithEntities::UpdateOccurrenceMeshActors(Context, Node);
+
+	for (TSharedPtr<IDatasmithMeshActorElement> MeshActor : Node.MeshActors)
+	{
+		Context.DatasmithScene->AddActor(MeshActor);
+	}
+}
+
+void FModel::ResetOccurrenceActors(FExportContext& Context, FNodeOccurence& Node)
+{
+	// Model actors are MeshActors added to DatasmithScene root
+	for (TSharedPtr<IDatasmithMeshActorElement> MeshActor : Node.MeshActors)
+	{
+		Context.DatasmithScene->RemoveActor(MeshActor, EDatasmithActorRemovalRule::RemoveChildren);
+	}
+}
+
 FDefinition::FDefinition(): bMeshesAdded(false)
                             , bGeometryInvalidated(true)
                             , bPropertiesInvalidated(true)
diff --git a/Engine/Source/Programs/Enterprise/Datasmith/DatasmithSketchUpRubyExporter/Private/DatasmithSketchUpComponent.h b/Engine/Source/Programs/Enterprise/Datasmith/DatasmithSketchUpRubyExporter/Private/DatasmithSketchUpComponent.h
--- a/Engine/Source/Programs/Enterprise/Datasmith/DatasmithSketchUpRubyExporter/Private/DatasmithSketchUpComponent.h
+++ b/Engine/Source/Programs/Enterprise/Datasmith/DatasmithSketchUpRubyExporter/Private/DatasmithSketchUpComponent.h
@@ -60,13 +60,13 @@ namespace DatasmithSketchUp
 		void UpdateVisibility(FExportContext& Context);
 		void Update(FExportContext& Context);
 
-		void UpdateMeshActors(FExportContext& Context); // Create/Free Datasmith Mesh Actors for meshes in the node's Entities
+		void ResetNodeActors(FExportContext& Context); // Reset actors before update
+		void RemoveDatasmithActorHierarchy(FExportContext& Context);  // Clean whole datasmith hierarchy from Datasmith Scene (e.g. when made invisible)
 
 		void InvalidateProperties(); // Invalidate name and transform. Invalidate propagates down the hierarchy - child transforms depend on the parent
 		void InvalidateMeshActors();
 		void SetVisibility(bool);
 
-		void RemoveDatasmithActorHierarchy(FExportContext& Context);
 
 		FString GetActorName();
 
@@ -117,10 +117,8 @@ namespace DatasmithSketchUp
 		virtual ~FDefinition();
 
 		virtual void Parse(FExportContext& Context) = 0;
-		virtual void CreateActor(FExportContext& Context, FNodeOccurence& Node) = 0; // Create Datasmith actor for node occurrence
 		virtual void UpdateGeometry(FExportContext& Context) = 0; // Convert definition's Entities geometry to Datasmith Mesh
 		virtual void UpdateMetadata(FExportContext& Context) = 0; 
-		virtual void BuildNodeNames(FNodeOccurence& Node) = 0;
 
 		void EntityVisible(FEntity* Entity, bool bVisible);
 
@@ -175,8 +173,6 @@ namespace DatasmithSketchUp
 
 		// Begin FDefinition
 		virtual void Parse(FExportContext& Context) override;
-		virtual void CreateActor(FExportContext& Context, FNodeOccurence& Node) override;
-		virtual void BuildNodeNames(FNodeOccurence& Node) override;
 		virtual void UpdateGeometry(FExportContext& Context) override;
 		virtual void UpdateMetadata(FExportContext& Context) override;
 
@@ -194,8 +190,6 @@ namespace DatasmithSketchUp
 		virtual SUTransformation GetMeshBakedTransform() override;
 		// End FDefinition
 
-
-
 		// Register/unregister instanced of this definition
 		void LinkComponentInstance(FComponentInstance* ComponentInstance);
 		void UnlinkComponentInstance(FComponentInstance* ComponentInstance);
@@ -209,14 +203,14 @@ namespace DatasmithSketchUp
 		FComponentDefinitionIDType SketchupSourceID;
 		TSet<FComponentInstance*> Instances; // Tracked instances of this ComponentDefinition
 
+		// Whether or not the source SketchUp component behaves like a billboard, always presenting a 2D surface perpendicular to the direction of camera.
+		bool bSketchupSourceFaceCamera = false;
+
 	private:
 		SUComponentDefinitionRef ComponentDefinitionRef;
 
 		TUniquePtr<FMetadata> ParsedMetadata; // Shared metadata parsed from source SU component to be added to each occurrence actor's datasmith metatada element
 
-		// Whether or not the source SketchUp component behaves like a billboard, always presenting a 2D surface perpendicular to the direction of camera.
-		bool bSketchupSourceFaceCamera = false;
-
 		bool bBakeTransformIntoMesh = false;
 	};
 
@@ -228,8 +222,7 @@ namespace DatasmithSketchUp
 
 		// Being FDefinition
 		virtual void Parse(FExportContext& Context) override;
-		virtual void CreateActor(FExportContext& Context, FNodeOccurence& Node) override;
-		virtual void BuildNodeNames(FNodeOccurence& Node) override;
+		
 		virtual void UpdateGeometry(FExportContext& Context) override;
 		virtual void UpdateMetadata(FExportContext& Context) override;
 
@@ -329,9 +322,9 @@ namespace DatasmithSketchUp
 		virtual void ApplyOverrideMaterialToNode(FNodeOccurence& Node, FMaterialOccurrence& Material) = 0;
 
 		virtual void UpdateOccurrence(FExportContext& Context, FNodeOccurence& Node) = 0; // Update occurrence of this entity
-		// Create an occurrence of this ComponentInstance (component instance can appear multiple times in SketchUp hierarchy)
-		virtual void UpdateOccurrenceVisibility(FExportContext& Context, FNodeOccurence&) = 0;
-		virtual bool UpdateOccurrenceMeshActors(FExportContext& Context, FNodeOccurence& Node) = 0;
+		virtual void UpdateOccurrenceVisibility(FExportContext& Context, FNodeOccurence&) = 0;  // Re-evaluate visibility of entity's occurrence
+		virtual void UpdateOccurrenceMeshActors(FExportContext& Context, FNodeOccurence& Node) = 0; // Rebuild datasmith actors of entity's occurrence
+		virtual void ResetOccurrenceActors(FExportContext& Context, FNodeOccurence& Node) = 0; // Remove datasmith actors of entity's occurrence from datasmith scene
 
 		virtual void InvalidateOccurrencesGeometry(FExportContext& Context) = 0;
 		virtual void InvalidateOccurrencesProperties(FExportContext& Context) = 0;
@@ -388,7 +381,7 @@ namespace DatasmithSketchUp
 
 		virtual void UpdateOccurrence(FExportContext& Context, FNodeOccurence& Node); // Update occurrence of this entity
 		
-		virtual bool UpdateOccurrenceMeshActors(FExportContext& Context, FNodeOccurence& Node);
+		virtual void UpdateOccurrenceMeshActors(FExportContext& Context, FNodeOccurence& Node);
 
 		virtual void EntityOccurrenceVisible(FNodeOccurence* Node, bool bUses) override;
 
@@ -414,15 +407,22 @@ namespace DatasmithSketchUp
 		virtual bool GetAssignedMaterial(FMaterialIDType& MaterialId) override;
 		virtual void InvalidateOccurrencesGeometry(FExportContext& Context) override;
 		virtual void UpdateOccurrence(FExportContext& Context, FNodeOccurence& Node) override;
+
 		virtual void InvalidateOccurrencesProperties(FExportContext& Context) override;
 		virtual int64 GetPersistentId() override;
 		virtual FString GetName() override;
 		virtual void UpdateOccurrenceVisibility(FExportContext& Context, FNodeOccurence&) override;
 		virtual void UpdateMetadata(FExportContext& Context) override;
 		virtual void UpdateEntityProperties(FExportContext& Context) override;
+		virtual void UpdateOccurrenceMeshActors(FExportContext& Context, FNodeOccurence& Node) override;
+		virtual void ResetOccurrenceActors(FExportContext& Context, FNodeOccurence& Node) override;
+
 		void ApplyOverrideMaterialToNode(FNodeOccurence& Node, FMaterialOccurrence& Material);
 		// <<< FEntity
 
+		void BuildNodeNames(FNodeOccurence& Node);
+		void SetupActor(FExportContext& Context, FNodeOccurence& Node);
+
 		void ParseNode(FExportContext& Context, FNodeOccurence&);
 		void RemoveComponentInstance(FExportContext& Context);
 
@@ -455,7 +455,8 @@ namespace DatasmithSketchUp
 		virtual void UpdateOccurrence(FExportContext& Context, FNodeOccurence& Node) override;
 		void RemoveImageFromDatasmithScene(FExportContext& Context);
 		virtual void UpdateOccurrenceVisibility(FExportContext& Context, FNodeOccurence&) override;
-		virtual bool UpdateOccurrenceMeshActors(FExportContext& Context, FNodeOccurence& Node) override;
+		virtual void UpdateOccurrenceMeshActors(FExportContext& Context, FNodeOccurence& Node) override;
+		virtual void ResetOccurrenceActors(FExportContext& Context, FNodeOccurence& Node) override;
 		virtual void InvalidateOccurrencesGeometry(FExportContext& Context) override;
 		virtual void InvalidateOccurrencesProperties(FExportContext& Context) override;
 		virtual void UpdateMetadata(FExportContext& Context) override;
@@ -474,7 +475,7 @@ namespace DatasmithSketchUp
 		FString MeshElementName;
 
 		void BuildNodeNames(FNodeOccurence& Node);
-		void CreateActor(FExportContext& Context, FNodeOccurence& Node);
+		void SetupActor(FExportContext& Context, FNodeOccurence& Node);
 
 		void RemoveImage(FExportContext& Context);
 
@@ -498,20 +499,21 @@ namespace DatasmithSketchUp
 
 		// >>> FEntity
 		virtual FDefinition* GetDefinition() override;
-		bool GetAssignedMaterial(FMaterialIDType& MaterialId) override;
-		void InvalidateOccurrencesGeometry(FExportContext& Context) override;
+		virtual bool GetAssignedMaterial(FMaterialIDType& MaterialId) override;
+		virtual void InvalidateOccurrencesGeometry(FExportContext& Context) override;
 		// void UpdateOccurrence(FExportContext& Context, FNodeOccurence& Node) override;
-		void InvalidateOccurrencesProperties(FExportContext& Context) override;
-		int64 GetPersistentId() override;
-		FString GetName() override;
-		void UpdateOccurrenceVisibility(FExportContext& Context, FNodeOccurence&) override;
-		void UpdateMetadata(FExportContext& Context) override;
-		void ApplyOverrideMaterialToNode(FNodeOccurence& Node, FMaterialOccurrence& Material);
+		virtual void InvalidateOccurrencesProperties(FExportContext& Context) override;
+		virtual int64 GetPersistentId() override;
+		virtual FString GetName() override;
+		virtual void UpdateOccurrenceVisibility(FExportContext& Context, FNodeOccurence&) override;
 
+		virtual void UpdateMetadata(FExportContext& Context) override;
+		virtual void ApplyOverrideMaterialToNode(FNodeOccurence& Node, FMaterialOccurrence& Material) override;
+		virtual void UpdateOccurrenceMeshActors(FExportContext& Context, FNodeOccurence& Node) override;
+		virtual void ResetOccurrenceActors(FExportContext& Context, FNodeOccurence& Node) override;
 		// <<< FEntity
 
 	private:
 		FModelDefinition& Definition;
 	};
-
 }
diff --git a/Engine/Source/Programs/Enterprise/Datasmith/DatasmithSketchUpRubyExporter/Private/DatasmithSketchUpImage.cpp b/Engine/Source/Programs/Enterprise/Datasmith/DatasmithSketchUpRubyExporter/Private/DatasmithSketchUpImage.cpp
--- a/Engine/Source/Programs/Enterprise/Datasmith/DatasmithSketchUpRubyExporter/Private/DatasmithSketchUpImage.cpp
+++ b/Engine/Source/Programs/Enterprise/Datasmith/DatasmithSketchUpRubyExporter/Private/DatasmithSketchUpImage.cpp
@@ -96,18 +96,13 @@ void FImage::ApplyOverrideMaterialToNode(FNodeOccurence& Node, FMaterialOccurren
 
 void FImage::UpdateOccurrence(FExportContext& Context, FNodeOccurence& Node)
 {
-	if (!Node.bVisible)
-	{
-		return;
-	}
-	
 	Node.EffectiveLayerRef = DatasmithSketchUpUtils::GetEffectiveLayer( SUImageToDrawingElement(SUImageFromEntity(EntityRef)), Node.ParentNode->EffectiveLayerRef);
 
 	BuildNodeNames(Node);
 
 	FString EffectiveLayerName = FDatasmithUtils::SanitizeObjectName(SuGetString(SULayerGetName, Node.EffectiveLayerRef));
 	Node.DatasmithActorElement->SetLayer(*EffectiveLayerName);
-	Node.DatasmithActorElement->SetLabel(*Node.DatasmithActorLabel);
+	Node.DatasmithActorElement->SetLabel(*Node.GetActorLabel());
 
 	// Compute the world transform of the SketchUp component instance.
 	SUTransformation LocalTransform;
@@ -220,23 +215,16 @@ void FImage::BuildNodeNames(FNodeOccurence& Node)
 	Node.DatasmithActorLabel = FDatasmithUtils::SanitizeObjectName(EntityName.IsEmpty() ? GetName() : EntityName);
 }
 
-void FImage::CreateActor(FExportContext& Context, FNodeOccurence& Node)
+void FImage::SetupActor(FExportContext& Context, FNodeOccurence& Node)
 {
-	BuildNodeNames(Node);
+	// Add the Datasmith actor component depth tag.
+	// We use component depth + 1 to factor in the added Datasmith scene root once imported in Unreal.
+	FString ComponentDepthTag = FString::Printf(TEXT("SU.DEPTH.%d"), Node.Depth);
+	Node.DatasmithActorElement->AddTag(*ComponentDepthTag);
 
-	// Create a Datasmith actor for the component instance.
-	Node.DatasmithActorElement = FDatasmithSceneFactory::CreateActor(*Node.DatasmithActorName); // a EDatasmithElementType::Actor
-
-	{
-		// Add the Datasmith actor component depth tag.
-		// We use component depth + 1 to factor in the added Datasmith scene root once imported in Unreal.
-		FString ComponentDepthTag = FString::Printf(TEXT("SU.DEPTH.%d"), Node.Depth);
-		Node.DatasmithActorElement->AddTag(*ComponentDepthTag);
-
-		// Add the Datasmith actor component instance path tag.
-		FString InstancePathTag = Node.DatasmithActorName.Replace(TEXT("SU"), TEXT("SU.PATH.0")).Replace(TEXT("_"), TEXT("."));
-		Node.DatasmithActorElement->AddTag(*InstancePathTag);
-	}
+	// Add the Datasmith actor component instance path tag.
+	FString InstancePathTag = Node.GetActorName().Replace(TEXT("SU"), TEXT("SU.PATH.0")).Replace(TEXT("_"), TEXT("."));
+	Node.DatasmithActorElement->AddTag(*InstancePathTag);
 
 	if (Node.ParentNode->DatasmithActorElement)
 	{
@@ -258,13 +246,7 @@ void FImage::UpdateOccurrenceVisibility(FExportContext& Context, FNodeOccurence&
 
 	if (Node.bVisible)
 	{
-		if (!Node.DatasmithActorElement)
-		{
-			CreateActor(Context, Node);
-
-			// Invalidate actor-dependent 
-			Node.InvalidateProperties();
-		}
+		Node.InvalidateProperties();
 		Node.InvalidateMeshActors();
 	}
 	else
@@ -278,16 +260,15 @@ void FImage::UpdateOccurrenceVisibility(FExportContext& Context, FNodeOccurence&
 	}
 }
 
-bool FImage::UpdateOccurrenceMeshActors(FExportContext& Context, FNodeOccurence& Node)
+void FImage::UpdateOccurrenceMeshActors(FExportContext& Context, FNodeOccurence& Node)
 {
-	Node.MeshActors.Reset(1);
+	BuildNodeNames(Node);
 
 	FString ComponentActorName = Node.GetActorName();
 	
 	{
-		FString MeshActorName = FString::Printf(TEXT("%ls_%d"), *ComponentActorName); // Count meshes/mesh actors from 1
+		FString MeshActorName = ComponentActorName;
 
-		// Create a Datasmith mesh actor for the Datasmith mesh element.
 		TSharedPtr<IDatasmithMeshActorElement> DMeshActorPtr = FDatasmithSceneFactory::CreateMeshActor(*MeshActorName);
 
 		Node.MeshActors.Add(DMeshActorPtr);
@@ -301,22 +282,20 @@ bool FImage::UpdateOccurrenceMeshActors(FExportContext& Context, FNodeOccurence&
 		FString InstancePathTag = ComponentActorName.Replace(TEXT("SU"), TEXT("SU.PATH.0")).Replace(TEXT("_"), TEXT("."));
 		DMeshActorPtr->AddTag(*InstancePathTag);
 
-		// Add the mesh actor to our component Datasmith actor hierarchy.
-		if (Node.DatasmithActorElement.IsValid()) 
-		{
-			Node.DatasmithActorElement->AddChild(DMeshActorPtr);
-		}
-		else
-		{
-			Context.DatasmithScene->AddActor(DMeshActorPtr);
-		}
-
 		// ADD_TRACE_LINE(TEXT("Actor %ls: %ls %ls %ls"), *MeshActorLabel, *ComponentDepthTag, *DefinitionGUIDTag, *InstancePathTag);
 
 		// Set the Datasmith mesh element used by the mesh actor.
 		DMeshActorPtr->SetStaticMeshPathName(GetMeshElementName());
 	}
-	return false;
+
+	Node.DatasmithActorElement = Node.MeshActors[0];
+
+	SetupActor(Context, Node);
+}
+
+void FImage::ResetOccurrenceActors(FExportContext& Context, FNodeOccurence& Node)
+{
+	Node.ResetNodeActors(Context);
 }
 
 void FImage::InvalidateOccurrencesGeometry(FExportContext& Context)
