[3ds Max] Keep Aspect Ratio for baked procedural textures

- also added tooltips for size dropdown items

diff --git a/Engine/Source/Programs/Enterprise/Datasmith/DatasmithMaxExporter/Extras/3dsMaxScripts/UI/DatasmithUI.xaml b/Engine/Source/Programs/Enterprise/Datasmith/DatasmithMaxExporter/Extras/3dsMaxScripts/UI/DatasmithUI.xaml
--- a/Engine/Source/Programs/Enterprise/Datasmith/DatasmithMaxExporter/Extras/3dsMaxScripts/UI/DatasmithUI.xaml
+++ b/Engine/Source/Programs/Enterprise/Datasmith/DatasmithMaxExporter/Extras/3dsMaxScripts/UI/DatasmithUI.xaml
@@ -283,93 +283,100 @@
             <ads:RibbonCommandItem
               AllowInToolBar="false"
               CommandHandler="{wpfcust:ActionItem MacroName=Datasmith_ExportOptions_TextureResolution64, MacroCategory=Datasmith}"
+              Description="Baked procedural texmaps would not exceed 4096 pixels (e.g. 64x64, 128x32)"
               Image="pack://application:,,,/AdWindows;component/Resources/Images/MissingImage.ico"
               IsCheckable="true"
               IsToolTipEnabled="true"
               LargeImage="pack://application:,,,/AdWindows;component/Resources/Images/MissingImage.ico"
               MinWidth="100"
-              Name="64"
+              Name="4K pixels"
               ShowImage="false"
               ShowText="false"
-              Text="64"
+              Text="4K pixels"
               UID="28d398da-505c-4f91-9cf3-14316cd2893b" />
             <ads:RibbonCommandItem
               AllowInToolBar="false"
               CommandHandler="{wpfcust:ActionItem MacroName=Datasmith_ExportOptions_TextureResolution128, MacroCategory=Datasmith}"
+              Description="Baked procedural texmaps would not exceed 16384 pixels (e.g. 128x128, 256x64)"
               Image="pack://application:,,,/AdWindows;component/Resources/Images/MissingImage.ico"
               IsCheckable="true"
               IsToolTipEnabled="true"
               LargeImage="pack://application:,,,/AdWindows;component/Resources/Images/MissingImage.ico"
               MinWidth="100"
-              Name="128"
+              Name="16K pixels"
               ShowImage="false"
               ShowText="false"
-              Text="128"
+              Text="16K pixels"
               UID="89a922df-7286-40b9-833e-c51419fe3c1d" />
             <ads:RibbonCommandItem
               AllowInToolBar="false"
               CommandHandler="{wpfcust:ActionItem MacroName=Datasmith_ExportOptions_TextureResolution256, MacroCategory=Datasmith}"
+              Description="Baked procedural texmaps would not exceed 65536 pixels (e.g. 256x256, 512x128)"
               Image="pack://application:,,,/AdWindows;component/Resources/Images/MissingImage.ico"
               IsCheckable="true"
               IsToolTipEnabled="true"
               LargeImage="pack://application:,,,/AdWindows;component/Resources/Images/MissingImage.ico"
               MinWidth="100"
-              Name="256"
+              Name="64K pixels"
               ShowImage="false"
               ShowText="false"
-              Text="256"
+              Text="64K pixels"
               UID="177f6116-6ee0-4b6f-8d06-72d73bcc89e4" />
             <ads:RibbonCommandItem
               AllowInToolBar="false"
               CommandHandler="{wpfcust:ActionItem MacroName=Datasmith_ExportOptions_TextureResolution512, MacroCategory=Datasmith}"
+              Description="Baked procedural texmaps would not exceed 0.25 Mega pixels (e.g. 512x512, 1024x256)"
               Image="pack://application:,,,/AdWindows;component/Resources/Images/MissingImage.ico"
               IsCheckable="true"
               IsToolTipEnabled="true"
               LargeImage="pack://application:,,,/AdWindows;component/Resources/Images/MissingImage.ico"
               MinWidth="100"
-              Name="512"
+              Name="256K pixels"
               ShowImage="false"
               ShowText="false"
-              Text="512"
+              Text="256K pixels"
               UID="a5dd7f25-278a-4f8e-a9c2-7d73164df5f0" />
             <ads:RibbonCommandItem
               AllowInToolBar="false"
               CommandHandler="{wpfcust:ActionItem MacroName=Datasmith_ExportOptions_TextureResolution1024, MacroCategory=Datasmith}"
+              Description="Baked procedural texmaps would not exceed 1 Mega pixel (e.g. 1024x1024, 2048x512)"
               Image="pack://application:,,,/AdWindows;component/Resources/Images/MissingImage.ico"
               IsCheckable="true"
               IsToolTipEnabled="true"
               LargeImage="pack://application:,,,/AdWindows;component/Resources/Images/MissingImage.ico"
               MinWidth="100"
-              Name="1024"
+              Name="1M pixels"
               ShowImage="false"
               ShowText="false"
-              Text="1024"
+              Text="1M pixels"
               UID="09fec31c-f9e9-4e55-b33c-f4fb2218a4a0" />
             <ads:RibbonCommandItem
               AllowInToolBar="false"
               CommandHandler="{wpfcust:ActionItem MacroName=Datasmith_ExportOptions_TextureResolution2048, MacroCategory=Datasmith}"
+              Description="Baked procedural texmaps would not exceed 4 Mega pixels (e.g. 2048x2048, 4096x1024)"
               Image="pack://application:,,,/AdWindows;component/Resources/Images/MissingImage.ico"
               IsCheckable="true"
               IsToolTipEnabled="true"
               LargeImage="pack://application:,,,/AdWindows;component/Resources/Images/MissingImage.ico"
               MinWidth="100"
-              Name="2048"
+              Name="4M pixels"
               ShowImage="false"
               ShowText="false"
-              Text="2048"
+              Text="4M pixels"
               UID="ef8585b8-1012-47a7-9af0-46371cf70df2" />
             <ads:RibbonCommandItem
               AllowInToolBar="false"
               CommandHandler="{wpfcust:ActionItem MacroName=Datasmith_ExportOptions_TextureResolution4096, MacroCategory=Datasmith}"
+              Description="Baked procedural texmaps would not exceed 16 Mega pixels (e.g. 4096x4096, 8192x2048)"
               Image="pack://application:,,,/AdWindows;component/Resources/Images/MissingImage.ico"
               IsCheckable="true"
               IsToolTipEnabled="true"
               LargeImage="pack://application:,,,/AdWindows;component/Resources/Images/MissingImage.ico"
               MinWidth="100"
-              Name="4096"
+              Name="16M pixels"
               ShowImage="false"
               ShowText="false"
-              Text="4096"
+              Text="16M pixels"
               UID="c5c4c61f-7d64-4267-a387-cd6df2704161" />
           </wpfcust:MaxRibbonSplitButton>
           <ads:RibbonRowBreak
diff --git a/Engine/Source/Programs/Enterprise/Datasmith/DatasmithMaxExporter/Private/DatasmithMaxScanline.cpp b/Engine/Source/Programs/Enterprise/Datasmith/DatasmithMaxExporter/Private/DatasmithMaxScanline.cpp
--- a/Engine/Source/Programs/Enterprise/Datasmith/DatasmithMaxExporter/Private/DatasmithMaxScanline.cpp
+++ b/Engine/Source/Programs/Enterprise/Datasmith/DatasmithMaxExporter/Private/DatasmithMaxScanline.cpp
@@ -1227,24 +1227,62 @@ FString FDatasmithMaxMatWriter::DumpFalloff(TSharedRef< IDatasmithScene > Datasm
 
 void GetBakeableMaximumSize(Texmap* InTexmap, int &Width, int &Height)
 {
-	if (InTexmap == NULL)
-	{
-		return;
-	}	
-	if (InTexmap->ClassID() == RBITMAPCLASS)
+	TArray<Texmap*> Texmaps;
+	Texmaps.Add(InTexmap);
+
+	// Limit baked texture size by subbitmap with largest area(number of pixels)
+	int32 BiggestSubBitmapArea = 0;  // Biggest pixel count
+	int32 BiggestSubBitmapWidth = 0;
+	int32 BiggestSubBitmapHeight = 0;
+
+	while (!Texmaps.IsEmpty())
 	{
-		BitmapTex* BitmapTexture = (BitmapTex*)InTexmap;
-		Bitmap* ActualBitmap = ((BitmapTex*)InTexmap)->GetBitmap(GetCOREInterface()->GetTime());
-		if (ActualBitmap != NULL)
+		Texmap* Texmap = Texmaps.Pop(false);
+
+		if (!Texmap)
+		{
+			continue;;
+		}
+
+		if (Texmap->ClassID() == RBITMAPCLASS)
+		{
+			BitmapTex* BitmapTexture = (BitmapTex*)Texmap;
+			Bitmap* ActualBitmap = ((BitmapTex*)Texmap)->GetBitmap(GetCOREInterface()->GetTime());
+			if (ActualBitmap != NULL)
+			{
+				int32 BitmapWidth = ActualBitmap->Width();
+				int32 BitmapHeight = ActualBitmap->Height();
+
+				int32 Area = BitmapWidth*BitmapHeight;
+				if (Area > BiggestSubBitmapArea)
+				{
+					BiggestSubBitmapArea = Area;
+					BiggestSubBitmapWidth = BitmapWidth;
+					BiggestSubBitmapHeight = BitmapHeight;
+				}
+
+			}
+		}
+
+		for (int SubTexmap = 0; SubTexmap < Texmap->NumSubTexmaps(); SubTexmap++)
 		{
-			Width = FMath::Min(Width, ActualBitmap->Width());
-			Height = FMath::Min(Height, ActualBitmap->Height());
+			Texmaps.Add(Texmap->GetSubTexmap(SubTexmap));
 		}
 	}
 
-	for (int SubTexmap = 0; SubTexmap < InTexmap->NumSubTexmaps(); SubTexmap++)
+	if (BiggestSubBitmapArea > 0)
 	{
-		GetBakeableMaximumSize(InTexmap->GetSubTexmap(SubTexmap), Width, Height);
+		int32 AreaLimit = Width*Height;
+
+		Width = BiggestSubBitmapWidth;
+		Height = BiggestSubBitmapHeight;
+		int32 Area = BiggestSubBitmapArea;
+		while (Area > AreaLimit)
+		{
+			Width = Width / 2;
+			Height = Height / 2;
+			Area = Width * Height;
+		}
 	}
 }
 
