[SketchUp] Camera updates and active flag

diff --git a/Engine/Source/Programs/Enterprise/Datasmith/DatasmithSketchUpRubyExporter/Private/DatasmithSketchUpCamera.cpp b/Engine/Source/Programs/Enterprise/Datasmith/DatasmithSketchUpRubyExporter/Private/DatasmithSketchUpCamera.cpp
--- a/Engine/Source/Programs/Enterprise/Datasmith/DatasmithSketchUpRubyExporter/Private/DatasmithSketchUpCamera.cpp
+++ b/Engine/Source/Programs/Enterprise/Datasmith/DatasmithSketchUpRubyExporter/Private/DatasmithSketchUpCamera.cpp
@@ -24,10 +24,8 @@ using namespace DatasmithSketchUp;
 TSharedPtr<FCamera> FCamera::Create(FExportContext& Context, SUCameraRef InCameraRef, const FString& Name)
 {
 	TSharedPtr<FCamera> Camera = MakeShared<FCamera>(InCameraRef);
-	Camera->DatasmithCamera = FDatasmithSceneFactory::CreateCameraActor(TEXT(""));
-	Context.DatasmithScene->AddActor(Camera->DatasmithCamera);
+
 	Camera->Name = Name;
-	Camera->Update(Context);
 	return Camera;
 }
 
@@ -42,6 +40,12 @@ TSharedPtr<FCamera> FCamera::Create(FExportContext& Context, SUSceneRef InSceneR
 
 void FCamera::Update(FExportContext& Context)
 {
+	if (!DatasmithCamera)
+	{
+		DatasmithCamera = FDatasmithSceneFactory::CreateCameraActor(TEXT(""));
+		Context.DatasmithScene->AddActor(DatasmithCamera);
+	}
+
 	SUPoint3D SourcePosition;
 	SUPoint3D SourceTarget;
 	SUVector3D SourceUpVector;
@@ -129,4 +133,77 @@ void FCamera::Update(FExportContext& Context)
 	FVector DistanceVector = FVector(DatasmithSketchUpUtils::FromSketchUp::ConvertPosition(SourceTarget.x - SourcePosition.x, SourceTarget.y - SourcePosition.y, SourceTarget.z - SourcePosition.z));
 	DatasmithCamera->SetFocusDistance(DistanceVector.Size());
 
+	// Using unused 'Visibility' flag for Camera to indicate currently active camera
+	// todo: replace someday with dedicated flag
+	DatasmithCamera->SetVisibility(bIsActive);
+}
+
+FMD5Hash FCamera::GetHash()
+{
+	FMD5 MD5;
+
+	MD5.Update(reinterpret_cast<const uint8*>(*Name), Name.Len() * sizeof(TCHAR));
+
+	SUPoint3D SourcePosition;
+	SUPoint3D SourceTarget;
+	SUVector3D SourceUpVector;
+
+	// Retrieve the SketckUp camera orientation.
+	SUCameraGetOrientation(CameraRef, &SourcePosition, &SourceTarget, &SourceUpVector);
+
+	MD5.Update(reinterpret_cast<const uint8*>(&SourcePosition), sizeof(SourcePosition));
+	MD5.Update(reinterpret_cast<const uint8*>(&SourceTarget), sizeof(SourceTarget));
+	MD5.Update(reinterpret_cast<const uint8*>(&SourceUpVector), sizeof(SourceUpVector));
+
+	// Get the SketchUp camera aspect ratio.
+	double CameraAspectRatio = 0.0;
+	SUResult Result = SUCameraGetAspectRatio(CameraRef, &CameraAspectRatio);
+	double SourceAspectRatio = 16.0 / 9.0;
+
+	// Keep the default aspect ratio when the camera uses the screen aspect ratio (SU_ERROR_NO_DATA).
+	if (Result == SU_ERROR_NONE)
+	{
+		SourceAspectRatio = CameraAspectRatio;
+	}
+	MD5.Update(reinterpret_cast<const uint8*>(&SourceAspectRatio), sizeof(SourceAspectRatio));
+
+	// Get the flag indicating whether or not the SketchUp scene camera is a perspective camera.
+	bool bCameraIsPerspective = false;
+	SUCameraGetPerspective(CameraRef, &bCameraIsPerspective); // we can ignore the returned SU_RESULT
+	MD5.Update(reinterpret_cast<const uint8*>(&bCameraIsPerspective), sizeof(bCameraIsPerspective));
+
+	// Get the flag indicating whether or not the SketchUp scene camera is a two dimensional camera.
+	bool bCameraIs2D = false;
+	SUCameraGet2D(CameraRef, &bCameraIs2D); // we can ignore the returned SU_RESULT
+	MD5.Update(reinterpret_cast<const uint8*>(&bCameraIs2D), sizeof(bCameraIs2D));
+
+	bool bSourceFOVForHeight = true;
+	double SourceFOV(60.0);             // default vertical field of view of 60 degrees
+	double SourceImageWidth(36.0);      // default image width of 36 mm (from Datasmith)
+
+	if (bCameraIsPerspective && !bCameraIs2D)
+	{
+		// Get the flag indicating whether or not the SketchUp camera field of view value represents the camera view height.
+		SUCameraGetFOVIsHeight(CameraRef, &bSourceFOVForHeight); // we can ignore the returned SU_RESULT
+		MD5.Update(reinterpret_cast<const uint8*>(&bSourceFOVForHeight), sizeof(bSourceFOVForHeight));
+
+		// Get the SketchUp camera field of view (in degrees).
+		SUCameraGetPerspectiveFrustumFOV(CameraRef, &SourceFOV); // we can ignore the returned SU_RESULT
+		MD5.Update(reinterpret_cast<const uint8*>(&SourceFOV), sizeof(SourceFOV));
+
+		// Get the SketchUp camera image width (in millimeters).
+		double CameraImageWidth = 0.0;
+		Result = SUCameraGetImageWidth(CameraRef, &CameraImageWidth);
+		// Keep the default image width when the camera does not have an image width.
+		if (CameraImageWidth > 0.0)
+		{
+			SourceImageWidth = CameraImageWidth;
+		}
+		MD5.Update(reinterpret_cast<const uint8*>(&SourceImageWidth), sizeof(SourceImageWidth));
+	}
+
+
+	FMD5Hash Hash;
+	Hash.Set(MD5);
+	return Hash;
 }
diff --git a/Engine/Source/Programs/Enterprise/Datasmith/DatasmithSketchUpRubyExporter/Private/DatasmithSketchUpCamera.h b/Engine/Source/Programs/Enterprise/Datasmith/DatasmithSketchUpRubyExporter/Private/DatasmithSketchUpCamera.h
--- a/Engine/Source/Programs/Enterprise/Datasmith/DatasmithSketchUpRubyExporter/Private/DatasmithSketchUpCamera.h
+++ b/Engine/Source/Programs/Enterprise/Datasmith/DatasmithSketchUpRubyExporter/Private/DatasmithSketchUpCamera.h
@@ -11,6 +11,7 @@
 #include "DatasmithSketchUpSDKCeases.h"
 
 #include "Containers/UnrealString.h"
+#include "Misc/SecureHash.h"
 
 class IDatasmithCameraActorElement;
 
@@ -30,10 +31,12 @@ namespace DatasmithSketchUp
 		static TSharedPtr<FCamera> Create(FExportContext& Context, SUSceneRef InSceneRef);
 
 		void Update(FExportContext& Context);
+		FMD5Hash GetHash();
 
 		SUCameraRef CameraRef;
 		FString Name;
 		TSharedPtr<IDatasmithCameraActorElement> DatasmithCamera;
+		bool bIsActive;
 	};
 
 }
diff --git a/Engine/Source/Programs/Enterprise/Datasmith/DatasmithSketchUpRubyExporter/Private/DatasmithSketchUpExportContext.cpp b/Engine/Source/Programs/Enterprise/Datasmith/DatasmithSketchUpRubyExporter/Private/DatasmithSketchUpExportContext.cpp
--- a/Engine/Source/Programs/Enterprise/Datasmith/DatasmithSketchUpRubyExporter/Private/DatasmithSketchUpExportContext.cpp
+++ b/Engine/Source/Programs/Enterprise/Datasmith/DatasmithSketchUpRubyExporter/Private/DatasmithSketchUpExportContext.cpp
@@ -95,7 +95,7 @@ void FExportContext::Populate()
 
 	// Parse/convert Model
 	Layers.PopulateFromModel(ModelRef);
-	Scenes.PopulateFromModel(ModelRef);
+	Scenes.Update();
 	ComponentDefinitions.PopulateFromModel(ModelRef);
 
 	ModelDefinition->ParseNode(*this, *RootNode); // Create node hierarchy
@@ -164,6 +164,8 @@ bool FExportContext::Update(bool bModifiedHint)
 
 	bModified |= PreUpdateColorByLayer();
 
+	bModified |= Scenes.Update();
+
 	if (!bModified)
 	{
 		// code below is not supposed to change Datasmith scene if not modification hint(i.e. no invalidated data) was present
@@ -301,18 +303,34 @@ bool FImageCollection::RemoveImage(FComponentInstanceIDType ParentEntityId, FEnt
 
 }
 
-void FSceneCollection::PopulateFromModel(SUModelRef InModelRef)
+bool FSceneCollection::SetActiveScene(const FEntityIDType& EntityID)
+{
+	if (ActiveSceneId == EntityID)
+	{
+		return false;
+	}
+	ActiveSceneId = EntityID;
+	return true;
+}
+
+bool FSceneCollection::Update()
 {
-	// Get the number of scenes in the SketchUp model.
+	// Extract set of scenes(cameras ) and compute data hash to determine if Datasmith update is needed
+
+	FMD5 MD5;
+
 	size_t SceneCount = 0;
-	SUModelGetNumScenes(InModelRef, &SceneCount); // we can ignore the returned SU_RESULT
+	SUModelGetNumScenes(Context.ModelRef, &SceneCount);
+	MD5.Update(reinterpret_cast<uint8*>(&SceneCount), sizeof(SceneCount));
+
+	TArray<FEntityIDType> SceneIds;
 
 	if (SceneCount > 0)
 	{
 		// Retrieve the scenes in the SketchUp model.
 		TArray<SUSceneRef> Scenes;
 		Scenes.Init(SU_INVALID, SceneCount);
-		SUResult SResult = SUModelGetScenes(InModelRef, SceneCount, Scenes.GetData(), &SceneCount);
+		SUResult SResult = SUModelGetScenes(Context.ModelRef, SceneCount, Scenes.GetData(), &SceneCount);
 		Scenes.SetNum(SceneCount);
 		// Make sure the SketchUp model has scenes to retrieve (no SU_ERROR_NO_DATA).
 		if (SResult == SU_ERROR_NONE)
@@ -322,15 +340,67 @@ void FSceneCollection::PopulateFromModel(SUModelRef InModelRef)
 				// Make sure the SketchUp scene uses a camera.
 				bool bSceneUseCamera = false;
 				SUSceneGetUseCamera(SceneRef, &bSceneUseCamera); // we can ignore the returned SU_RESULT
+				MD5.Update(reinterpret_cast<uint8*>(&bSceneUseCamera), sizeof(bSceneUseCamera));
 
 				if (bSceneUseCamera)
 				{
-					TSharedPtr<DatasmithSketchUp::FCamera> Camera = FCamera::Create(Context, SceneRef);
-					SceneIdToCameraMap.Add(DatasmithSketchUpUtils::GetSceneID(SceneRef), Camera);
+					FEntityIDType SceneId = DatasmithSketchUpUtils::GetSceneID(SceneRef);
+
+					TSharedPtr<DatasmithSketchUp::FCamera> Camera;
+					if (TSharedPtr<DatasmithSketchUp::FCamera>* Found =  SceneIdToCameraMap.Find(SceneId))
+					{
+						Camera = *Found;
+					}
+					else
+					{
+						Camera = FCamera::Create(Context, SceneRef);
+						SceneIdToCameraMap.Add(SceneId, Camera);
+					}
+
+					MD5.Update(reinterpret_cast<uint8*>(&SceneId), sizeof(SceneId));
+
+					Camera->bIsActive = SceneId == ActiveSceneId;
+					
+					FMD5Hash CameraHash = Camera->GetHash();
+					
+					MD5.Update(CameraHash.GetBytes(), CameraHash.GetSize());
+
+					SceneIds.Add(SceneId);
 				}
 			}
 		}
 	}
+
+	MD5.Update(reinterpret_cast<uint8*>(&ActiveSceneId), sizeof(ActiveSceneId));
+
+	FMD5Hash Hash;
+	Hash.Set(MD5);
+
+	bool bModified = ScenesHash != Hash;
+
+	ScenesHash = Hash;
+
+	if (bModified)
+	{
+
+		for (TPair<FEntityIDType, TSharedPtr<FCamera>> IdCamera : SceneIdToCameraMap.Array())
+		{
+			FEntityIDType SceneId = IdCamera.Key;
+			TSharedPtr<FCamera> Camera = IdCamera.Value;
+			if (SceneIds.Contains(SceneId))
+			{
+				Camera->Update(Context);
+			}
+			else
+			{
+				// Cleanup removed scenes
+				Context.DatasmithScene->RemoveActor(Camera->DatasmithCamera, EDatasmithActorRemovalRule::RemoveChildren);
+				SceneIdToCameraMap.Remove(SceneId);
+			}
+		}
+	}
+
+	return bModified;
 }
 
 void FLayerCollection::PopulateFromModel(SUModelRef InModelRef)
diff --git a/Engine/Source/Programs/Enterprise/Datasmith/DatasmithSketchUpRubyExporter/Private/DatasmithSketchUpExportContext.h b/Engine/Source/Programs/Enterprise/Datasmith/DatasmithSketchUpRubyExporter/Private/DatasmithSketchUpExportContext.h
--- a/Engine/Source/Programs/Enterprise/Datasmith/DatasmithSketchUpRubyExporter/Private/DatasmithSketchUpExportContext.h
+++ b/Engine/Source/Programs/Enterprise/Datasmith/DatasmithSketchUpRubyExporter/Private/DatasmithSketchUpExportContext.h
@@ -202,10 +202,16 @@ namespace DatasmithSketchUp
 		void PopulateFromModel(
 			SUModelRef InSModelRef // model containing SketchUp camera definitions
 		);
+		bool SetActiveScene(const FEntityIDType& EntityID);
+
+		bool Update();
 
 		TMap<FSceneIDType, TSharedPtr<FCamera>> SceneIdToCameraMap;
+		FSceneIDType ActiveSceneId = FSceneIDType();
 	private:
 		FExportContext& Context;
+
+		FMD5Hash ScenesHash;
 	};
 
 	class FRegularMaterials
diff --git a/Engine/Source/Programs/Enterprise/Datasmith/DatasmithSketchUpRubyExporter/Private/DatasmithSketchUpExporter.cpp b/Engine/Source/Programs/Enterprise/Datasmith/DatasmithSketchUpRubyExporter/Private/DatasmithSketchUpExporter.cpp
--- a/Engine/Source/Programs/Enterprise/Datasmith/DatasmithSketchUpRubyExporter/Private/DatasmithSketchUpExporter.cpp
+++ b/Engine/Source/Programs/Enterprise/Datasmith/DatasmithSketchUpRubyExporter/Private/DatasmithSketchUpExporter.cpp
@@ -547,6 +547,11 @@ public:
 	{
 		return Context.InvalidateColorByLayer();
 	}
+
+	bool SetActiveScene(const DatasmithSketchUp::FEntityIDType& EntityID)
+	{
+		return Context.Scenes.SetActiveScene(EntityID);
+	}
 };
 
 ///////////////////////////////////////////////////////////////////////////
@@ -644,6 +649,19 @@ VALUE DatasmithSketchUpDirectLinkExporter_send_update(VALUE self)
 	return Qtrue;
 }
 
+VALUE DatasmithSketchUpDirectLinkExporter_set_active_scene(VALUE self, VALUE ruby_entity_id)
+{
+	// Converting args
+	FDatasmithSketchUpDirectLinkExporter* Ptr;
+	Data_Get_Struct(self, FDatasmithSketchUpDirectLinkExporter, Ptr);
+
+	Check_Type(ruby_entity_id, T_FIXNUM);
+	int32 EntityId = FIX2LONG(ruby_entity_id);
+	// Done converting args
+
+	return Ptr->SetActiveScene(DatasmithSketchUp::FEntityIDType(EntityId)) ? Qtrue : Qfalse;
+}
+
 VALUE DatasmithSketchUpDirectLinkExporter_update(VALUE self, VALUE modified_hint)
 {
 	// Converting args
@@ -943,6 +961,8 @@ extern "C" DLLEXPORT void Init_DatasmithSketchUp()
 
 	rb_define_method(DatasmithSketchUpDirectLinkExporterCRubyClass, "get_connection_status", ToRuby(DatasmithSketchUpDirectLinkExporter_get_connection_status), 0);
 
+	rb_define_method(DatasmithSketchUpDirectLinkExporterCRubyClass, "set_active_scene", ToRuby(DatasmithSketchUpDirectLinkExporter_set_active_scene), 1);
+
 }
 
 /* todo:
