[3ds Max] Export Selected Button

- replaced dropdown option to choose what to export with dedicated button to export all selected
- additionally: used 'gear' ison for settings and new 'wrench' icon for tools menus

diff --git a/Engine/Source/Programs/Enterprise/Datasmith/DatasmithMaxExporter/Extras/3dsMaxScripts/UI/DatasmithUI.xaml b/Engine/Source/Programs/Enterprise/Datasmith/DatasmithMaxExporter/Extras/3dsMaxScripts/UI/DatasmithUI.xaml
--- a/Engine/Source/Programs/Enterprise/Datasmith/DatasmithMaxExporter/Extras/3dsMaxScripts/UI/DatasmithUI.xaml
+++ b/Engine/Source/Programs/Enterprise/Datasmith/DatasmithMaxExporter/Extras/3dsMaxScripts/UI/DatasmithUI.xaml
@@ -94,6 +94,20 @@
           Size="Large"
           Text="Export"
           UID="a1de5acd-fef8-45be-8e01-ec194b9b5191" />
+        <ads:RibbonButton
+          CommandHandler="{wpfcust:ActionItem ActionId=5, TableId=689133272}"
+          Description="Export Selected objects to Datasmith file"
+          Image="Icons/Datasmith/DatasmithIcon32.png"
+          IsCheckable="true"
+          IsToolTipEnabled="false"
+          LargeImage="Icons/Datasmith/DatasmithIcon32.png"
+          MinWidth="0"
+          Orientation="Vertical"
+          ResizeStyle="HideText"
+          ShowText="true"
+          Size="Large"
+          Text="Export Selected"
+          UID="405855b6-f3c3-4589-af11-7969a5badaa6" />
         <ads:RibbonPanelBreak
           AllowInToolBar="false"
           IsToolTipEnabled="true"
@@ -108,70 +122,6 @@
           MinWidth="100"
           ShowText="false"
           UID="4f43d4b3-f741-4262-8a64-ed43f3354248">
-          <wpfcust:MaxRibbonLabel
-            AllowInToolBar="false"
-            IsToolTipEnabled="false"
-            MinWidth="150"
-            ResizeStyle="NoResize"
-            ShowImage="false"
-            ShowText="true"
-            Text="Include:"
-            UID="8d852bc3-7af5-4060-920a-a0514588d74a" />
-          <ads:RibbonSeparator
-            AllowInToolBar="false"
-            IsToolTipEnabled="true"
-            MinWidth="0"
-            ResizeStyle="NoResize"
-            SeparatorStyle="Spacer"
-            ShowText="false"
-            UID="918543c4-50dc-4203-bbdf-fe5fa4bd10fe" />
-          <wpfcust:MaxRibbonSplitButton
-            Image="pack://application:,,,/AdWindows;component/Resources/Images/MissingImage.ico"
-            IsCheckable="false"
-            IsSplit="false"
-            IsSynchronizedWithCurrentItem="true"
-            IsToolTipEnabled="true"
-            LargeImage="pack://application:,,,/AdWindows;component/Resources/Images/MissingImage.ico"
-            ListImageSize="Standard"
-            MaxscriptIndexCallback="Datasmith_ExportOptions_IncludeIndex()"
-            MinWidth="0"
-            ResizeStyle="HideText"
-            ShowImage="false"
-            ShowText="true"
-            UID="f97d81dd-03f8-4a14-b718-06733fb224f2">
-            <ads:RibbonCommandItem
-              AllowInToolBar="false"
-              CommandHandler="{wpfcust:ActionItem MacroName=Datasmith_ExportOptions_IncludeAllVisibleObjects, MacroCategory=Datasmith}"
-              Image="pack://application:,,,/AdWindows;component/Resources/Images/MissingImage.ico"
-              IsCheckable="true"
-              IsToolTipEnabled="true"
-              LargeImage="pack://application:,,,/AdWindows;component/Resources/Images/MissingImage.ico"
-              MinWidth="100"
-              Name="Visible Objects"
-              ShowImage="false"
-              ShowText="false"
-              Text="Visible Objects"
-              UID="6ef627ea-f1fb-4b7a-9235-6f302d5ac4de" />
-            <ads:RibbonCommandItem
-              AllowInToolBar="false"
-              CommandHandler="{wpfcust:ActionItem MacroName=Datasmith_ExportOptions_IncludeSelectedOnly, MacroCategory=Datasmith}"
-              Image="pack://application:,,,/AdWindows;component/Resources/Images/MissingImage.ico"
-              IsCheckable="true"
-              IsToolTipEnabled="true"
-              LargeImage="pack://application:,,,/AdWindows;component/Resources/Images/MissingImage.ico"
-              MinWidth="100"
-              Name="Selected Objects"
-              ShowImage="false"
-              ShowText="false"
-              Text="Selected Objects"
-              UID="9509a5c2-1553-4499-91c3-7b7b0a220cee" />
-          </wpfcust:MaxRibbonSplitButton>
-          <ads:RibbonRowBreak
-            AllowInToolBar="false"
-            IsToolTipEnabled="true"
-            MinWidth="100"
-            ShowText="false"
-            UID="81ee274b-ff06-4d67-a974-6461da0acf54" />
           <wpfcust:MaxRibbonLabel
             AllowInToolBar="false"
             IsToolTipEnabled="false"
@@ -234,7 +184,7 @@
       </wpfcust:MaxRibbonPanelSource>
     </wpfcust:MaxModelingRibbonPanel>
     <wpfcust:MaxModelingRibbonPanel
-      CollapsedPanelImage="Icons/MissingImage_32.png"
+      CollapsedPanelImage="Icons/Datasmith/DatasmithSettingsIcon32.png"
       FloatingGroup="0"
       FloatingOrder="0"
       UID="42ca9067-e3a4-4d3c-8d83-af9bb10b1124">
@@ -391,7 +341,7 @@
       </wpfcust:MaxRibbonPanelSource>
     </wpfcust:MaxModelingRibbonPanel>
     <wpfcust:MaxModelingRibbonPanel
-      CollapsedPanelImage="Icons/Datasmith/DatasmithSettingsIcon32.png"
+      CollapsedPanelImage="Icons/Datasmith/DatasmithToolsIcon32.png"
       FloatingGroup="0"
       FloatingOrder="0"
       ResizeStyle="NeverCollapsePanel"
diff --git a/Engine/Source/Programs/Enterprise/Datasmith/DatasmithMaxExporter/Private/DatasmithMaxDirectLinkMaxScript.cpp b/Engine/Source/Programs/Enterprise/Datasmith/DatasmithMaxExporter/Private/DatasmithMaxDirectLinkMaxScript.cpp
--- a/Engine/Source/Programs/Enterprise/Datasmith/DatasmithMaxExporter/Private/DatasmithMaxDirectLinkMaxScript.cpp
+++ b/Engine/Source/Programs/Enterprise/Datasmith/DatasmithMaxExporter/Private/DatasmithMaxDirectLinkMaxScript.cpp
@@ -157,8 +157,9 @@ Value* Export_cf(Value** arg_list, int count)
 	Value* pOutputPath = arg_list[1];
 
 	bool bQuiet = key_arg_or_default(quiet, &false_value)->to_bool();
+	bool bSelected = key_arg_or_default(selected, &false_value)->to_bool();
 
-	bool bResult = Export(pName->to_string(), pOutputPath->to_string(), bQuiet);;
+	bool bResult = Export(pName->to_string(), pOutputPath->to_string(), bQuiet, bSelected);;
 	return bool_result(bResult);
 }
 
@@ -314,7 +315,6 @@ Value* SetExportOption_##name##_cf(Value** arg_list, int count) \
 Primitive SetExportOption_##name##_pf(_M("Datasmith_SetExportOption_" #name), SetExportOption_##name##_cf);\
 
 //////////////////////////////////////////
-DefinePersistentExportOption(SelectedOnly)
 DefinePersistentExportOption(AnimatedTransforms)
 
 DefinePersistentExportOption(StatSync)
@@ -555,6 +555,7 @@ public:
 		ID_CONNECTIONS_ACTION_ID,
 		ID_EXPORT_ACTION_ID,
 		ID_SHOWLOG_ACTION_ID,
+		ID_EXPORT_SELECTED_ACTION_ID,
 	};
 
 	class FDatasmithActionTable : public ActionTable
@@ -612,6 +613,11 @@ public:
 				Icon.Reset(new MaxBmpFileIcon(_T(":/Datasmith/Icons/DatasmithLogIcon")));
 				break;
 			}
+			case ID_EXPORT_SELECTED_ACTION_ID:
+			{
+				Icon.Reset(new MaxBmpFileIcon(_T(":/Datasmith/Icons/DatasmithIcon")));
+				break;
+			}
 			}
 			return Icon.Get();
 		}
@@ -652,7 +658,7 @@ public:
 			{
 				if (GetExporter())
 				{
-					MCHAR* ScriptCode = _T("Datasmith_ExportDialog()");
+					MCHAR* ScriptCode = _T("Datasmith_ExportDialog selected:false");
 #if MAX_PRODUCT_YEAR_NUMBER >= 2022
 					 ExecuteMAXScriptScript(ScriptCode, MAXScript::ScriptSource::NonEmbedded);
 #else
@@ -666,6 +672,19 @@ public:
 				MessagesDialog->OpenWindow();
 				return true;
 			}
+			case ID_EXPORT_SELECTED_ACTION_ID:
+			{
+				if (GetExporter())
+				{
+					MCHAR* ScriptCode = _T("Datasmith_ExportDialog selected:true");
+#if MAX_PRODUCT_YEAR_NUMBER >= 2022
+					 ExecuteMAXScriptScript(ScriptCode, MAXScript::ScriptSource::NonEmbedded);
+#else
+					 ExecuteMAXScriptScript(ScriptCode);
+#endif
+				}
+				return true;
+			}
 			}
 		    return false;
 		}
@@ -688,6 +707,7 @@ public:
 			DATASMITH_ACTION(CONNECTIONS),
 			DATASMITH_ACTION(EXPORT),
 			DATASMITH_ACTION(SHOWLOG),
+			DATASMITH_ACTION(EXPORT_SELECTED),
 		};
 
 		FDatasmithActionTable* Table = new FDatasmithActionTable(Name); // Table, registered with RegisterActionTable will be deallocated by Max
diff --git a/Engine/Source/Programs/Enterprise/Datasmith/DatasmithMaxExporter/Private/DatasmithMaxExporterWithDirectLink.cpp b/Engine/Source/Programs/Enterprise/Datasmith/DatasmithMaxExporter/Private/DatasmithMaxExporterWithDirectLink.cpp
--- a/Engine/Source/Programs/Enterprise/Datasmith/DatasmithMaxExporter/Private/DatasmithMaxExporterWithDirectLink.cpp
+++ b/Engine/Source/Programs/Enterprise/Datasmith/DatasmithMaxExporter/Private/DatasmithMaxExporterWithDirectLink.cpp
@@ -566,7 +566,6 @@ public:
 		{
 			return;
 		}
-		GetBool(TEXT("SelectedOnly"), Options.bSelectedOnly);
 		GetBool(TEXT("AnimatedTransforms"), Options.bAnimatedTransforms);
 		GetInt(TEXT("TextureResolution"), Options.TextureResolution);
 		bLoaded = true;
@@ -620,16 +619,6 @@ public:
 		return FPaths::Combine(PlugCfgPath, TEXT("UnrealDatasmithMax.ini"));
 	}
 
-	virtual void SetSelectedOnly(bool bValue) override
-	{
-		Options.bSelectedOnly = bValue;
-		SetBool(TEXT("SelectedOnly"), Options.bSelectedOnly);
-	}
-	virtual bool GetSelectedOnly() override
-	{
-		return Options.bSelectedOnly;
-	}
-
 	virtual void SetAnimatedTransforms(bool bValue) override
 	{
 		Options.bAnimatedTransforms = bValue;
@@ -3529,7 +3518,7 @@ void FExporter::Shutdown()
 	FDatasmithExporterManager::Shutdown();
 }
 
-bool Export(const TCHAR* Name, const TCHAR* OutputPath, bool bQuiet)
+bool Export(const TCHAR* Name, const TCHAR* OutputPath, bool bQuiet, bool bSelected)
 {
 	FUpdateProgress ProgressManager(!bQuiet, 3);
 	FUpdateProgress::FStage& MainStage = ProgressManager.MainStage;
@@ -3539,7 +3528,9 @@ bool Export(const TCHAR* Name, const TCHAR* OutputPath, bool bQuiet)
 	ExportedScene.SetName(Name);
 	ExportedScene.SetOutputPath(OutputPath);
 
-	FSceneTracker SceneTracker(PersistentExportOptions.Options, ExportedScene, nullptr);
+	FExportOptions ExportOptions = PersistentExportOptions.Options;
+	ExportOptions.bSelectedOnly = bSelected;
+	FSceneTracker SceneTracker((ExportOptions), ExportedScene, nullptr);
 
 	bool bCancelled = false;
 
@@ -3548,7 +3539,7 @@ bool Export(const TCHAR* Name, const TCHAR* OutputPath, bool bQuiet)
 		bCancelled = true;
 	}
 
-	if (PersistentExportOptions.Options.bAnimatedTransforms && !bCancelled)
+	if (ExportOptions.bAnimatedTransforms && !bCancelled)
 	{
 		if (!SceneTracker.ExportAnimations(MainStage))
 		{
diff --git a/Engine/Source/Programs/Enterprise/Datasmith/DatasmithMaxExporter/Public/DatasmithMaxDirectLink.h b/Engine/Source/Programs/Enterprise/Datasmith/DatasmithMaxExporter/Public/DatasmithMaxDirectLink.h
--- a/Engine/Source/Programs/Enterprise/Datasmith/DatasmithMaxExporter/Public/DatasmithMaxDirectLink.h
+++ b/Engine/Source/Programs/Enterprise/Datasmith/DatasmithMaxExporter/Public/DatasmithMaxDirectLink.h
@@ -49,10 +49,6 @@ class FLayerTracker;
 class IPersistentExportOptions
 {
 public:
-	 // Whether to export all visible or selected objects only
-	virtual void SetSelectedOnly(bool) = 0;
-	virtual bool GetSelectedOnly() = 0;
-
 
 	// Whether to export animation or not
 	virtual void SetAnimatedTransforms(bool) = 0;
@@ -103,7 +99,7 @@ class IExporter
 bool CreateExporter(bool bEnableUI, const TCHAR* EnginePath); // Create exporter with ability for DirectLink change tracking
 IExporter* GetExporter();
 void ShutdownExporter();
-bool Export(const TCHAR* Name, const TCHAR* OutputPath, bool bQuiet);
+bool Export(const TCHAR* Name, const TCHAR* OutputPath, bool bQuiet, bool bSelected);
 void ShutdownScripts();
 
 
diff --git a/Engine/Source/Programs/Enterprise/Datasmith/DatasmithMaxExporter/Resources/Windows/Resource.rc b/Engine/Source/Programs/Enterprise/Datasmith/DatasmithMaxExporter/Resources/Windows/Resource.rc
--- a/Engine/Source/Programs/Enterprise/Datasmith/DatasmithMaxExporter/Resources/Windows/Resource.rc
+++ b/Engine/Source/Programs/Enterprise/Datasmith/DatasmithMaxExporter/Resources/Windows/Resource.rc
@@ -148,6 +148,8 @@ BEGIN
     IDS_SHOWLOG_NAME    "Log"
     IDS_SHOWLOG_DESC    "Show Messages"
     IDS_DATASMITH_CATEGORY "Datasmith"
+    IDS_EXPORT_SELECTED_NAME    "Export Selected"
+    IDS_EXPORT_SELECTED_DESC    "Export Selected objects to Datasmith file"
 END
 
 #endif    // English (United States) resources
diff --git a/Engine/Source/Programs/Enterprise/Datasmith/DatasmithMaxExporter/Resources/Windows/resource.h b/Engine/Source/Programs/Enterprise/Datasmith/DatasmithMaxExporter/Resources/Windows/resource.h
--- a/Engine/Source/Programs/Enterprise/Datasmith/DatasmithMaxExporter/Resources/Windows/resource.h
+++ b/Engine/Source/Programs/Enterprise/Datasmith/DatasmithMaxExporter/Resources/Windows/resource.h
@@ -23,6 +23,8 @@
 #define IDS_SHOWLOG_NAME                117
 #define IDS_SHOWLOG_DESC                118
 #define IDS_DATASMITH_CATEGORY          119
+#define IDS_EXPORT_SELECTED_NAME        120
+#define IDS_EXPORT_SELECTED_DESC        121
 #define IDD_ERROR_MSGS                  181
 #define IDC_BYOBJECT                    1016
 #define IDC_BYMATERIAL                  1017
