[3ds Max] Datasmith Version

- printed in log and displayed in Messages title

diff --git a/Engine/Source/Programs/Enterprise/Datasmith/DatasmithMaxExporter/Private/DatasmithMaxDirectLinkMaxScript.cpp b/Engine/Source/Programs/Enterprise/Datasmith/DatasmithMaxExporter/Private/DatasmithMaxDirectLinkMaxScript.cpp
--- a/Engine/Source/Programs/Enterprise/Datasmith/DatasmithMaxExporter/Private/DatasmithMaxDirectLinkMaxScript.cpp
+++ b/Engine/Source/Programs/Enterprise/Datasmith/DatasmithMaxExporter/Private/DatasmithMaxDirectLinkMaxScript.cpp
@@ -87,6 +87,8 @@ Value* OnLoad_cf(Value **arg_list, int count)
 		MessagesDialog = Ui::CreateMessagesWindow();
 	}
 
+	LogInfo(FString::Printf(TEXT("Initialized Datasmith SDK v.%s"), *FDatasmithUtils::GetEnterpriseVersionAsString()));
+
 	return bool_result(bResult);
 }
 
diff --git a/Engine/Source/Programs/Enterprise/Datasmith/DatasmithMaxExporter/Private/DatasmithMaxExporterWithDirectLink.cpp b/Engine/Source/Programs/Enterprise/Datasmith/DatasmithMaxExporter/Private/DatasmithMaxExporterWithDirectLink.cpp
--- a/Engine/Source/Programs/Enterprise/Datasmith/DatasmithMaxExporter/Private/DatasmithMaxExporterWithDirectLink.cpp
+++ b/Engine/Source/Programs/Enterprise/Datasmith/DatasmithMaxExporter/Private/DatasmithMaxExporterWithDirectLink.cpp
@@ -3461,11 +3461,13 @@ bool CreateExporter(bool bEnableUI, const TCHAR* EnginePath)
 
 	if (!FDatasmithExporterManager::Initialize(Options))
 	{
+		LogError(TEXT("Failed to initialize Datasmith Exporter Manager"));
 		return false;
 	}
 
 	if (int32 ErrorCode = FDatasmithDirectLink::ValidateCommunicationSetup())
 	{
+		LogError(TEXT("Failed to validate DatasmithDirect Link Communication setup"));
 		return false;
 	}
 
diff --git a/Engine/Source/Programs/Enterprise/Datasmith/DatasmithMaxExporter/Private/UI/Datasmith3dsMaxUI.cpp b/Engine/Source/Programs/Enterprise/Datasmith/DatasmithMaxExporter/Private/UI/Datasmith3dsMaxUI.cpp
--- a/Engine/Source/Programs/Enterprise/Datasmith/DatasmithMaxExporter/Private/UI/Datasmith3dsMaxUI.cpp
+++ b/Engine/Source/Programs/Enterprise/Datasmith/DatasmithMaxExporter/Private/UI/Datasmith3dsMaxUI.cpp
@@ -26,6 +26,7 @@
 #include "HAL/PlatformApplicationMisc.h"  // For ClipboardCopy
 
 #include "DatasmithExporterManager.h"  // Using FDatasmithExporterManager::PushCommandIntoGameThread
+#include "DatasmithUtils.h"
 
 
 #define LOCTEXT_NAMESPACE "Datasmith3dsMaxUI"
@@ -243,6 +244,13 @@ public:
 class FMessagesWindow : public IMessagesWindow
 {
 public:
+	FText GetWindowTitle()
+	{
+		FFormatNamedArguments WindowTitleArgs;
+		WindowTitleArgs.Add(TEXT("DatasmithSDKVersion"), FText::FromString(FDatasmithUtils::GetEnterpriseVersionAsString()));
+		return FText::Format(LOCTEXT("DatasmithMessagesWindowTitle", "Datasmith Messages v.{DatasmithSDKVersion}"), WindowTitleArgs);
+	}
+
 	virtual void OpenWindow() override
 	{
 		FSimpleDelegate RunOnUIThread;
@@ -265,7 +273,7 @@ public:
 					.AutoCenter( EAutoCenter::PrimaryWorkArea )
 					.SizingRule( ESizingRule::UserSized )
 					.FocusWhenFirstShown( true )
-					.Title( LOCTEXT("DatasmithMessagesWindowTitle", "Datasmith Messages") );
+					.Title( GetWindowTitle() );
 
 				TSharedRef<SWindowTitleBar> WindowTitleBar = SNew( SWindowTitleBar, Window, nullptr, GetWindowTitleAlignement() )
 					.Visibility( EVisibility::Visible )
