[3ds Max] XRef Scenes export/sync toggle

diff --git a/Engine/Source/Programs/Enterprise/Datasmith/DatasmithMaxExporter/Extras/3dsMaxScripts/UI/DatasmithUI.xaml b/Engine/Source/Programs/Enterprise/Datasmith/DatasmithMaxExporter/Extras/3dsMaxScripts/UI/DatasmithUI.xaml
--- a/Engine/Source/Programs/Enterprise/Datasmith/DatasmithMaxExporter/Extras/3dsMaxScripts/UI/DatasmithUI.xaml
+++ b/Engine/Source/Programs/Enterprise/Datasmith/DatasmithMaxExporter/Extras/3dsMaxScripts/UI/DatasmithUI.xaml
@@ -337,6 +337,17 @@
             MinWidth="100"
             ShowText="false"
             UID="189ae82b-5e0f-4d3b-aa93-81e928234b27" />
+          <wpfcust:MaxRibbonCheckBox
+            CommandHandler="{wpfcust:ActionItem MacroName=Datasmith_ExportOptions_XRefScenes, MacroCategory=Datasmith}"
+            Description="Enables export/sync of XRef Scenes"
+            IsCheckable="true"
+            IsToolTipEnabled="true"
+            MinWidth="0"
+            Name="Export XRef Scenes"
+            ResizeStyle="NoResize"
+            ShowText="true"
+            Text="Export XRef Scenes"
+            UID="f8ed3b9b-6c4b-4b0c-92cb-618f4784ab94" />
         </ads:RibbonRowPanel>
       </wpfcust:MaxRibbonPanelSource>
     </wpfcust:MaxModelingRibbonPanel>
diff --git a/Engine/Source/Programs/Enterprise/Datasmith/DatasmithMaxExporter/Private/DatasmithMaxDirectLinkMaxScript.cpp b/Engine/Source/Programs/Enterprise/Datasmith/DatasmithMaxExporter/Private/DatasmithMaxDirectLinkMaxScript.cpp
--- a/Engine/Source/Programs/Enterprise/Datasmith/DatasmithMaxExporter/Private/DatasmithMaxDirectLinkMaxScript.cpp
+++ b/Engine/Source/Programs/Enterprise/Datasmith/DatasmithMaxExporter/Private/DatasmithMaxDirectLinkMaxScript.cpp
@@ -321,6 +321,7 @@ DefinePersistentExportOption(AnimatedTransforms)
 
 DefinePersistentExportOption(StatSync)
 DefinePersistentExportOptionInt(TextureResolution)
+DefinePersistentExportOption(XRefScenes)
 //////////////////////////////////////////
 
 #undef DefinePersistentExportOption
diff --git a/Engine/Source/Programs/Enterprise/Datasmith/DatasmithMaxExporter/Private/DatasmithMaxExporterWithDirectLink.cpp b/Engine/Source/Programs/Enterprise/Datasmith/DatasmithMaxExporter/Private/DatasmithMaxExporterWithDirectLink.cpp
--- a/Engine/Source/Programs/Enterprise/Datasmith/DatasmithMaxExporter/Private/DatasmithMaxExporterWithDirectLink.cpp
+++ b/Engine/Source/Programs/Enterprise/Datasmith/DatasmithMaxExporter/Private/DatasmithMaxExporterWithDirectLink.cpp
@@ -553,7 +553,9 @@ struct FExportOptions
 	bool bStatSync = false;
 	int32 TextureResolution = 4;
 
-	static const int32 TextureResolutionMax = 6; // 
+	static const int32 TextureResolutionMax = 6; //
+
+	bool bXRefScenes = true;
 };
 
 // Global export options, stored in preferences
@@ -568,6 +570,7 @@ public:
 		}
 		GetBool(TEXT("AnimatedTransforms"), Options.bAnimatedTransforms);
 		GetInt(TEXT("TextureResolution"), Options.TextureResolution);
+		GetBool(TEXT("XRefScenes"), Options.bXRefScenes);
 		bLoaded = true;
 	}
 
@@ -630,6 +633,17 @@ public:
 		return Options.bAnimatedTransforms;
 	}
 
+	virtual void SetXRefScenes(bool bValue) override
+	{
+		Options.bXRefScenes = bValue;
+		SetBool(TEXT("XRefScenes"), Options.bXRefScenes);
+	}
+
+	virtual bool GetXRefScenes() override
+	{
+		return Options.bXRefScenes;
+	}
+
 	virtual void SetStatSync(bool bValue) override
 	{
 		Options.bStatSync = bValue;
@@ -1026,7 +1040,7 @@ public:
 	}
 
 	bool bParseXRefScenes = true;
-	bool bIncludeXRefWhileParsing = false;
+	bool bIncludeXRefWhileParsing = false; // note - this should not be set when bParseXRefScenes if disabled
 
 	virtual bool IsUpdateInProgress() override
 	{
@@ -1044,7 +1058,12 @@ public:
 		// nodes comming from XRef Scenes/Objects could be null
 		if (!SceneRootNode)
 		{
-			return false;
+			return true;
+		}
+
+		if (XRefScene && !bParseXRefScenes)
+		{
+			return true; 
 		}
 
 		if (!bIncludeXRefWhileParsing)
@@ -1349,6 +1368,22 @@ public:
 
 		Stats.Reset();
 
+
+		if (bParseXRefScenes != Options.bXRefScenes)
+		{
+			bParseXRefScenes = Options.bXRefScenes;
+
+			// Clear all tracked nodes when XRef scenes export option changes
+			// note: when including xref into hierarchy it might be impossible to invalidate/remove only xref nodes as they are indistinquishable
+			//  But when parsing XRef scenes explicitly it seems to be possible but need to track XRefTree Node
+			TArray<FNodeKey> NodeTrackerKeys;
+			NodeTrackers.GetKeys(NodeTrackerKeys);
+			for (FNodeKey NodeTrackerKey: NodeTrackerKeys)
+			{
+				RemoveNodeTracker(*GetNodeTracker(NodeTrackerKey));
+			}
+		}
+
 		{
 			PROGRESS_STAGE("Refresh layers")
 			bChangeEncountered |= UpdateLayers();
@@ -1886,7 +1921,9 @@ public:
 	// todo: make fine invalidates - full only something like geometry change, but finer for transform, name change and more
 	FNodeTracker* InvalidateNode(FNodeKey NodeKey, bool bCheckCalledInProgress = true)
 	{
-		LogDebugNode(TEXT("InvalidateNode"), NodeEventNamespace::GetNodeByKey(NodeKey));
+		INode* Node = NodeEventNamespace::GetNodeByKey(NodeKey);
+
+		LogDebugNode(TEXT("InvalidateNode"), Node);
 		// We don't expect node chances while Update inprogress(unless Invalidate called explicitly)
 		if (bCheckCalledInProgress)
 		{
@@ -1895,7 +1932,7 @@ public:
 
 		if (FNodeTracker* NodeTracker =  GetNodeTracker(NodeKey))
 		{
-			if (NodeEventNamespace::GetNodeByKey(NodeKey))
+			if (Node)
 			{
 				InvalidateNode(*NodeTracker, bCheckCalledInProgress);
 				return NodeTracker;
@@ -1911,7 +1948,11 @@ public:
 		}
 		else
 		{
-			AddNode(NodeKey, NodeEventNamespace::GetNodeByKey(NodeKey));
+			if (Node)
+			{
+				return AddNode(NodeKey, Node);
+			}
+			// See comment above. When Containers are used and no XRef Scenes export is enabled event may be received for node which doesn't exist
 		}
 		return nullptr;
 	}
diff --git a/Engine/Source/Programs/Enterprise/Datasmith/DatasmithMaxExporter/Public/DatasmithMaxDirectLink.h b/Engine/Source/Programs/Enterprise/Datasmith/DatasmithMaxExporter/Public/DatasmithMaxDirectLink.h
--- a/Engine/Source/Programs/Enterprise/Datasmith/DatasmithMaxExporter/Public/DatasmithMaxDirectLink.h
+++ b/Engine/Source/Programs/Enterprise/Datasmith/DatasmithMaxExporter/Public/DatasmithMaxDirectLink.h
@@ -60,6 +60,10 @@ public:
 
 	virtual void SetTextureResolution(int32 Bool) = 0;
 	virtual int32 GetTextureResolution() = 0;
+
+	virtual void SetXRefScenes(bool) = 0;
+	virtual bool GetXRefScenes() = 0;
+
 };
 
 //---- Main class for export/change tracking
