[Non-inclusivity] Removed banned words from Datasmith API and code using the API

diff --git a/Engine/Source/Developer/Datasmith/DatasmithFacade/Private/DatasmithFacadeMaterial.cpp b/Engine/Source/Developer/Datasmith/DatasmithFacade/Private/DatasmithFacadeMaterial.cpp
--- a/Engine/Source/Developer/Datasmith/DatasmithFacade/Private/DatasmithFacadeMaterial.cpp
+++ b/Engine/Source/Developer/Datasmith/DatasmithFacade/Private/DatasmithFacadeMaterial.cpp
@@ -31,9 +31,9 @@ FDatasmithFacadeBaseMaterial::EDatasmithMaterialType FDatasmithFacadeBaseMateria
 	{
 		return EDatasmithMaterialType::UEPbrMaterial;
 	}
-	else if (InMaterial->IsA(EDatasmithElementType::MasterMaterial))
+	else if (InMaterial->IsA(EDatasmithElementType::MaterialInstance))
 	{
-		return EDatasmithMaterialType::MasterMaterial;
+		return EDatasmithMaterialType::MaterialInstance;
 	}
 	else if (InMaterial->IsA(EDatasmithElementType::DecalMaterial))
 	{
@@ -61,8 +61,8 @@ FDatasmithFacadeBaseMaterial* FDatasmithFacadeBaseMaterial::GetNewFacadeBaseMate
 		{
 		case EDatasmithMaterialType::UEPbrMaterial:
 			return new FDatasmithFacadeUEPbrMaterial(StaticCastSharedRef<IDatasmithUEPbrMaterialElement>(MaterialRef));
-		case EDatasmithMaterialType::MasterMaterial:
-			return new FDatasmithFacadeMasterMaterial(StaticCastSharedRef<IDatasmithMasterMaterialElement>(MaterialRef));
+		case EDatasmithMaterialType::MaterialInstance:
+			return new FDatasmithFacadeMaterialInstance(StaticCastSharedRef<IDatasmithMaterialInstanceElement>(MaterialRef));
 		case EDatasmithMaterialType::DecalMaterial:
 			return new FDatasmithFacadeDecalMaterial(StaticCastSharedRef<IDatasmithDecalMaterialElement>(MaterialRef));
 		case EDatasmithMaterialType::Unsupported:
@@ -74,54 +74,54 @@ FDatasmithFacadeBaseMaterial* FDatasmithFacadeBaseMaterial::GetNewFacadeBaseMate
 	return nullptr;
 }
 
-FDatasmithFacadeMasterMaterial::FDatasmithFacadeMasterMaterial(const TCHAR* InElementName) 
-	: FDatasmithFacadeBaseMaterial( FDatasmithSceneFactory::CreateMasterMaterial(InElementName))
+FDatasmithFacadeMaterialInstance::FDatasmithFacadeMaterialInstance(const TCHAR* InElementName) 
+	: FDatasmithFacadeBaseMaterial( FDatasmithSceneFactory::CreateMaterialInstance(InElementName))
 {
-	TSharedPtr<IDatasmithMasterMaterialElement> MasterMaterial = GetDatasmithMasterMaterial();
-	MasterMaterial->SetMaterialType(EDatasmithMasterMaterialType::Opaque);
+	TSharedPtr<IDatasmithMaterialInstanceElement> MaterialInstance = GetDatasmithMaterialInstance();
+	MaterialInstance->SetMaterialType(EDatasmithReferenceMaterialType::Opaque);
 }
 
-FDatasmithFacadeMasterMaterial::FDatasmithFacadeMasterMaterial(const TSharedRef<IDatasmithMasterMaterialElement>& InMaterialRef)
+FDatasmithFacadeMaterialInstance::FDatasmithFacadeMaterialInstance(const TSharedRef<IDatasmithMaterialInstanceElement>& InMaterialRef)
 	: FDatasmithFacadeBaseMaterial(InMaterialRef)
 {}
 
-FDatasmithFacadeMasterMaterial::EMasterMaterialType FDatasmithFacadeMasterMaterial::GetMaterialType() const
+FDatasmithFacadeMaterialInstance::EMaterialInstanceType FDatasmithFacadeMaterialInstance::GetMaterialType() const
 {
-	return static_cast<EMasterMaterialType>(GetDatasmithMasterMaterial()->GetMaterialType());
+	return static_cast<EMaterialInstanceType>(GetDatasmithMaterialInstance()->GetMaterialType());
 }
 
-void FDatasmithFacadeMasterMaterial::SetMaterialType(
-	EMasterMaterialType InMasterMaterialType
+void FDatasmithFacadeMaterialInstance::SetMaterialType(
+	EMaterialInstanceType InMaterialInstanceType
 )
 {
-	GetDatasmithMasterMaterial()->SetMaterialType(static_cast<EDatasmithMasterMaterialType>(InMasterMaterialType));
+	GetDatasmithMaterialInstance()->SetMaterialType(static_cast<EDatasmithReferenceMaterialType>(InMaterialInstanceType));
 }
 
-FDatasmithFacadeMasterMaterial::EMasterMaterialQuality FDatasmithFacadeMasterMaterial::GetQuality() const
+FDatasmithFacadeMaterialInstance::EMaterialInstanceQuality FDatasmithFacadeMaterialInstance::GetQuality() const
 {
-	return static_cast<EMasterMaterialQuality>(GetDatasmithMasterMaterial()->GetQuality());
+	return static_cast<EMaterialInstanceQuality>(GetDatasmithMaterialInstance()->GetQuality());
 }
 
-void FDatasmithFacadeMasterMaterial::SetQuality(
-	EMasterMaterialQuality InQuality
+void FDatasmithFacadeMaterialInstance::SetQuality(
+	EMaterialInstanceQuality InQuality
 )
 {
-	GetDatasmithMasterMaterial()->SetQuality(static_cast<EDatasmithMasterMaterialQuality>(InQuality));
+	GetDatasmithMaterialInstance()->SetQuality(static_cast<EDatasmithReferenceMaterialQuality>(InQuality));
 }
 
-const TCHAR* FDatasmithFacadeMasterMaterial::GetCustomMaterialPathName() const
+const TCHAR* FDatasmithFacadeMaterialInstance::GetCustomMaterialPathName() const
 {
-	return GetDatasmithMasterMaterial()->GetCustomMaterialPathName();
+	return GetDatasmithMaterialInstance()->GetCustomMaterialPathName();
 }
 
-void FDatasmithFacadeMasterMaterial::SetCustomMaterialPathName(
+void FDatasmithFacadeMaterialInstance::SetCustomMaterialPathName(
 	const TCHAR* InPathName
 )
 {
-	GetDatasmithMasterMaterial()->SetCustomMaterialPathName(InPathName);
+	GetDatasmithMaterialInstance()->SetCustomMaterialPathName(InPathName);
 }
 
-void FDatasmithFacadeMasterMaterial::AddColor(
+void FDatasmithFacadeMaterialInstance::AddColor(
 	const TCHAR*  InPropertyName,
 	unsigned char InR,
 	unsigned char InG,
@@ -136,7 +136,7 @@ void FDatasmithFacadeMasterMaterial::AddColor(
 	AddColor(InPropertyName, LinearColor.R, LinearColor.G, LinearColor.B, LinearColor.A);
 }
 
-void FDatasmithFacadeMasterMaterial::AddColor(
+void FDatasmithFacadeMaterialInstance::AddColor(
 	const TCHAR* InPropertyName,
 	float        InR,
 	float        InG,
@@ -152,10 +152,10 @@ void FDatasmithFacadeMasterMaterial::AddColor(
 	MaterialPropertyPtr->SetValue(*LinearColor.ToString());
 
 	// Add the new property to the Datasmith material properties.
-	GetDatasmithMasterMaterial()->AddProperty(MaterialPropertyPtr);
+	GetDatasmithMaterialInstance()->AddProperty(MaterialPropertyPtr);
 }
 
-void FDatasmithFacadeMasterMaterial::AddTexture(
+void FDatasmithFacadeMaterialInstance::AddTexture(
 	const TCHAR* InPropertyName,
 	const FDatasmithFacadeTexture* InTexture
 )
@@ -168,11 +168,11 @@ void FDatasmithFacadeMasterMaterial::AddTexture(
 		MaterialPropertyPtr->SetValue(InTexture->GetName());
 
 		// Add the new property to the Datasmith material properties.
-		GetDatasmithMasterMaterial()->AddProperty(MaterialPropertyPtr);
+		GetDatasmithMaterialInstance()->AddProperty(MaterialPropertyPtr);
 	}
 }
 
-void FDatasmithFacadeMasterMaterial::AddString(
+void FDatasmithFacadeMaterialInstance::AddString(
 	const TCHAR* InPropertyName,
 	const TCHAR* InPropertyValue
 )
@@ -185,11 +185,11 @@ void FDatasmithFacadeMasterMaterial::AddString(
 		MaterialPropertyPtr->SetValue(InPropertyValue);
 
 		// Add the new property to the array of Datasmith material properties.
-		GetDatasmithMasterMaterial()->AddProperty(MaterialPropertyPtr);
+		GetDatasmithMaterialInstance()->AddProperty(MaterialPropertyPtr);
 	}
 }
 
-void FDatasmithFacadeMasterMaterial::AddFloat(
+void FDatasmithFacadeMaterialInstance::AddFloat(
 	const TCHAR* InPropertyName,
 	float        InPropertyValue
 )
@@ -200,10 +200,10 @@ void FDatasmithFacadeMasterMaterial::AddFloat(
 	MaterialPropertyPtr->SetValue(*FString::Printf(TEXT("%f"), InPropertyValue));
 
 	// Add the new property to the Datasmith material properties.
-	GetDatasmithMasterMaterial()->AddProperty(MaterialPropertyPtr);
+	GetDatasmithMaterialInstance()->AddProperty(MaterialPropertyPtr);
 }
 
-void FDatasmithFacadeMasterMaterial::AddBoolean(
+void FDatasmithFacadeMaterialInstance::AddBoolean(
 	const TCHAR* InPropertyName,
 	bool         bInPropertyValue
 )
@@ -214,19 +214,19 @@ void FDatasmithFacadeMasterMaterial::AddBoolean(
 	MaterialPropertyPtr->SetValue(bInPropertyValue ? TEXT("True") : TEXT("False"));
 
 	// Add the new property to the Datasmith material properties.
-	GetDatasmithMasterMaterial()->AddProperty(MaterialPropertyPtr);
+	GetDatasmithMaterialInstance()->AddProperty(MaterialPropertyPtr);
 }
 
-int32 FDatasmithFacadeMasterMaterial::GetPropertiesCount() const
+int32 FDatasmithFacadeMaterialInstance::GetPropertiesCount() const
 {
-	return GetDatasmithMasterMaterial()->GetPropertiesCount();
+	return GetDatasmithMaterialInstance()->GetPropertiesCount();
 }
 
-FDatasmithFacadeKeyValueProperty* FDatasmithFacadeMasterMaterial::GetNewProperty(
+FDatasmithFacadeKeyValueProperty* FDatasmithFacadeMaterialInstance::GetNewProperty(
 	int32 PropertyIndex
 ) const
 {
-	if (const TSharedPtr<IDatasmithKeyValueProperty>& Property = GetDatasmithMasterMaterial()->GetProperty(PropertyIndex))
+	if (const TSharedPtr<IDatasmithKeyValueProperty>& Property = GetDatasmithMaterialInstance()->GetProperty(PropertyIndex))
 	{
 		return new FDatasmithFacadeKeyValueProperty(Property.ToSharedRef());
 	}
@@ -236,11 +236,11 @@ FDatasmithFacadeKeyValueProperty* FDatasmithFacadeMasterMaterial::GetNewProperty
 	}
 }
 
-FDatasmithFacadeKeyValueProperty* FDatasmithFacadeMasterMaterial::GetNewPropertyByName(
+FDatasmithFacadeKeyValueProperty* FDatasmithFacadeMaterialInstance::GetNewPropertyByName(
 	const TCHAR* PropertyName
 ) const
 {
-	if (const TSharedPtr<IDatasmithKeyValueProperty>& Property = GetDatasmithMasterMaterial()->GetPropertyByName(PropertyName))
+	if (const TSharedPtr<IDatasmithKeyValueProperty>& Property = GetDatasmithMaterialInstance()->GetPropertyByName(PropertyName))
 	{
 		return new FDatasmithFacadeKeyValueProperty(Property.ToSharedRef());
 	}
@@ -250,7 +250,7 @@ FDatasmithFacadeKeyValueProperty* FDatasmithFacadeMasterMaterial::GetNewProperty
 	}
 }
 
-TSharedRef<IDatasmithMasterMaterialElement> FDatasmithFacadeMasterMaterial::GetDatasmithMasterMaterial() const
+TSharedRef<IDatasmithMaterialInstanceElement> FDatasmithFacadeMaterialInstance::GetDatasmithMaterialInstance() const
 {
-	return StaticCastSharedRef<IDatasmithMasterMaterialElement>( InternalDatasmithElement );
+	return StaticCastSharedRef<IDatasmithMaterialInstanceElement>( InternalDatasmithElement );
 }
\ No newline at end of file
diff --git a/Engine/Source/Developer/Datasmith/DatasmithFacade/Public/DatasmithFacadeMaterial.h b/Engine/Source/Developer/Datasmith/DatasmithFacade/Public/DatasmithFacadeMaterial.h
--- a/Engine/Source/Developer/Datasmith/DatasmithFacade/Public/DatasmithFacadeMaterial.h
+++ b/Engine/Source/Developer/Datasmith/DatasmithFacade/Public/DatasmithFacadeMaterial.h
@@ -14,7 +14,7 @@ class DATASMITHFACADE_API FDatasmithFacadeBaseMaterial :
 public:
 	enum class EDatasmithMaterialType
 	{
-		MasterMaterial,
+		MaterialInstance,
 		UEPbrMaterial,
 		DecalMaterial,
 		Unsupported,
@@ -40,15 +40,15 @@ protected:
 	);
 };
 
-class DATASMITHFACADE_API FDatasmithFacadeMasterMaterial :
+class DATASMITHFACADE_API FDatasmithFacadeMaterialInstance :
 	public FDatasmithFacadeBaseMaterial
 {
 	friend class FDatasmithFacadeScene;
 
 public:
 
-	// Possible Datasmith master material types, from EDatasmithMasterMaterialType in DatasmithDefinitions.h
-	enum class EMasterMaterialType
+	// Possible Datasmith master material types, from EDatasmithMaterialInstanceType in DatasmithDefinitions.h
+	enum class EMaterialInstanceType
 	{
 		Auto,
 		Opaque,
@@ -64,7 +64,7 @@ public:
 		Count
 	};
 
-	enum class EMasterMaterialQuality : uint8
+	enum class EMaterialInstanceQuality : uint8
 	{
 		High,
 		Low,
@@ -74,26 +74,26 @@ public:
 
 public:
 
-	FDatasmithFacadeMasterMaterial(
+	FDatasmithFacadeMaterialInstance(
 		const TCHAR* InElementName // Datasmith element name
 	);
 
-	virtual ~FDatasmithFacadeMasterMaterial() {}
+	virtual ~FDatasmithFacadeMaterialInstance() {}
 
 	/** Get the Datasmith master material type. */
-	EMasterMaterialType GetMaterialType() const;
+	EMaterialInstanceType GetMaterialType() const;
 
 	/** Set the Datasmith master material type. */
 	void SetMaterialType(
-		EMasterMaterialType InMasterMaterialType // master material type
+		EMaterialInstanceType InMaterialInstanceType // master material type
 	);
 
 	/** Get the Datasmith master material quality. */
-	EMasterMaterialQuality GetQuality() const;
+	EMaterialInstanceQuality GetQuality() const;
 
 	/** Set the Datasmith master material quality. */
 	void SetQuality(
-		EMasterMaterialQuality InQuality
+		EMaterialInstanceQuality InQuality
 	);
 
 	/** Get the material path name used when master material type is set to Custom */
@@ -162,9 +162,9 @@ public:
 protected:
 #endif
 
-	FDatasmithFacadeMasterMaterial(
-		const TSharedRef<IDatasmithMasterMaterialElement>& InMaterialRef // Datasmith master material element
+	FDatasmithFacadeMaterialInstance(
+		const TSharedRef<IDatasmithMaterialInstanceElement>& InMaterialRef // Datasmith master material element
 	);
 
-	TSharedRef<IDatasmithMasterMaterialElement> GetDatasmithMasterMaterial() const;
+	TSharedRef<IDatasmithMaterialInstanceElement> GetDatasmithMaterialInstance() const;
 };
diff --git a/Engine/Source/Programs/Enterprise/Datasmith/DatasmithARCHICADExporter/DatasmithValidator/Private/SceneValidator.cpp b/Engine/Source/Programs/Enterprise/Datasmith/DatasmithARCHICADExporter/DatasmithValidator/Private/SceneValidator.cpp
--- a/Engine/Source/Programs/Enterprise/Datasmith/DatasmithARCHICADExporter/DatasmithValidator/Private/SceneValidator.cpp
+++ b/Engine/Source/Programs/Enterprise/Datasmith/DatasmithARCHICADExporter/DatasmithValidator/Private/SceneValidator.cpp
@@ -84,9 +84,9 @@ FString FSceneValidator::GetElementTypes(const IDatasmithElement& InElement)
 	{
 		TypesString += TEXT(", BaseMaterial");
 	}
-	if (InElement.IsA(EDatasmithElementType::MasterMaterial))
+	if (InElement.IsA(EDatasmithElementType::MaterialInstance))
 	{
-		TypesString += TEXT(", MasterMaterial");
+		TypesString += TEXT(", MaterialInstance");
 	}
 	if (InElement.IsA(EDatasmithElementType::KeyValueProperty))
 	{
diff --git a/Engine/Source/Programs/Enterprise/Datasmith/DatasmithFacadeCSharp/Private/DatasmithFacadeCSharp.cpp b/Engine/Source/Programs/Enterprise/Datasmith/DatasmithFacadeCSharp/Private/DatasmithFacadeCSharp.cpp
--- a/Engine/Source/Programs/Enterprise/Datasmith/DatasmithFacadeCSharp/Private/DatasmithFacadeCSharp.cpp
+++ b/Engine/Source/Programs/Enterprise/Datasmith/DatasmithFacadeCSharp/Private/DatasmithFacadeCSharp.cpp
@@ -787,101 +787,101 @@ SWIGEXPORT void SWIGSTDCALL CSharp_delete_FDatasmithFacadeBaseMaterial(void * ja
 }
 
 
-SWIGEXPORT void * SWIGSTDCALL CSharp_new_FDatasmithFacadeMasterMaterial(wchar_t * jarg1) {
+SWIGEXPORT void * SWIGSTDCALL CSharp_new_FDatasmithFacadeMaterialInstance(wchar_t * jarg1) {
   void * jresult ;
   TCHAR *arg1 = (TCHAR *) 0 ;
-  FDatasmithFacadeMasterMaterial *result = 0 ;
+  FDatasmithFacadeMaterialInstance *result = 0 ;
   
   arg1 = (TCHAR *)jarg1; 
-  result = (FDatasmithFacadeMasterMaterial *)new FDatasmithFacadeMasterMaterial((TCHAR const *)arg1);
+  result = (FDatasmithFacadeMaterialInstance *)new FDatasmithFacadeMaterialInstance((TCHAR const *)arg1);
   jresult = (void *)result; 
   return jresult;
 }
 
 
-SWIGEXPORT void SWIGSTDCALL CSharp_delete_FDatasmithFacadeMasterMaterial(void * jarg1) {
-  FDatasmithFacadeMasterMaterial *arg1 = (FDatasmithFacadeMasterMaterial *) 0 ;
+SWIGEXPORT void SWIGSTDCALL CSharp_delete_FDatasmithFacadeMaterialInstance(void * jarg1) {
+  FDatasmithFacadeMaterialInstance *arg1 = (FDatasmithFacadeMaterialInstance *) 0 ;
   
-  arg1 = (FDatasmithFacadeMasterMaterial *)jarg1; 
+  arg1 = (FDatasmithFacadeMaterialInstance *)jarg1; 
   delete arg1;
 }
 
 
-SWIGEXPORT int SWIGSTDCALL CSharp_FDatasmithFacadeMasterMaterial_GetMaterialType(void * jarg1) {
+SWIGEXPORT int SWIGSTDCALL CSharp_FDatasmithFacadeMaterialInstance_GetMaterialType(void * jarg1) {
   int jresult ;
-  FDatasmithFacadeMasterMaterial *arg1 = (FDatasmithFacadeMasterMaterial *) 0 ;
-  FDatasmithFacadeMasterMaterial::EMasterMaterialType result;
+  FDatasmithFacadeMaterialInstance *arg1 = (FDatasmithFacadeMaterialInstance *) 0 ;
+  FDatasmithFacadeMaterialInstance::EMaterialInstanceType result;
   
-  arg1 = (FDatasmithFacadeMasterMaterial *)jarg1; 
-  result = (FDatasmithFacadeMasterMaterial::EMasterMaterialType)((FDatasmithFacadeMasterMaterial const *)arg1)->GetMaterialType();
+  arg1 = (FDatasmithFacadeMaterialInstance *)jarg1; 
+  result = (FDatasmithFacadeMaterialInstance::EMaterialInstanceType)((FDatasmithFacadeMaterialInstance const *)arg1)->GetMaterialType();
   jresult = (int)result; 
   return jresult;
 }
 
 
-SWIGEXPORT void SWIGSTDCALL CSharp_FDatasmithFacadeMasterMaterial_SetMaterialType(void * jarg1, int jarg2) {
-  FDatasmithFacadeMasterMaterial *arg1 = (FDatasmithFacadeMasterMaterial *) 0 ;
-  FDatasmithFacadeMasterMaterial::EMasterMaterialType arg2 ;
+SWIGEXPORT void SWIGSTDCALL CSharp_FDatasmithFacadeMaterialInstance_SetMaterialType(void * jarg1, int jarg2) {
+  FDatasmithFacadeMaterialInstance *arg1 = (FDatasmithFacadeMaterialInstance *) 0 ;
+  FDatasmithFacadeMaterialInstance::EMaterialInstanceType arg2 ;
   
-  arg1 = (FDatasmithFacadeMasterMaterial *)jarg1; 
-  arg2 = (FDatasmithFacadeMasterMaterial::EMasterMaterialType)jarg2; 
+  arg1 = (FDatasmithFacadeMaterialInstance *)jarg1; 
+  arg2 = (FDatasmithFacadeMaterialInstance::EMaterialInstanceType)jarg2; 
   (arg1)->SetMaterialType(arg2);
 }
 
 
-SWIGEXPORT int SWIGSTDCALL CSharp_FDatasmithFacadeMasterMaterial_GetQuality(void * jarg1) {
+SWIGEXPORT int SWIGSTDCALL CSharp_FDatasmithFacadeMaterialInstance_GetQuality(void * jarg1) {
   int jresult ;
-  FDatasmithFacadeMasterMaterial *arg1 = (FDatasmithFacadeMasterMaterial *) 0 ;
-  FDatasmithFacadeMasterMaterial::EMasterMaterialQuality result;
+  FDatasmithFacadeMaterialInstance *arg1 = (FDatasmithFacadeMaterialInstance *) 0 ;
+  FDatasmithFacadeMaterialInstance::EMaterialInstanceQuality result;
   
-  arg1 = (FDatasmithFacadeMasterMaterial *)jarg1; 
-  result = (FDatasmithFacadeMasterMaterial::EMasterMaterialQuality)((FDatasmithFacadeMasterMaterial const *)arg1)->GetQuality();
+  arg1 = (FDatasmithFacadeMaterialInstance *)jarg1; 
+  result = (FDatasmithFacadeMaterialInstance::EMaterialInstanceQuality)((FDatasmithFacadeMaterialInstance const *)arg1)->GetQuality();
   jresult = (int)result; 
   return jresult;
 }
 
 
-SWIGEXPORT void SWIGSTDCALL CSharp_FDatasmithFacadeMasterMaterial_SetQuality(void * jarg1, int jarg2) {
-  FDatasmithFacadeMasterMaterial *arg1 = (FDatasmithFacadeMasterMaterial *) 0 ;
-  FDatasmithFacadeMasterMaterial::EMasterMaterialQuality arg2 ;
+SWIGEXPORT void SWIGSTDCALL CSharp_FDatasmithFacadeMaterialInstance_SetQuality(void * jarg1, int jarg2) {
+  FDatasmithFacadeMaterialInstance *arg1 = (FDatasmithFacadeMaterialInstance *) 0 ;
+  FDatasmithFacadeMaterialInstance::EMaterialInstanceQuality arg2 ;
   
-  arg1 = (FDatasmithFacadeMasterMaterial *)jarg1; 
-  arg2 = (FDatasmithFacadeMasterMaterial::EMasterMaterialQuality)jarg2; 
+  arg1 = (FDatasmithFacadeMaterialInstance *)jarg1; 
+  arg2 = (FDatasmithFacadeMaterialInstance::EMaterialInstanceQuality)jarg2; 
   (arg1)->SetQuality(arg2);
 }
 
 
-SWIGEXPORT wchar_t * SWIGSTDCALL CSharp_FDatasmithFacadeMasterMaterial_GetCustomMaterialPathName(void * jarg1) {
+SWIGEXPORT wchar_t * SWIGSTDCALL CSharp_FDatasmithFacadeMaterialInstance_GetCustomMaterialPathName(void * jarg1) {
   wchar_t * jresult ;
-  FDatasmithFacadeMasterMaterial *arg1 = (FDatasmithFacadeMasterMaterial *) 0 ;
+  FDatasmithFacadeMaterialInstance *arg1 = (FDatasmithFacadeMaterialInstance *) 0 ;
   TCHAR *result = 0 ;
   
-  arg1 = (FDatasmithFacadeMasterMaterial *)jarg1; 
-  result = (TCHAR *)((FDatasmithFacadeMasterMaterial const *)arg1)->GetCustomMaterialPathName();
+  arg1 = (FDatasmithFacadeMaterialInstance *)jarg1; 
+  result = (TCHAR *)((FDatasmithFacadeMaterialInstance const *)arg1)->GetCustomMaterialPathName();
   jresult = (wchar_t *)result; 
   return jresult;
 }
 
 
-SWIGEXPORT void SWIGSTDCALL CSharp_FDatasmithFacadeMasterMaterial_SetCustomMaterialPathName(void * jarg1, wchar_t * jarg2) {
-  FDatasmithFacadeMasterMaterial *arg1 = (FDatasmithFacadeMasterMaterial *) 0 ;
+SWIGEXPORT void SWIGSTDCALL CSharp_FDatasmithFacadeMaterialInstance_SetCustomMaterialPathName(void * jarg1, wchar_t * jarg2) {
+  FDatasmithFacadeMaterialInstance *arg1 = (FDatasmithFacadeMaterialInstance *) 0 ;
   TCHAR *arg2 = (TCHAR *) 0 ;
   
-  arg1 = (FDatasmithFacadeMasterMaterial *)jarg1; 
+  arg1 = (FDatasmithFacadeMaterialInstance *)jarg1; 
   arg2 = (TCHAR *)jarg2; 
   (arg1)->SetCustomMaterialPathName((TCHAR const *)arg2);
 }
 
 
-SWIGEXPORT void SWIGSTDCALL CSharp_FDatasmithFacadeMasterMaterial_AddColor__SWIG_0(void * jarg1, wchar_t * jarg2, unsigned char jarg3, unsigned char jarg4, unsigned char jarg5, unsigned char jarg6) {
-  FDatasmithFacadeMasterMaterial *arg1 = (FDatasmithFacadeMasterMaterial *) 0 ;
+SWIGEXPORT void SWIGSTDCALL CSharp_FDatasmithFacadeMaterialInstance_AddColor__SWIG_0(void * jarg1, wchar_t * jarg2, unsigned char jarg3, unsigned char jarg4, unsigned char jarg5, unsigned char jarg6) {
+  FDatasmithFacadeMaterialInstance *arg1 = (FDatasmithFacadeMaterialInstance *) 0 ;
   TCHAR *arg2 = (TCHAR *) 0 ;
   unsigned char arg3 ;
   unsigned char arg4 ;
   unsigned char arg5 ;
   unsigned char arg6 ;
   
-  arg1 = (FDatasmithFacadeMasterMaterial *)jarg1; 
+  arg1 = (FDatasmithFacadeMaterialInstance *)jarg1; 
   arg2 = (TCHAR *)jarg2; 
   arg3 = (unsigned char)jarg3; 
   arg4 = (unsigned char)jarg4; 
@@ -891,15 +891,15 @@ SWIGEXPORT void SWIGSTDCALL CSharp_FDatasmithFacadeMasterMaterial_AddColor__SWIG
 }
 
 
-SWIGEXPORT void SWIGSTDCALL CSharp_FDatasmithFacadeMasterMaterial_AddColor__SWIG_1(void * jarg1, wchar_t * jarg2, float jarg3, float jarg4, float jarg5, float jarg6) {
-  FDatasmithFacadeMasterMaterial *arg1 = (FDatasmithFacadeMasterMaterial *) 0 ;
+SWIGEXPORT void SWIGSTDCALL CSharp_FDatasmithFacadeMaterialInstance_AddColor__SWIG_1(void * jarg1, wchar_t * jarg2, float jarg3, float jarg4, float jarg5, float jarg6) {
+  FDatasmithFacadeMaterialInstance *arg1 = (FDatasmithFacadeMaterialInstance *) 0 ;
   TCHAR *arg2 = (TCHAR *) 0 ;
   float arg3 ;
   float arg4 ;
   float arg5 ;
   float arg6 ;
   
-  arg1 = (FDatasmithFacadeMasterMaterial *)jarg1; 
+  arg1 = (FDatasmithFacadeMaterialInstance *)jarg1; 
   arg2 = (TCHAR *)jarg2; 
   arg3 = (float)jarg3; 
   arg4 = (float)jarg4; 
@@ -909,89 +909,89 @@ SWIGEXPORT void SWIGSTDCALL CSharp_FDatasmithFacadeMasterMaterial_AddColor__SWIG
 }
 
 
-SWIGEXPORT void SWIGSTDCALL CSharp_FDatasmithFacadeMasterMaterial_AddTexture(void * jarg1, wchar_t * jarg2, void * jarg3) {
-  FDatasmithFacadeMasterMaterial *arg1 = (FDatasmithFacadeMasterMaterial *) 0 ;
+SWIGEXPORT void SWIGSTDCALL CSharp_FDatasmithFacadeMaterialInstance_AddTexture(void * jarg1, wchar_t * jarg2, void * jarg3) {
+  FDatasmithFacadeMaterialInstance *arg1 = (FDatasmithFacadeMaterialInstance *) 0 ;
   TCHAR *arg2 = (TCHAR *) 0 ;
   FDatasmithFacadeTexture *arg3 = (FDatasmithFacadeTexture *) 0 ;
   
-  arg1 = (FDatasmithFacadeMasterMaterial *)jarg1; 
+  arg1 = (FDatasmithFacadeMaterialInstance *)jarg1; 
   arg2 = (TCHAR *)jarg2; 
   arg3 = (FDatasmithFacadeTexture *)jarg3; 
   (arg1)->AddTexture((TCHAR const *)arg2,(FDatasmithFacadeTexture const *)arg3);
 }
 
 
-SWIGEXPORT void SWIGSTDCALL CSharp_FDatasmithFacadeMasterMaterial_AddString(void * jarg1, wchar_t * jarg2, wchar_t * jarg3) {
-  FDatasmithFacadeMasterMaterial *arg1 = (FDatasmithFacadeMasterMaterial *) 0 ;
+SWIGEXPORT void SWIGSTDCALL CSharp_FDatasmithFacadeMaterialInstance_AddString(void * jarg1, wchar_t * jarg2, wchar_t * jarg3) {
+  FDatasmithFacadeMaterialInstance *arg1 = (FDatasmithFacadeMaterialInstance *) 0 ;
   TCHAR *arg2 = (TCHAR *) 0 ;
   TCHAR *arg3 = (TCHAR *) 0 ;
   
-  arg1 = (FDatasmithFacadeMasterMaterial *)jarg1; 
+  arg1 = (FDatasmithFacadeMaterialInstance *)jarg1; 
   arg2 = (TCHAR *)jarg2; 
   arg3 = (TCHAR *)jarg3; 
   (arg1)->AddString((TCHAR const *)arg2,(TCHAR const *)arg3);
 }
 
 
-SWIGEXPORT void SWIGSTDCALL CSharp_FDatasmithFacadeMasterMaterial_AddFloat(void * jarg1, wchar_t * jarg2, float jarg3) {
-  FDatasmithFacadeMasterMaterial *arg1 = (FDatasmithFacadeMasterMaterial *) 0 ;
+SWIGEXPORT void SWIGSTDCALL CSharp_FDatasmithFacadeMaterialInstance_AddFloat(void * jarg1, wchar_t * jarg2, float jarg3) {
+  FDatasmithFacadeMaterialInstance *arg1 = (FDatasmithFacadeMaterialInstance *) 0 ;
   TCHAR *arg2 = (TCHAR *) 0 ;
   float arg3 ;
   
-  arg1 = (FDatasmithFacadeMasterMaterial *)jarg1; 
+  arg1 = (FDatasmithFacadeMaterialInstance *)jarg1; 
   arg2 = (TCHAR *)jarg2; 
   arg3 = (float)jarg3; 
   (arg1)->AddFloat((TCHAR const *)arg2,arg3);
 }
 
 
-SWIGEXPORT void SWIGSTDCALL CSharp_FDatasmithFacadeMasterMaterial_AddBoolean(void * jarg1, wchar_t * jarg2, unsigned int jarg3) {
-  FDatasmithFacadeMasterMaterial *arg1 = (FDatasmithFacadeMasterMaterial *) 0 ;
+SWIGEXPORT void SWIGSTDCALL CSharp_FDatasmithFacadeMaterialInstance_AddBoolean(void * jarg1, wchar_t * jarg2, unsigned int jarg3) {
+  FDatasmithFacadeMaterialInstance *arg1 = (FDatasmithFacadeMaterialInstance *) 0 ;
   TCHAR *arg2 = (TCHAR *) 0 ;
   bool arg3 ;
   
-  arg1 = (FDatasmithFacadeMasterMaterial *)jarg1; 
+  arg1 = (FDatasmithFacadeMaterialInstance *)jarg1; 
   arg2 = (TCHAR *)jarg2; 
   arg3 = jarg3 ? true : false; 
   (arg1)->AddBoolean((TCHAR const *)arg2,arg3);
 }
 
 
-SWIGEXPORT int SWIGSTDCALL CSharp_FDatasmithFacadeMasterMaterial_GetPropertiesCount(void * jarg1) {
+SWIGEXPORT int SWIGSTDCALL CSharp_FDatasmithFacadeMaterialInstance_GetPropertiesCount(void * jarg1) {
   int jresult ;
-  FDatasmithFacadeMasterMaterial *arg1 = (FDatasmithFacadeMasterMaterial *) 0 ;
+  FDatasmithFacadeMaterialInstance *arg1 = (FDatasmithFacadeMaterialInstance *) 0 ;
   int32 result;
   
-  arg1 = (FDatasmithFacadeMasterMaterial *)jarg1; 
-  result = (int32)((FDatasmithFacadeMasterMaterial const *)arg1)->GetPropertiesCount();
+  arg1 = (FDatasmithFacadeMaterialInstance *)jarg1; 
+  result = (int32)((FDatasmithFacadeMaterialInstance const *)arg1)->GetPropertiesCount();
   jresult = result; 
   return jresult;
 }
 
 
-SWIGEXPORT void * SWIGSTDCALL CSharp_FDatasmithFacadeMasterMaterial_GetProperty(void * jarg1, int jarg2) {
+SWIGEXPORT void * SWIGSTDCALL CSharp_FDatasmithFacadeMaterialInstance_GetProperty(void * jarg1, int jarg2) {
   void * jresult ;
-  FDatasmithFacadeMasterMaterial *arg1 = (FDatasmithFacadeMasterMaterial *) 0 ;
+  FDatasmithFacadeMaterialInstance *arg1 = (FDatasmithFacadeMaterialInstance *) 0 ;
   int32 arg2 ;
   FDatasmithFacadeKeyValueProperty *result = 0 ;
   
-  arg1 = (FDatasmithFacadeMasterMaterial *)jarg1; 
+  arg1 = (FDatasmithFacadeMaterialInstance *)jarg1; 
   arg2 = (int32)jarg2; 
-  result = (FDatasmithFacadeKeyValueProperty *)((FDatasmithFacadeMasterMaterial const *)arg1)->GetNewProperty(arg2);
+  result = (FDatasmithFacadeKeyValueProperty *)((FDatasmithFacadeMaterialInstance const *)arg1)->GetNewProperty(arg2);
   jresult = (void *)result; 
   return jresult;
 }
 
 
-SWIGEXPORT void * SWIGSTDCALL CSharp_FDatasmithFacadeMasterMaterial_GetPropertyByName(void * jarg1, wchar_t * jarg2) {
+SWIGEXPORT void * SWIGSTDCALL CSharp_FDatasmithFacadeMaterialInstance_GetPropertyByName(void * jarg1, wchar_t * jarg2) {
   void * jresult ;
-  FDatasmithFacadeMasterMaterial *arg1 = (FDatasmithFacadeMasterMaterial *) 0 ;
+  FDatasmithFacadeMaterialInstance *arg1 = (FDatasmithFacadeMaterialInstance *) 0 ;
   TCHAR *arg2 = (TCHAR *) 0 ;
   FDatasmithFacadeKeyValueProperty *result = 0 ;
   
-  arg1 = (FDatasmithFacadeMasterMaterial *)jarg1; 
+  arg1 = (FDatasmithFacadeMaterialInstance *)jarg1; 
   arg2 = (TCHAR *)jarg2; 
-  result = (FDatasmithFacadeKeyValueProperty *)((FDatasmithFacadeMasterMaterial const *)arg1)->GetNewPropertyByName((TCHAR const *)arg2);
+  result = (FDatasmithFacadeKeyValueProperty *)((FDatasmithFacadeMaterialInstance const *)arg1)->GetNewPropertyByName((TCHAR const *)arg2);
   jresult = (void *)result; 
   return jresult;
 }
@@ -7517,7 +7517,7 @@ SWIGEXPORT FDatasmithFacadeElement * SWIGSTDCALL CSharp_FDatasmithFacadeBaseMate
     return (FDatasmithFacadeElement *)jarg1;
 }
 
-SWIGEXPORT FDatasmithFacadeBaseMaterial * SWIGSTDCALL CSharp_FDatasmithFacadeMasterMaterial_SWIGUpcast(FDatasmithFacadeMasterMaterial *jarg1) {
+SWIGEXPORT FDatasmithFacadeBaseMaterial * SWIGSTDCALL CSharp_FDatasmithFacadeMaterialInstance_SWIGUpcast(FDatasmithFacadeMaterialInstance *jarg1) {
     return (FDatasmithFacadeBaseMaterial *)jarg1;
 }
 
diff --git a/Engine/Source/Programs/Enterprise/Datasmith/DatasmithFacadeCSharp/Public/DatasmithFacadeCSharpPINVOKE.cs b/Engine/Source/Programs/Enterprise/Datasmith/DatasmithFacadeCSharp/Public/DatasmithFacadeCSharpPINVOKE.cs
--- a/Engine/Source/Programs/Enterprise/Datasmith/DatasmithFacadeCSharp/Public/DatasmithFacadeCSharpPINVOKE.cs
+++ b/Engine/Source/Programs/Enterprise/Datasmith/DatasmithFacadeCSharp/Public/DatasmithFacadeCSharpPINVOKE.cs
@@ -335,56 +335,56 @@ class DatasmithFacadeCSharpPINVOKE {
   [global::System.Runtime.InteropServices.DllImport("DatasmithFacadeCSharp", EntryPoint="CSharp_delete_FDatasmithFacadeBaseMaterial")]
   public static extern void delete_FDatasmithFacadeBaseMaterial(global::System.Runtime.InteropServices.HandleRef jarg1);
 
-  [global::System.Runtime.InteropServices.DllImport("DatasmithFacadeCSharp", EntryPoint="CSharp_new_FDatasmithFacadeMasterMaterial")]
-  public static extern global::System.IntPtr new_FDatasmithFacadeMasterMaterial([global::System.Runtime.InteropServices.MarshalAs(global::System.Runtime.InteropServices.UnmanagedType.LPWStr)]string jarg1);
+  [global::System.Runtime.InteropServices.DllImport("DatasmithFacadeCSharp", EntryPoint="CSharp_new_FDatasmithFacadeMaterialInstance")]
+  public static extern global::System.IntPtr new_FDatasmithFacadeMaterialInstance([global::System.Runtime.InteropServices.MarshalAs(global::System.Runtime.InteropServices.UnmanagedType.LPWStr)]string jarg1);
 
-  [global::System.Runtime.InteropServices.DllImport("DatasmithFacadeCSharp", EntryPoint="CSharp_delete_FDatasmithFacadeMasterMaterial")]
-  public static extern void delete_FDatasmithFacadeMasterMaterial(global::System.Runtime.InteropServices.HandleRef jarg1);
+  [global::System.Runtime.InteropServices.DllImport("DatasmithFacadeCSharp", EntryPoint="CSharp_delete_FDatasmithFacadeMaterialInstance")]
+  public static extern void delete_FDatasmithFacadeMaterialInstance(global::System.Runtime.InteropServices.HandleRef jarg1);
 
-  [global::System.Runtime.InteropServices.DllImport("DatasmithFacadeCSharp", EntryPoint="CSharp_FDatasmithFacadeMasterMaterial_GetMaterialType")]
-  public static extern int FDatasmithFacadeMasterMaterial_GetMaterialType(global::System.Runtime.InteropServices.HandleRef jarg1);
+  [global::System.Runtime.InteropServices.DllImport("DatasmithFacadeCSharp", EntryPoint="CSharp_FDatasmithFacadeMaterialInstance_GetMaterialType")]
+  public static extern int FDatasmithFacadeMaterialInstance_GetMaterialType(global::System.Runtime.InteropServices.HandleRef jarg1);
 
-  [global::System.Runtime.InteropServices.DllImport("DatasmithFacadeCSharp", EntryPoint="CSharp_FDatasmithFacadeMasterMaterial_SetMaterialType")]
-  public static extern void FDatasmithFacadeMasterMaterial_SetMaterialType(global::System.Runtime.InteropServices.HandleRef jarg1, int jarg2);
+  [global::System.Runtime.InteropServices.DllImport("DatasmithFacadeCSharp", EntryPoint="CSharp_FDatasmithFacadeMaterialInstance_SetMaterialType")]
+  public static extern void FDatasmithFacadeMaterialInstance_SetMaterialType(global::System.Runtime.InteropServices.HandleRef jarg1, int jarg2);
 
-  [global::System.Runtime.InteropServices.DllImport("DatasmithFacadeCSharp", EntryPoint="CSharp_FDatasmithFacadeMasterMaterial_GetQuality")]
-  public static extern int FDatasmithFacadeMasterMaterial_GetQuality(global::System.Runtime.InteropServices.HandleRef jarg1);
+  [global::System.Runtime.InteropServices.DllImport("DatasmithFacadeCSharp", EntryPoint="CSharp_FDatasmithFacadeMaterialInstance_GetQuality")]
+  public static extern int FDatasmithFacadeMaterialInstance_GetQuality(global::System.Runtime.InteropServices.HandleRef jarg1);
 
-  [global::System.Runtime.InteropServices.DllImport("DatasmithFacadeCSharp", EntryPoint="CSharp_FDatasmithFacadeMasterMaterial_SetQuality")]
-  public static extern void FDatasmithFacadeMasterMaterial_SetQuality(global::System.Runtime.InteropServices.HandleRef jarg1, int jarg2);
+  [global::System.Runtime.InteropServices.DllImport("DatasmithFacadeCSharp", EntryPoint="CSharp_FDatasmithFacadeMaterialInstance_SetQuality")]
+  public static extern void FDatasmithFacadeMaterialInstance_SetQuality(global::System.Runtime.InteropServices.HandleRef jarg1, int jarg2);
 
-  [global::System.Runtime.InteropServices.DllImport("DatasmithFacadeCSharp", EntryPoint="CSharp_FDatasmithFacadeMasterMaterial_GetCustomMaterialPathName")]
-  public static extern global::System.IntPtr FDatasmithFacadeMasterMaterial_GetCustomMaterialPathName(global::System.Runtime.InteropServices.HandleRef jarg1);
+  [global::System.Runtime.InteropServices.DllImport("DatasmithFacadeCSharp", EntryPoint="CSharp_FDatasmithFacadeMaterialInstance_GetCustomMaterialPathName")]
+  public static extern global::System.IntPtr FDatasmithFacadeMaterialInstance_GetCustomMaterialPathName(global::System.Runtime.InteropServices.HandleRef jarg1);
 
-  [global::System.Runtime.InteropServices.DllImport("DatasmithFacadeCSharp", EntryPoint="CSharp_FDatasmithFacadeMasterMaterial_SetCustomMaterialPathName")]
-  public static extern void FDatasmithFacadeMasterMaterial_SetCustomMaterialPathName(global::System.Runtime.InteropServices.HandleRef jarg1, [global::System.Runtime.InteropServices.MarshalAs(global::System.Runtime.InteropServices.UnmanagedType.LPWStr)]string jarg2);
+  [global::System.Runtime.InteropServices.DllImport("DatasmithFacadeCSharp", EntryPoint="CSharp_FDatasmithFacadeMaterialInstance_SetCustomMaterialPathName")]
+  public static extern void FDatasmithFacadeMaterialInstance_SetCustomMaterialPathName(global::System.Runtime.InteropServices.HandleRef jarg1, [global::System.Runtime.InteropServices.MarshalAs(global::System.Runtime.InteropServices.UnmanagedType.LPWStr)]string jarg2);
 
-  [global::System.Runtime.InteropServices.DllImport("DatasmithFacadeCSharp", EntryPoint="CSharp_FDatasmithFacadeMasterMaterial_AddColor__SWIG_0")]
-  public static extern void FDatasmithFacadeMasterMaterial_AddColor__SWIG_0(global::System.Runtime.InteropServices.HandleRef jarg1, [global::System.Runtime.InteropServices.MarshalAs(global::System.Runtime.InteropServices.UnmanagedType.LPWStr)]string jarg2, byte jarg3, byte jarg4, byte jarg5, byte jarg6);
+  [global::System.Runtime.InteropServices.DllImport("DatasmithFacadeCSharp", EntryPoint="CSharp_FDatasmithFacadeMaterialInstance_AddColor__SWIG_0")]
+  public static extern void FDatasmithFacadeMaterialInstance_AddColor__SWIG_0(global::System.Runtime.InteropServices.HandleRef jarg1, [global::System.Runtime.InteropServices.MarshalAs(global::System.Runtime.InteropServices.UnmanagedType.LPWStr)]string jarg2, byte jarg3, byte jarg4, byte jarg5, byte jarg6);
 
-  [global::System.Runtime.InteropServices.DllImport("DatasmithFacadeCSharp", EntryPoint="CSharp_FDatasmithFacadeMasterMaterial_AddColor__SWIG_1")]
-  public static extern void FDatasmithFacadeMasterMaterial_AddColor__SWIG_1(global::System.Runtime.InteropServices.HandleRef jarg1, [global::System.Runtime.InteropServices.MarshalAs(global::System.Runtime.InteropServices.UnmanagedType.LPWStr)]string jarg2, float jarg3, float jarg4, float jarg5, float jarg6);
+  [global::System.Runtime.InteropServices.DllImport("DatasmithFacadeCSharp", EntryPoint="CSharp_FDatasmithFacadeMaterialInstance_AddColor__SWIG_1")]
+  public static extern void FDatasmithFacadeMaterialInstance_AddColor__SWIG_1(global::System.Runtime.InteropServices.HandleRef jarg1, [global::System.Runtime.InteropServices.MarshalAs(global::System.Runtime.InteropServices.UnmanagedType.LPWStr)]string jarg2, float jarg3, float jarg4, float jarg5, float jarg6);
 
-  [global::System.Runtime.InteropServices.DllImport("DatasmithFacadeCSharp", EntryPoint="CSharp_FDatasmithFacadeMasterMaterial_AddTexture")]
-  public static extern void FDatasmithFacadeMasterMaterial_AddTexture(global::System.Runtime.InteropServices.HandleRef jarg1, [global::System.Runtime.InteropServices.MarshalAs(global::System.Runtime.InteropServices.UnmanagedType.LPWStr)]string jarg2, global::System.Runtime.InteropServices.HandleRef jarg3);
+  [global::System.Runtime.InteropServices.DllImport("DatasmithFacadeCSharp", EntryPoint="CSharp_FDatasmithFacadeMaterialInstance_AddTexture")]
+  public static extern void FDatasmithFacadeMaterialInstance_AddTexture(global::System.Runtime.InteropServices.HandleRef jarg1, [global::System.Runtime.InteropServices.MarshalAs(global::System.Runtime.InteropServices.UnmanagedType.LPWStr)]string jarg2, global::System.Runtime.InteropServices.HandleRef jarg3);
 
-  [global::System.Runtime.InteropServices.DllImport("DatasmithFacadeCSharp", EntryPoint="CSharp_FDatasmithFacadeMasterMaterial_AddString")]
-  public static extern void FDatasmithFacadeMasterMaterial_AddString(global::System.Runtime.InteropServices.HandleRef jarg1, [global::System.Runtime.InteropServices.MarshalAs(global::System.Runtime.InteropServices.UnmanagedType.LPWStr)]string jarg2, [global::System.Runtime.InteropServices.MarshalAs(global::System.Runtime.InteropServices.UnmanagedType.LPWStr)]string jarg3);
+  [global::System.Runtime.InteropServices.DllImport("DatasmithFacadeCSharp", EntryPoint="CSharp_FDatasmithFacadeMaterialInstance_AddString")]
+  public static extern void FDatasmithFacadeMaterialInstance_AddString(global::System.Runtime.InteropServices.HandleRef jarg1, [global::System.Runtime.InteropServices.MarshalAs(global::System.Runtime.InteropServices.UnmanagedType.LPWStr)]string jarg2, [global::System.Runtime.InteropServices.MarshalAs(global::System.Runtime.InteropServices.UnmanagedType.LPWStr)]string jarg3);
 
-  [global::System.Runtime.InteropServices.DllImport("DatasmithFacadeCSharp", EntryPoint="CSharp_FDatasmithFacadeMasterMaterial_AddFloat")]
-  public static extern void FDatasmithFacadeMasterMaterial_AddFloat(global::System.Runtime.InteropServices.HandleRef jarg1, [global::System.Runtime.InteropServices.MarshalAs(global::System.Runtime.InteropServices.UnmanagedType.LPWStr)]string jarg2, float jarg3);
+  [global::System.Runtime.InteropServices.DllImport("DatasmithFacadeCSharp", EntryPoint="CSharp_FDatasmithFacadeMaterialInstance_AddFloat")]
+  public static extern void FDatasmithFacadeMaterialInstance_AddFloat(global::System.Runtime.InteropServices.HandleRef jarg1, [global::System.Runtime.InteropServices.MarshalAs(global::System.Runtime.InteropServices.UnmanagedType.LPWStr)]string jarg2, float jarg3);
 
-  [global::System.Runtime.InteropServices.DllImport("DatasmithFacadeCSharp", EntryPoint="CSharp_FDatasmithFacadeMasterMaterial_AddBoolean")]
-  public static extern void FDatasmithFacadeMasterMaterial_AddBoolean(global::System.Runtime.InteropServices.HandleRef jarg1, [global::System.Runtime.InteropServices.MarshalAs(global::System.Runtime.InteropServices.UnmanagedType.LPWStr)]string jarg2, bool jarg3);
+  [global::System.Runtime.InteropServices.DllImport("DatasmithFacadeCSharp", EntryPoint="CSharp_FDatasmithFacadeMaterialInstance_AddBoolean")]
+  public static extern void FDatasmithFacadeMaterialInstance_AddBoolean(global::System.Runtime.InteropServices.HandleRef jarg1, [global::System.Runtime.InteropServices.MarshalAs(global::System.Runtime.InteropServices.UnmanagedType.LPWStr)]string jarg2, bool jarg3);
 
-  [global::System.Runtime.InteropServices.DllImport("DatasmithFacadeCSharp", EntryPoint="CSharp_FDatasmithFacadeMasterMaterial_GetPropertiesCount")]
-  public static extern int FDatasmithFacadeMasterMaterial_GetPropertiesCount(global::System.Runtime.InteropServices.HandleRef jarg1);
+  [global::System.Runtime.InteropServices.DllImport("DatasmithFacadeCSharp", EntryPoint="CSharp_FDatasmithFacadeMaterialInstance_GetPropertiesCount")]
+  public static extern int FDatasmithFacadeMaterialInstance_GetPropertiesCount(global::System.Runtime.InteropServices.HandleRef jarg1);
 
-  [global::System.Runtime.InteropServices.DllImport("DatasmithFacadeCSharp", EntryPoint="CSharp_FDatasmithFacadeMasterMaterial_GetProperty")]
-  public static extern global::System.IntPtr FDatasmithFacadeMasterMaterial_GetProperty(global::System.Runtime.InteropServices.HandleRef jarg1, int jarg2);
+  [global::System.Runtime.InteropServices.DllImport("DatasmithFacadeCSharp", EntryPoint="CSharp_FDatasmithFacadeMaterialInstance_GetProperty")]
+  public static extern global::System.IntPtr FDatasmithFacadeMaterialInstance_GetProperty(global::System.Runtime.InteropServices.HandleRef jarg1, int jarg2);
 
-  [global::System.Runtime.InteropServices.DllImport("DatasmithFacadeCSharp", EntryPoint="CSharp_FDatasmithFacadeMasterMaterial_GetPropertyByName")]
-  public static extern global::System.IntPtr FDatasmithFacadeMasterMaterial_GetPropertyByName(global::System.Runtime.InteropServices.HandleRef jarg1, [global::System.Runtime.InteropServices.MarshalAs(global::System.Runtime.InteropServices.UnmanagedType.LPWStr)]string jarg2);
+  [global::System.Runtime.InteropServices.DllImport("DatasmithFacadeCSharp", EntryPoint="CSharp_FDatasmithFacadeMaterialInstance_GetPropertyByName")]
+  public static extern global::System.IntPtr FDatasmithFacadeMaterialInstance_GetPropertyByName(global::System.Runtime.InteropServices.HandleRef jarg1, [global::System.Runtime.InteropServices.MarshalAs(global::System.Runtime.InteropServices.UnmanagedType.LPWStr)]string jarg2);
 
   [global::System.Runtime.InteropServices.DllImport("DatasmithFacadeCSharp", EntryPoint="CSharp_new_FDatasmithFacadeMaterialID")]
   public static extern global::System.IntPtr new_FDatasmithFacadeMaterialID([global::System.Runtime.InteropServices.MarshalAs(global::System.Runtime.InteropServices.UnmanagedType.LPWStr)]string jarg1);
@@ -2081,8 +2081,8 @@ class DatasmithFacadeCSharpPINVOKE {
   [global::System.Runtime.InteropServices.DllImport("DatasmithFacadeCSharp", EntryPoint="CSharp_FDatasmithFacadeBaseMaterial_SWIGUpcast")]
   public static extern global::System.IntPtr FDatasmithFacadeBaseMaterial_SWIGUpcast(global::System.IntPtr jarg1);
 
-  [global::System.Runtime.InteropServices.DllImport("DatasmithFacadeCSharp", EntryPoint="CSharp_FDatasmithFacadeMasterMaterial_SWIGUpcast")]
-  public static extern global::System.IntPtr FDatasmithFacadeMasterMaterial_SWIGUpcast(global::System.IntPtr jarg1);
+  [global::System.Runtime.InteropServices.DllImport("DatasmithFacadeCSharp", EntryPoint="CSharp_FDatasmithFacadeMaterialInstance_SWIGUpcast")]
+  public static extern global::System.IntPtr FDatasmithFacadeMaterialInstance_SWIGUpcast(global::System.IntPtr jarg1);
 
   [global::System.Runtime.InteropServices.DllImport("DatasmithFacadeCSharp", EntryPoint="CSharp_FDatasmithFacadeMaterialID_SWIGUpcast")]
   public static extern global::System.IntPtr FDatasmithFacadeMaterialID_SWIGUpcast(global::System.IntPtr jarg1);
diff --git a/Engine/Source/Programs/Enterprise/Datasmith/DatasmithFacadeCSharp/Public/FDatasmithFacadeBaseMaterial.cs b/Engine/Source/Programs/Enterprise/Datasmith/DatasmithFacadeCSharp/Public/FDatasmithFacadeBaseMaterial.cs
--- a/Engine/Source/Programs/Enterprise/Datasmith/DatasmithFacadeCSharp/Public/FDatasmithFacadeBaseMaterial.cs
+++ b/Engine/Source/Programs/Enterprise/Datasmith/DatasmithFacadeCSharp/Public/FDatasmithFacadeBaseMaterial.cs
@@ -41,7 +41,7 @@ public class FDatasmithFacadeBaseMaterial : FDatasmithFacadeElement {
   }
 
   public enum EDatasmithMaterialType {
-    MasterMaterial,
+    MaterialInstance,
     UEPbrMaterial,
     DecalMaterial,
     Unsupported
diff --git a/Engine/Source/Programs/Enterprise/Datasmith/DatasmithFacadeCSharp/Public/FDatasmithFacadeScene.cs b/Engine/Source/Programs/Enterprise/Datasmith/DatasmithFacadeCSharp/Public/FDatasmithFacadeScene.cs
--- a/Engine/Source/Programs/Enterprise/Datasmith/DatasmithFacadeCSharp/Public/FDatasmithFacadeScene.cs
+++ b/Engine/Source/Programs/Enterprise/Datasmith/DatasmithFacadeCSharp/Public/FDatasmithFacadeScene.cs
@@ -133,8 +133,8 @@ public class FDatasmithFacadeScene : global::System.IDisposable {
 		{
 		case FDatasmithFacadeBaseMaterial.EDatasmithMaterialType.UEPbrMaterial:
 			return new FDatasmithFacadeUEPbrMaterial(objectPtr, true);
-		case FDatasmithFacadeBaseMaterial.EDatasmithMaterialType.MasterMaterial:
-			return new FDatasmithFacadeMasterMaterial(objectPtr, true);
+		case FDatasmithFacadeBaseMaterial.EDatasmithMaterialType.MaterialInstance:
+			return new FDatasmithFacadeMaterialInstance(objectPtr, true);
 		case FDatasmithFacadeBaseMaterial.EDatasmithMaterialType.DecalMaterial:
 			return new FDatasmithFacadeDecalMaterial(objectPtr, true);
 		case FDatasmithFacadeBaseMaterial.EDatasmithMaterialType.Unsupported:
diff --git a/Engine/Source/Programs/Enterprise/Datasmith/DatasmithRevitExporter/Private/DatasmithRevitDocumentData.cs b/Engine/Source/Programs/Enterprise/Datasmith/DatasmithRevitExporter/Private/DatasmithRevitDocumentData.cs
--- a/Engine/Source/Programs/Enterprise/Datasmith/DatasmithRevitExporter/Private/DatasmithRevitDocumentData.cs
+++ b/Engine/Source/Programs/Enterprise/Datasmith/DatasmithRevitExporter/Private/DatasmithRevitDocumentData.cs
@@ -718,7 +718,7 @@ namespace DatasmithRevitExporter
 					OutDatasmithMesh = RPCMesh;
 					OutDatasmithMeshElement = new FDatasmithFacadeMeshElement(HashedName);
 					OutDatasmithMeshElement.SetLabel(GetActorLabel());
-					OutDatasmithMeshElement.SetMaterial(InMaterialData.MasterMaterial.GetName(), MeshMaterialIndex);
+					OutDatasmithMeshElement.SetMaterial(InMaterialData.MaterialInstance.GetName(), MeshMaterialIndex);
 				}
 				else
 				{
@@ -2056,8 +2056,8 @@ namespace DatasmithRevitExporter
 
 				FDecalMaterial DecalMaterial = DecalMaterialPair.Value;
 
-				FDatasmithFacadeMasterMaterial DatasmithMaterial = new FDatasmithFacadeMasterMaterial(DecalMaterial.MaterialName);
-				DatasmithMaterial.SetMaterialType(FDatasmithFacadeMasterMaterial.EMasterMaterialType.Decal);
+				FDatasmithFacadeMaterialInstance DatasmithMaterial = new FDatasmithFacadeMaterialInstance(DecalMaterial.MaterialName);
+				DatasmithMaterial.SetMaterialType(FDatasmithFacadeMaterialInstance.EMaterialInstanceType.Decal);
 
 				if (!string.IsNullOrEmpty(DecalMaterial.DiffuseTexturePath))
 				{
@@ -2161,7 +2161,7 @@ namespace DatasmithRevitExporter
 			// Add the collected master materials from the material data dictionary to the Datasmith scene.
 			foreach (FMaterialData CollectedMaterialData in NewMaterialsMap.Values)
 			{
-				InDatasmithScene.AddMaterial(CollectedMaterialData.MasterMaterial);
+				InDatasmithScene.AddMaterial(CollectedMaterialData.MaterialInstance);
 
 				foreach(FDatasmithFacadeTexture CurrentTexture in CollectedMaterialData.CollectedTextures)
 				{
