[Non-inclusivity] Removed banned words from Datasmith API and code using the API

diff --git a/Engine/Source/Programs/Enterprise/Datasmith/DatasmithRevitExporter/Private/DatasmithRevitMaterial.cs b/Engine/Source/Programs/Enterprise/Datasmith/DatasmithRevitExporter/Private/DatasmithRevitMaterial.cs
--- a/Engine/Source/Programs/Enterprise/Datasmith/DatasmithRevitExporter/Private/DatasmithRevitMaterial.cs
+++ b/Engine/Source/Programs/Enterprise/Datasmith/DatasmithRevitExporter/Private/DatasmithRevitMaterial.cs
@@ -12,12 +12,12 @@ namespace DatasmithRevitExporter
 {
 	public class FMaterialData
 	{
-		private Material                       CurrentMaterial   = null;
-		private string                         MaterialLabel     = null;	
-		public  FDatasmithFacadeMasterMaterial MasterMaterial    = null;
-		public  List<FDatasmithFacadeTexture>  CollectedTextures = new List<FDatasmithFacadeTexture>();
-		private IList<string>                  ExtraTexturePaths = null;
-		public  IList<string>                  MessageList       = new List<string>();
+		private Material                         CurrentMaterial   = null;
+		private string                           MaterialLabel     = null;	
+		public  FDatasmithFacadeMaterialInstance MaterialInstance    = null;
+		public  List<FDatasmithFacadeTexture>    CollectedTextures = new List<FDatasmithFacadeTexture>();
+		private IList<string>                    ExtraTexturePaths = null;
+		public  IList<string>                    MessageList       = new List<string>();
 
 		// Multi-line debug log.
 		// private FDatasmithFacadeLog DebugLog = null;
@@ -72,13 +72,13 @@ namespace DatasmithRevitExporter
 			// Create a new Datasmith master material.
 			// Hash the Datasmith master material name to shorten it.
 			string HashedMaterialName = GetMaterialName(InMaterialNode, CurrentMaterial);
-			MasterMaterial = new FDatasmithFacadeMasterMaterial(HashedMaterialName);
-			MasterMaterial.SetLabel(GetMaterialLabel(InMaterialNode, CurrentMaterial));
+			MaterialInstance = new FDatasmithFacadeMaterialInstance(HashedMaterialName);
+			MaterialInstance.SetLabel(GetMaterialLabel(InMaterialNode, CurrentMaterial));
 
 			// Set the properties of the Datasmith master material.
-			if (!SetMasterMaterial(CurrentMaterial, MasterMaterial, CollectedTextures))
+			if (!SetMaterialInstance(CurrentMaterial, MaterialInstance, CollectedTextures))
 			{
-				SetFallbackMaterial(InMaterialNode.Color, (float) InMaterialNode.Transparency, InMaterialNode.Smoothness / 100.0F, MasterMaterial);
+				SetFallbackMaterial(InMaterialNode.Color, (float) InMaterialNode.Transparency, InMaterialNode.Smoothness / 100.0F, MaterialInstance);
 			}
 		}
 
@@ -91,11 +91,11 @@ namespace DatasmithRevitExporter
 			MaterialLabel = InMaterialLabel;
 
 			// Create a new Datasmith master material.
-			MasterMaterial = new FDatasmithFacadeMasterMaterial(InHashedMaterialName);
-			MasterMaterial.SetLabel(MaterialLabel);
+			MaterialInstance = new FDatasmithFacadeMaterialInstance(InHashedMaterialName);
+			MaterialInstance.SetLabel(MaterialLabel);
 
 			// Set the properties of the Datasmith master material.
-			SetFallbackMaterial(InMaterialColor, 0.0F, 0.5F, MasterMaterial);
+			SetFallbackMaterial(InMaterialColor, 0.0F, 0.5F, MaterialInstance);
 		}
 
 		public void Log(
@@ -168,9 +168,9 @@ namespace DatasmithRevitExporter
 			return InMaterialNode.NodeName;
 		}
 
-		private bool SetMasterMaterial(
+		private bool SetMaterialInstance(
 			Material                       InMaterial,
-			FDatasmithFacadeMasterMaterial IOMasterMaterial,
+			FDatasmithFacadeMaterialInstance IOMaterialInstance,
 			List<FDatasmithFacadeTexture>  IOCollectedTextures
 		)
 		{
@@ -209,12 +209,12 @@ namespace DatasmithRevitExporter
 			{
 				case "Ceramic":
 				{
-					IOMasterMaterial.SetMaterialType(FDatasmithFacadeMasterMaterial.EMasterMaterialType.Opaque);
+					IOMaterialInstance.SetMaterialType(FDatasmithFacadeMaterialInstance.EMaterialInstanceType.Opaque);
 
 					Color color = GetColorPropertyValue(RenderingAsset, Ceramic.CeramicColor, sourceMaterialColor);
 
 					// Control the Unreal material Base Color.
-					IOMasterMaterial.AddColor("DiffuseColor", color.Red / 255.0F, color.Green / 255.0F, color.Blue / 255.0F, 1.0F);
+					IOMaterialInstance.AddColor("DiffuseColor", color.Red / 255.0F, color.Green / 255.0F, color.Blue / 255.0F, 1.0F);
 
 					string diffuseMapPath = GetTexturePropertyPath(RenderingAsset, Ceramic.CeramicColor);
 
@@ -227,32 +227,32 @@ namespace DatasmithRevitExporter
 						float diffuseMapUVWAngle  = GetTexturePropertyAngle(RenderingAsset, Ceramic.CeramicColor, UnifiedBitmap.TextureWAngle);
 
 						// Control the Unreal material Base Color.
-						IOMasterMaterial.AddFloat("DiffuseMapFading", 1.0F);
+						IOMaterialInstance.AddFloat("DiffuseMapFading", 1.0F);
 						FDatasmithFacadeTexture TextureElement = FDatasmithFacadeMaterialsUtils.CreateSimpleTextureElement(diffuseMapPath);
 						TextureElement.SetSRGB(FDatasmithFacadeTexture.EColorSpace.sRGB);
 						TextureElement.SetTextureMode(FDatasmithFacadeTexture.ETextureMode.Diffuse);
 						IOCollectedTextures.Add(TextureElement);
-						IOMasterMaterial.AddTexture("DiffuseMap", TextureElement);
-						IOMasterMaterial.AddFloat("DiffuseMap_UVOffsetX", diffuseMapUVOffsetX);
-						IOMasterMaterial.AddFloat("DiffuseMap_UVOffsetY", diffuseMapUVOffsetY);
-						IOMasterMaterial.AddFloat("DiffuseMap_UVScaleX",  diffuseMapUVScaleX);
-						IOMasterMaterial.AddFloat("DiffuseMap_UVScaleY",  diffuseMapUVScaleY);
-						IOMasterMaterial.AddFloat("DiffuseMap_UVWAngle",  diffuseMapUVWAngle);
+						IOMaterialInstance.AddTexture("DiffuseMap", TextureElement);
+						IOMaterialInstance.AddFloat("DiffuseMap_UVOffsetX", diffuseMapUVOffsetX);
+						IOMaterialInstance.AddFloat("DiffuseMap_UVOffsetY", diffuseMapUVOffsetY);
+						IOMaterialInstance.AddFloat("DiffuseMap_UVScaleX",  diffuseMapUVScaleX);
+						IOMaterialInstance.AddFloat("DiffuseMap_UVScaleY",  diffuseMapUVScaleY);
+						IOMaterialInstance.AddFloat("DiffuseMap_UVWAngle",  diffuseMapUVWAngle);
 					}
 					else
 					{
-						IOMasterMaterial.AddFloat("DiffuseMapFading", 0.0F);                          
+						IOMaterialInstance.AddFloat("DiffuseMapFading", 0.0F);                          
 					}
 
 					bool  tintEnabled = GetBooleanPropertyValue(RenderingAsset, Ceramic.CommonTintToggle, false);
 					Color tintColor   = GetColorPropertyValue(RenderingAsset, Ceramic.CommonTintColor, sourceMaterialColor);
 
 					// Control the Unreal material Base Color.
-					IOMasterMaterial.AddBoolean("TintEnabled", tintEnabled);
-					IOMasterMaterial.AddColor("TintColor", tintColor.Red / 255.0F, tintColor.Green / 255.0F, tintColor.Blue / 255.0F, 1.0F);
+					IOMaterialInstance.AddBoolean("TintEnabled", tintEnabled);
+					IOMaterialInstance.AddColor("TintColor", tintColor.Red / 255.0F, tintColor.Green / 255.0F, tintColor.Blue / 255.0F, 1.0F);
 
 					// Control the Unreal material Metallic.
-					IOMasterMaterial.AddBoolean("IsMetal", false);
+					IOMaterialInstance.AddBoolean("IsMetal", false);
 
 					float glossiness = 0.6F; // CeramicApplicationType.Satin
 					switch ((CeramicApplicationType) GetIntegerPropertyValue(RenderingAsset, Ceramic.CeramicApplication, (int) CeramicApplicationType.Satin))
@@ -269,10 +269,10 @@ namespace DatasmithRevitExporter
 					}                        
 
 					// Control the Unreal material Roughness.
-					IOMasterMaterial.AddFloat("Glossiness", glossiness);
+					IOMaterialInstance.AddFloat("Glossiness", glossiness);
 
 					// Control the Unreal material Emissive Color.
-					IOMasterMaterial.AddFloat("SelfIlluminationLuminance", 0.0F);
+					IOMaterialInstance.AddFloat("SelfIlluminationLuminance", 0.0F);
 
 					if ((CeramicPatternType) GetIntegerPropertyValue(RenderingAsset, Ceramic.CeramicPattern, (int) CeramicPatternType.None) == CeramicPatternType.Custom)
 					{
@@ -288,37 +288,37 @@ namespace DatasmithRevitExporter
 							float bumpMapUVWAngle  = GetTexturePropertyAngle(RenderingAsset, Ceramic.CeramicPatternMap, UnifiedBitmap.TextureWAngle);
 
 							// Control the Unreal material Normal.
-							IOMasterMaterial.AddFloat("BumpAmount", bumpAmount);
+							IOMaterialInstance.AddFloat("BumpAmount", bumpAmount);
 							FDatasmithFacadeTexture TextureElement = FDatasmithFacadeMaterialsUtils.CreateSimpleTextureElement(bumpMapPath);
 							TextureElement.SetTextureMode(FDatasmithFacadeTexture.ETextureMode.Bump);
 							IOCollectedTextures.Add(TextureElement);
-							IOMasterMaterial.AddTexture("BumpMap", TextureElement);
-							IOMasterMaterial.AddFloat("BumpMap_UVOffsetX", bumpMapUVOffsetX);
-							IOMasterMaterial.AddFloat("BumpMap_UVOffsetY", bumpMapUVOffsetY);
-							IOMasterMaterial.AddFloat("BumpMap_UVScaleX",  bumpMapUVScaleX);
-							IOMasterMaterial.AddFloat("BumpMap_UVScaleY",  bumpMapUVScaleY);
-							IOMasterMaterial.AddFloat("BumpMap_UVWAngle",  bumpMapUVWAngle);
+							IOMaterialInstance.AddTexture("BumpMap", TextureElement);
+							IOMaterialInstance.AddFloat("BumpMap_UVOffsetX", bumpMapUVOffsetX);
+							IOMaterialInstance.AddFloat("BumpMap_UVOffsetY", bumpMapUVOffsetY);
+							IOMaterialInstance.AddFloat("BumpMap_UVScaleX",  bumpMapUVScaleX);
+							IOMaterialInstance.AddFloat("BumpMap_UVScaleY",  bumpMapUVScaleY);
+							IOMaterialInstance.AddFloat("BumpMap_UVWAngle",  bumpMapUVWAngle);
 						}
 						else
 						{
-							IOMasterMaterial.AddFloat("BumpAmount", 0.0F);
+							IOMaterialInstance.AddFloat("BumpAmount", 0.0F);
 						}
 					}
 					else
 					{
-						IOMasterMaterial.AddFloat("BumpAmount", 0.0F);
+						IOMaterialInstance.AddFloat("BumpAmount", 0.0F);
 					}
 				}
 				break;
 
 				case "Concrete":
 				{
-					IOMasterMaterial.SetMaterialType(FDatasmithFacadeMasterMaterial.EMasterMaterialType.Opaque);
+					IOMaterialInstance.SetMaterialType(FDatasmithFacadeMaterialInstance.EMaterialInstanceType.Opaque);
 
 					Color color = GetColorPropertyValue(RenderingAsset, Concrete.ConcreteColor, sourceMaterialColor);
 
 					// Control the Unreal material Base Color.
-					IOMasterMaterial.AddColor("DiffuseColor", color.Red / 255.0F, color.Green / 255.0F, color.Blue / 255.0F, 1.0F);
+					IOMaterialInstance.AddColor("DiffuseColor", color.Red / 255.0F, color.Green / 255.0F, color.Blue / 255.0F, 1.0F);
 
 					string diffuseMapPath = GetTexturePropertyPath(RenderingAsset, Concrete.ConcreteColor);
 
@@ -331,38 +331,38 @@ namespace DatasmithRevitExporter
 						float diffuseMapUVWAngle  = GetTexturePropertyAngle(RenderingAsset, Concrete.ConcreteColor, UnifiedBitmap.TextureWAngle);
 
 						// Control the Unreal material Base Color.
-						IOMasterMaterial.AddFloat("DiffuseMapFading", 1.0F);
+						IOMaterialInstance.AddFloat("DiffuseMapFading", 1.0F);
 						FDatasmithFacadeTexture TextureElement = FDatasmithFacadeMaterialsUtils.CreateSimpleTextureElement(diffuseMapPath);
 						TextureElement.SetSRGB(FDatasmithFacadeTexture.EColorSpace.sRGB);
 						TextureElement.SetTextureMode(FDatasmithFacadeTexture.ETextureMode.Diffuse);
 						IOCollectedTextures.Add(TextureElement);
-						IOMasterMaterial.AddTexture("DiffuseMap", TextureElement);                    
-						IOMasterMaterial.AddFloat("DiffuseMap_UVOffsetX", diffuseMapUVOffsetX);
-						IOMasterMaterial.AddFloat("DiffuseMap_UVOffsetY", diffuseMapUVOffsetY);
-						IOMasterMaterial.AddFloat("DiffuseMap_UVScaleX",  diffuseMapUVScaleX);
-						IOMasterMaterial.AddFloat("DiffuseMap_UVScaleY",  diffuseMapUVScaleY);
-						IOMasterMaterial.AddFloat("DiffuseMap_UVWAngle",  diffuseMapUVWAngle);
+						IOMaterialInstance.AddTexture("DiffuseMap", TextureElement);                    
+						IOMaterialInstance.AddFloat("DiffuseMap_UVOffsetX", diffuseMapUVOffsetX);
+						IOMaterialInstance.AddFloat("DiffuseMap_UVOffsetY", diffuseMapUVOffsetY);
+						IOMaterialInstance.AddFloat("DiffuseMap_UVScaleX",  diffuseMapUVScaleX);
+						IOMaterialInstance.AddFloat("DiffuseMap_UVScaleY",  diffuseMapUVScaleY);
+						IOMaterialInstance.AddFloat("DiffuseMap_UVWAngle",  diffuseMapUVWAngle);
 					}
 					else
 					{
-						IOMasterMaterial.AddFloat("DiffuseMapFading", 0.0F);                          
+						IOMaterialInstance.AddFloat("DiffuseMapFading", 0.0F);                          
 					}
 
 					bool  tintEnabled = GetBooleanPropertyValue(RenderingAsset, Concrete.CommonTintToggle, false);
 					Color tintColor   = GetColorPropertyValue(RenderingAsset, Concrete.CommonTintColor, sourceMaterialColor);
 
 					// Control the Unreal material Base Color.
-					IOMasterMaterial.AddBoolean("TintEnabled", tintEnabled);
-					IOMasterMaterial.AddColor("TintColor", tintColor.Red / 255.0F, tintColor.Green / 255.0F, tintColor.Blue / 255.0F, 1.0F);
+					IOMaterialInstance.AddBoolean("TintEnabled", tintEnabled);
+					IOMaterialInstance.AddColor("TintColor", tintColor.Red / 255.0F, tintColor.Green / 255.0F, tintColor.Blue / 255.0F, 1.0F);
 
 					// Control the Unreal material Metallic.
-					IOMasterMaterial.AddBoolean("IsMetal", false);
+					IOMaterialInstance.AddBoolean("IsMetal", false);
 
 					// Control the Unreal material Roughness.
-					IOMasterMaterial.AddFloat("Glossiness", 0.5F);
+					IOMaterialInstance.AddFloat("Glossiness", 0.5F);
 
 					// Control the Unreal material Emissive Color.
-					IOMasterMaterial.AddFloat("SelfIlluminationLuminance", 0.0F);
+					IOMaterialInstance.AddFloat("SelfIlluminationLuminance", 0.0F);
 
 					if ((ConcreteFinishType) GetIntegerPropertyValue(RenderingAsset, Concrete.ConcreteFinish, (int) ConcreteFinishType.Straight) == ConcreteFinishType.Custom)
 					{
@@ -378,32 +378,32 @@ namespace DatasmithRevitExporter
 							float bumpMapUVWAngle  = GetTexturePropertyAngle(RenderingAsset, Concrete.ConcreteBumpMap, UnifiedBitmap.TextureWAngle);
 
 							// Control the Unreal material Normal.
-							IOMasterMaterial.AddFloat("BumpAmount", bumpAmount);
+							IOMaterialInstance.AddFloat("BumpAmount", bumpAmount);
 							FDatasmithFacadeTexture TextureElement = FDatasmithFacadeMaterialsUtils.CreateSimpleTextureElement(bumpMapPath);
 							TextureElement.SetTextureMode(FDatasmithFacadeTexture.ETextureMode.Bump);
 							IOCollectedTextures.Add(TextureElement);
-							IOMasterMaterial.AddTexture("BumpMap", TextureElement);
-							IOMasterMaterial.AddFloat("BumpMap_UVOffsetX", bumpMapUVOffsetX);
-							IOMasterMaterial.AddFloat("BumpMap_UVOffsetY", bumpMapUVOffsetY);
-							IOMasterMaterial.AddFloat("BumpMap_UVScaleX",  bumpMapUVScaleX);
-							IOMasterMaterial.AddFloat("BumpMap_UVScaleY",  bumpMapUVScaleY);
-							IOMasterMaterial.AddFloat("BumpMap_UVWAngle",  bumpMapUVWAngle);
+							IOMaterialInstance.AddTexture("BumpMap", TextureElement);
+							IOMaterialInstance.AddFloat("BumpMap_UVOffsetX", bumpMapUVOffsetX);
+							IOMaterialInstance.AddFloat("BumpMap_UVOffsetY", bumpMapUVOffsetY);
+							IOMaterialInstance.AddFloat("BumpMap_UVScaleX",  bumpMapUVScaleX);
+							IOMaterialInstance.AddFloat("BumpMap_UVScaleY",  bumpMapUVScaleY);
+							IOMaterialInstance.AddFloat("BumpMap_UVWAngle",  bumpMapUVWAngle);
 						}
 						else
 						{
-							IOMasterMaterial.AddFloat("BumpAmount", 0.0F);
+							IOMaterialInstance.AddFloat("BumpAmount", 0.0F);
 						}
 					}
 					else
 					{
-						IOMasterMaterial.AddFloat("BumpAmount", 0.0F);
+						IOMaterialInstance.AddFloat("BumpAmount", 0.0F);
 					}
 				}
 				break;
 
 				case "Glazing":
 				{
-					IOMasterMaterial.SetMaterialType(FDatasmithFacadeMasterMaterial.EMasterMaterialType.Transparent);
+					IOMaterialInstance.SetMaterialType(FDatasmithFacadeMaterialInstance.EMaterialInstanceType.Transparent);
 
 					// TODO: Should use the Glazing.GlazingTransmittanceColor to select a predefined color value.
 					Color color           = sourceMaterialColor;
@@ -415,7 +415,7 @@ namespace DatasmithRevitExporter
 					}
 
 					// Control the Unreal material Base Color.
-					IOMasterMaterial.AddColor("DiffuseColor", color.Red / 255.0F, color.Green / 255.0F, color.Blue / 255.0F, 1.0F);
+					IOMaterialInstance.AddColor("DiffuseColor", color.Red / 255.0F, color.Green / 255.0F, color.Blue / 255.0F, 1.0F);
 
 					if (!string.IsNullOrEmpty(diffuseMapPath))
 					{
@@ -426,53 +426,53 @@ namespace DatasmithRevitExporter
 						float diffuseMapUVWAngle  = GetTexturePropertyAngle(RenderingAsset, SolidGlass.SolidglassTransmittanceCustomColor, UnifiedBitmap.TextureWAngle);
 
 						// Control the Unreal material Base Color.
-						IOMasterMaterial.AddFloat("DiffuseMapFading", 1.0F);
+						IOMaterialInstance.AddFloat("DiffuseMapFading", 1.0F);
 						FDatasmithFacadeTexture TextureElement = FDatasmithFacadeMaterialsUtils.CreateSimpleTextureElement(diffuseMapPath);
 						TextureElement.SetSRGB(FDatasmithFacadeTexture.EColorSpace.sRGB);
 						TextureElement.SetTextureMode(FDatasmithFacadeTexture.ETextureMode.Diffuse);
 						IOCollectedTextures.Add(TextureElement);
-						IOMasterMaterial.AddTexture("DiffuseMap", TextureElement);                    
-						IOMasterMaterial.AddFloat("DiffuseMap_UVOffsetX", diffuseMapUVOffsetX);
-						IOMasterMaterial.AddFloat("DiffuseMap_UVOffsetY", diffuseMapUVOffsetY);
-						IOMasterMaterial.AddFloat("DiffuseMap_UVScaleX",  diffuseMapUVScaleX);
-						IOMasterMaterial.AddFloat("DiffuseMap_UVScaleY",  diffuseMapUVScaleY);
-						IOMasterMaterial.AddFloat("DiffuseMap_UVWAngle",  diffuseMapUVWAngle);
+						IOMaterialInstance.AddTexture("DiffuseMap", TextureElement);                    
+						IOMaterialInstance.AddFloat("DiffuseMap_UVOffsetX", diffuseMapUVOffsetX);
+						IOMaterialInstance.AddFloat("DiffuseMap_UVOffsetY", diffuseMapUVOffsetY);
+						IOMaterialInstance.AddFloat("DiffuseMap_UVScaleX",  diffuseMapUVScaleX);
+						IOMaterialInstance.AddFloat("DiffuseMap_UVScaleY",  diffuseMapUVScaleY);
+						IOMaterialInstance.AddFloat("DiffuseMap_UVWAngle",  diffuseMapUVWAngle);
 					}
 					else
 					{
-						IOMasterMaterial.AddFloat("DiffuseMapFading", 0.0F);                          
+						IOMaterialInstance.AddFloat("DiffuseMapFading", 0.0F);                          
 					}
 
 					bool  tintEnabled = GetBooleanPropertyValue(RenderingAsset, Glazing.CommonTintToggle, false);
 					Color tintColor   = GetColorPropertyValue(RenderingAsset, Glazing.CommonTintColor, sourceMaterialColor);
 
 					// Control the Unreal material Base Color.
-					IOMasterMaterial.AddBoolean("TintEnabled", tintEnabled);
-					IOMasterMaterial.AddColor("TintColor", tintColor.Red / 255.0F, tintColor.Green / 255.0F, tintColor.Blue / 255.0F, 1.0F);
+					IOMaterialInstance.AddBoolean("TintEnabled", tintEnabled);
+					IOMaterialInstance.AddColor("TintColor", tintColor.Red / 255.0F, tintColor.Green / 255.0F, tintColor.Blue / 255.0F, 1.0F);
 
 					// Control the Unreal material Emissive Color.
-					IOMasterMaterial.AddFloat("SelfIlluminationLuminance", 0.0F);
+					IOMaterialInstance.AddFloat("SelfIlluminationLuminance", 0.0F);
 
 					float transparency = 1.0F - GetFloatPropertyValue(RenderingAsset, Glazing.GlazingReflectance, 0.0F);
 
 					// Control the Unreal material Opacity.
-					IOMasterMaterial.AddFloat("Transparency", transparency);
-					IOMasterMaterial.AddFloat("TransparencyMapFading", 0.0F);
+					IOMaterialInstance.AddFloat("Transparency", transparency);
+					IOMaterialInstance.AddFloat("TransparencyMapFading", 0.0F);
 
 					// Control the Unreal material Normal.
-					IOMasterMaterial.AddFloat("BumpAmount", 0.0F);
+					IOMaterialInstance.AddFloat("BumpAmount", 0.0F);
 
 					// Control the Unreal material Refraction.
-					IOMasterMaterial.AddFloat("RefractionIndex", 1.01F);
+					IOMaterialInstance.AddFloat("RefractionIndex", 1.01F);
 				}
 				break;
 
 				case "Hardwood":
 				{
-					IOMasterMaterial.SetMaterialType(FDatasmithFacadeMasterMaterial.EMasterMaterialType.Opaque);
+					IOMaterialInstance.SetMaterialType(FDatasmithFacadeMaterialInstance.EMaterialInstanceType.Opaque);
 
 					// Control the Unreal material Base Color.
-					IOMasterMaterial.AddColor("DiffuseColor", sourceMaterialColor.Red / 255.0F, sourceMaterialColor.Green / 255.0F, sourceMaterialColor.Blue / 255.0F, 1.0F);
+					IOMaterialInstance.AddColor("DiffuseColor", sourceMaterialColor.Red / 255.0F, sourceMaterialColor.Green / 255.0F, sourceMaterialColor.Blue / 255.0F, 1.0F);
 
 					string diffuseMapPath = GetTexturePropertyPath(RenderingAsset, Hardwood.HardwoodColor);
 
@@ -485,32 +485,32 @@ namespace DatasmithRevitExporter
 						float diffuseMapUVWAngle  = GetTexturePropertyAngle(RenderingAsset, Hardwood.HardwoodColor, UnifiedBitmap.TextureWAngle);
 
 						// Control the Unreal material Base Color.
-						IOMasterMaterial.AddFloat("DiffuseMapFading", 1.0F);
+						IOMaterialInstance.AddFloat("DiffuseMapFading", 1.0F);
 						FDatasmithFacadeTexture TextureElement = FDatasmithFacadeMaterialsUtils.CreateSimpleTextureElement(diffuseMapPath);
 						TextureElement.SetSRGB(FDatasmithFacadeTexture.EColorSpace.sRGB);
 						TextureElement.SetTextureMode(FDatasmithFacadeTexture.ETextureMode.Diffuse);
 						IOCollectedTextures.Add(TextureElement);
-						IOMasterMaterial.AddTexture("DiffuseMap", TextureElement);                    
-						IOMasterMaterial.AddFloat("DiffuseMap_UVOffsetX", diffuseMapUVOffsetX);
-						IOMasterMaterial.AddFloat("DiffuseMap_UVOffsetY", diffuseMapUVOffsetY);
-						IOMasterMaterial.AddFloat("DiffuseMap_UVScaleX",  diffuseMapUVScaleX);
-						IOMasterMaterial.AddFloat("DiffuseMap_UVScaleY",  diffuseMapUVScaleY);
-						IOMasterMaterial.AddFloat("DiffuseMap_UVWAngle",  diffuseMapUVWAngle);
+						IOMaterialInstance.AddTexture("DiffuseMap", TextureElement);                    
+						IOMaterialInstance.AddFloat("DiffuseMap_UVOffsetX", diffuseMapUVOffsetX);
+						IOMaterialInstance.AddFloat("DiffuseMap_UVOffsetY", diffuseMapUVOffsetY);
+						IOMaterialInstance.AddFloat("DiffuseMap_UVScaleX",  diffuseMapUVScaleX);
+						IOMaterialInstance.AddFloat("DiffuseMap_UVScaleY",  diffuseMapUVScaleY);
+						IOMaterialInstance.AddFloat("DiffuseMap_UVWAngle",  diffuseMapUVWAngle);
 					}
 					else
 					{
-						IOMasterMaterial.AddFloat("DiffuseMapFading", 0.0F);                          
+						IOMaterialInstance.AddFloat("DiffuseMapFading", 0.0F);                          
 					}
 
 					bool  tintEnabled = GetBooleanPropertyValue(RenderingAsset, Hardwood.CommonTintToggle, false);
 					Color tintColor   = GetColorPropertyValue(RenderingAsset, Hardwood.CommonTintColor, sourceMaterialColor);
 
 					// Control the Unreal material Base Color.
-					IOMasterMaterial.AddBoolean("TintEnabled", tintEnabled);
-					IOMasterMaterial.AddColor("TintColor", tintColor.Red / 255.0F, tintColor.Green / 255.0F, tintColor.Blue / 255.0F, 1.0F);
+					IOMaterialInstance.AddBoolean("TintEnabled", tintEnabled);
+					IOMaterialInstance.AddColor("TintColor", tintColor.Red / 255.0F, tintColor.Green / 255.0F, tintColor.Blue / 255.0F, 1.0F);
 
 					// Control the Unreal material Metallic.
-					IOMasterMaterial.AddBoolean("IsMetal", false);
+					IOMaterialInstance.AddBoolean("IsMetal", false);
 
 					float glossiness = 0.45F; // HardwoodFinishType.Satin
 					switch ((HardwoodFinishType) GetIntegerPropertyValue(RenderingAsset, Hardwood.HardwoodFinish, (int) HardwoodFinishType.Satin))
@@ -530,10 +530,10 @@ namespace DatasmithRevitExporter
 					}                        
 
 					// Control the Unreal material Roughness.
-					IOMasterMaterial.AddFloat("Glossiness", glossiness);
+					IOMaterialInstance.AddFloat("Glossiness", glossiness);
 
 					// Control the Unreal material Emissive Color.
-					IOMasterMaterial.AddFloat("SelfIlluminationLuminance", 0.0F);
+					IOMaterialInstance.AddFloat("SelfIlluminationLuminance", 0.0F);
 
 					HardwoodImperfectionsType hardwoodImperfectionsType = (HardwoodImperfectionsType) GetIntegerPropertyValue(RenderingAsset, Hardwood.HardwoodImperfections, (int) HardwoodImperfectionsType.None);
 					if (hardwoodImperfectionsType != HardwoodImperfectionsType.None)
@@ -552,37 +552,37 @@ namespace DatasmithRevitExporter
 							float bumpMapUVWAngle  = GetTexturePropertyAngle(RenderingAsset, texturePropertyName, UnifiedBitmap.TextureWAngle);
 
 							// Control the Unreal material Normal.
-							IOMasterMaterial.AddFloat("BumpAmount", bumpAmount);
+							IOMaterialInstance.AddFloat("BumpAmount", bumpAmount);
 							FDatasmithFacadeTexture TextureElement = FDatasmithFacadeMaterialsUtils.CreateSimpleTextureElement(bumpMapPath);
 							TextureElement.SetTextureMode(FDatasmithFacadeTexture.ETextureMode.Bump);
 							IOCollectedTextures.Add(TextureElement);
-							IOMasterMaterial.AddTexture("BumpMap", TextureElement);
-							IOMasterMaterial.AddFloat("BumpMap_UVOffsetX", bumpMapUVOffsetX);
-							IOMasterMaterial.AddFloat("BumpMap_UVOffsetY", bumpMapUVOffsetY);
-							IOMasterMaterial.AddFloat("BumpMap_UVScaleX",  bumpMapUVScaleX);
-							IOMasterMaterial.AddFloat("BumpMap_UVScaleY",  bumpMapUVScaleY);
-							IOMasterMaterial.AddFloat("BumpMap_UVWAngle",  bumpMapUVWAngle);
+							IOMaterialInstance.AddTexture("BumpMap", TextureElement);
+							IOMaterialInstance.AddFloat("BumpMap_UVOffsetX", bumpMapUVOffsetX);
+							IOMaterialInstance.AddFloat("BumpMap_UVOffsetY", bumpMapUVOffsetY);
+							IOMaterialInstance.AddFloat("BumpMap_UVScaleX",  bumpMapUVScaleX);
+							IOMaterialInstance.AddFloat("BumpMap_UVScaleY",  bumpMapUVScaleY);
+							IOMaterialInstance.AddFloat("BumpMap_UVWAngle",  bumpMapUVWAngle);
 						}
 						else
 						{
-							IOMasterMaterial.AddFloat("BumpAmount", 0.0F);
+							IOMaterialInstance.AddFloat("BumpAmount", 0.0F);
 						}
 					}
 					else
 					{
-						IOMasterMaterial.AddFloat("BumpAmount", 0.0F);
+						IOMaterialInstance.AddFloat("BumpAmount", 0.0F);
 					}
 				}
 				break;
 
 				case "MasonryCMU":
 				{
-					IOMasterMaterial.SetMaterialType(FDatasmithFacadeMasterMaterial.EMasterMaterialType.Opaque);
+					IOMaterialInstance.SetMaterialType(FDatasmithFacadeMaterialInstance.EMaterialInstanceType.Opaque);
 
 					Color color = GetColorPropertyValue(RenderingAsset, MasonryCMU.MasonryCMUColor, sourceMaterialColor);
 
 					// Control the Unreal material Base Color.
-					IOMasterMaterial.AddColor("DiffuseColor", color.Red / 255.0F, color.Green / 255.0F, color.Blue / 255.0F, 1.0F);
+					IOMaterialInstance.AddColor("DiffuseColor", color.Red / 255.0F, color.Green / 255.0F, color.Blue / 255.0F, 1.0F);
 
 					string diffuseMapPath = GetTexturePropertyPath(RenderingAsset, MasonryCMU.MasonryCMUColor);
 
@@ -595,32 +595,32 @@ namespace DatasmithRevitExporter
 						float diffuseMapUVWAngle  = GetTexturePropertyAngle(RenderingAsset, MasonryCMU.MasonryCMUColor, UnifiedBitmap.TextureWAngle);
 
 						// Control the Unreal material Base Color.
-						IOMasterMaterial.AddFloat("DiffuseMapFading", 1.0F);
+						IOMaterialInstance.AddFloat("DiffuseMapFading", 1.0F);
 						FDatasmithFacadeTexture TextureElement = FDatasmithFacadeMaterialsUtils.CreateSimpleTextureElement(diffuseMapPath);
 						TextureElement.SetSRGB(FDatasmithFacadeTexture.EColorSpace.sRGB);
 						TextureElement.SetTextureMode(FDatasmithFacadeTexture.ETextureMode.Diffuse);
 						IOCollectedTextures.Add(TextureElement);
-						IOMasterMaterial.AddTexture("DiffuseMap", TextureElement);                    
-						IOMasterMaterial.AddFloat("DiffuseMap_UVOffsetX", diffuseMapUVOffsetX);
-						IOMasterMaterial.AddFloat("DiffuseMap_UVOffsetY", diffuseMapUVOffsetY);
-						IOMasterMaterial.AddFloat("DiffuseMap_UVScaleX",  diffuseMapUVScaleX);
-						IOMasterMaterial.AddFloat("DiffuseMap_UVScaleY",  diffuseMapUVScaleY);
-						IOMasterMaterial.AddFloat("DiffuseMap_UVWAngle",  diffuseMapUVWAngle);
+						IOMaterialInstance.AddTexture("DiffuseMap", TextureElement);                    
+						IOMaterialInstance.AddFloat("DiffuseMap_UVOffsetX", diffuseMapUVOffsetX);
+						IOMaterialInstance.AddFloat("DiffuseMap_UVOffsetY", diffuseMapUVOffsetY);
+						IOMaterialInstance.AddFloat("DiffuseMap_UVScaleX",  diffuseMapUVScaleX);
+						IOMaterialInstance.AddFloat("DiffuseMap_UVScaleY",  diffuseMapUVScaleY);
+						IOMaterialInstance.AddFloat("DiffuseMap_UVWAngle",  diffuseMapUVWAngle);
 					}
 					else
 					{
-						IOMasterMaterial.AddFloat("DiffuseMapFading", 0.0F);                          
+						IOMaterialInstance.AddFloat("DiffuseMapFading", 0.0F);                          
 					}
 
 					bool  tintEnabled = GetBooleanPropertyValue(RenderingAsset, MasonryCMU.CommonTintToggle, false);
 					Color tintColor   = GetColorPropertyValue(RenderingAsset, MasonryCMU.CommonTintColor, sourceMaterialColor);
 
 					// Control the Unreal material Base Color.
-					IOMasterMaterial.AddBoolean("TintEnabled", tintEnabled);
-					IOMasterMaterial.AddColor("TintColor", tintColor.Red / 255.0F, tintColor.Green / 255.0F, tintColor.Blue / 255.0F, 1.0F);
+					IOMaterialInstance.AddBoolean("TintEnabled", tintEnabled);
+					IOMaterialInstance.AddColor("TintColor", tintColor.Red / 255.0F, tintColor.Green / 255.0F, tintColor.Blue / 255.0F, 1.0F);
 
 					// Control the Unreal material Metallic.
-					IOMasterMaterial.AddBoolean("IsMetal", false);
+					IOMaterialInstance.AddBoolean("IsMetal", false);
 
 					float glossiness = 0.25F; // MasonryCMUApplicationType.Matte
 					switch ((MasonryCMUApplicationType) GetIntegerPropertyValue(RenderingAsset, MasonryCMU.MasonryCMUApplication, (int) MasonryCMUApplicationType.Matte))
@@ -635,10 +635,10 @@ namespace DatasmithRevitExporter
 					}                        
 
 					// Control the Unreal material Roughness.
-					IOMasterMaterial.AddFloat("Glossiness", glossiness);
+					IOMaterialInstance.AddFloat("Glossiness", glossiness);
 
 					// Control the Unreal material Emissive Color.
-					IOMasterMaterial.AddFloat("SelfIlluminationLuminance", 0.0F);
+					IOMaterialInstance.AddFloat("SelfIlluminationLuminance", 0.0F);
 
 					if ((MasonryCMUPatternType) GetIntegerPropertyValue(RenderingAsset, MasonryCMU.MasonryCMUPattern, (int) MasonryCMUPatternType.None) == MasonryCMUPatternType.Custom)
 					{
@@ -654,48 +654,48 @@ namespace DatasmithRevitExporter
 							float bumpMapUVWAngle  = GetTexturePropertyAngle(RenderingAsset, MasonryCMU.MasonryCMUPatternMap, UnifiedBitmap.TextureWAngle);
 
 							// Control the Unreal material Normal.
-							IOMasterMaterial.AddFloat("BumpAmount", bumpAmount);
+							IOMaterialInstance.AddFloat("BumpAmount", bumpAmount);
 							FDatasmithFacadeTexture TextureElement = FDatasmithFacadeMaterialsUtils.CreateSimpleTextureElement(bumpMapPath);
 							TextureElement.SetTextureMode(FDatasmithFacadeTexture.ETextureMode.Bump);
 							IOCollectedTextures.Add(TextureElement);
-							IOMasterMaterial.AddTexture("BumpMap", TextureElement);
-							IOMasterMaterial.AddFloat("BumpMap_UVOffsetX", bumpMapUVOffsetX);
-							IOMasterMaterial.AddFloat("BumpMap_UVOffsetY", bumpMapUVOffsetY);
-							IOMasterMaterial.AddFloat("BumpMap_UVScaleX",  bumpMapUVScaleX);
-							IOMasterMaterial.AddFloat("BumpMap_UVScaleY",  bumpMapUVScaleY);
-							IOMasterMaterial.AddFloat("BumpMap_UVWAngle",  bumpMapUVWAngle);
+							IOMaterialInstance.AddTexture("BumpMap", TextureElement);
+							IOMaterialInstance.AddFloat("BumpMap_UVOffsetX", bumpMapUVOffsetX);
+							IOMaterialInstance.AddFloat("BumpMap_UVOffsetY", bumpMapUVOffsetY);
+							IOMaterialInstance.AddFloat("BumpMap_UVScaleX",  bumpMapUVScaleX);
+							IOMaterialInstance.AddFloat("BumpMap_UVScaleY",  bumpMapUVScaleY);
+							IOMaterialInstance.AddFloat("BumpMap_UVWAngle",  bumpMapUVWAngle);
 						}
 						else
 						{
-							IOMasterMaterial.AddFloat("BumpAmount", 0.0F);
+							IOMaterialInstance.AddFloat("BumpAmount", 0.0F);
 						}
 					}
 					else
 					{
-						IOMasterMaterial.AddFloat("BumpAmount", 0.0F);
+						IOMaterialInstance.AddFloat("BumpAmount", 0.0F);
 					}
 				}
 				break;
 
 				case "Metal":
 				{
-					IOMasterMaterial.SetMaterialType(FDatasmithFacadeMasterMaterial.EMasterMaterialType.Opaque);
+					IOMaterialInstance.SetMaterialType(FDatasmithFacadeMaterialInstance.EMaterialInstanceType.Opaque);
 
 					// TODO: Should use the Metal.MetalColor to select a predefined color value.
 
 					// Control the Unreal material Base Color.
-					IOMasterMaterial.AddColor("DiffuseColor", sourceMaterialColor.Red / 255.0F, sourceMaterialColor.Green / 255.0F, sourceMaterialColor.Blue / 255.0F, 1.0F);
-					IOMasterMaterial.AddFloat("DiffuseMapFading", 0.0F);
+					IOMaterialInstance.AddColor("DiffuseColor", sourceMaterialColor.Red / 255.0F, sourceMaterialColor.Green / 255.0F, sourceMaterialColor.Blue / 255.0F, 1.0F);
+					IOMaterialInstance.AddFloat("DiffuseMapFading", 0.0F);
 
 					bool  tintEnabled = GetBooleanPropertyValue(RenderingAsset, Metal.CommonTintToggle, false);
 					Color tintColor   = GetColorPropertyValue(RenderingAsset, Metal.CommonTintColor, sourceMaterialColor);
 
 					// Control the Unreal material Base Color.
-					IOMasterMaterial.AddBoolean("TintEnabled", tintEnabled);
-					IOMasterMaterial.AddColor("TintColor", tintColor.Red / 255.0F, tintColor.Green / 255.0F, tintColor.Blue / 255.0F, 1.0F);
+					IOMaterialInstance.AddBoolean("TintEnabled", tintEnabled);
+					IOMaterialInstance.AddColor("TintColor", tintColor.Red / 255.0F, tintColor.Green / 255.0F, tintColor.Blue / 255.0F, 1.0F);
 
 					// Control the Unreal material Metallic.
-					IOMasterMaterial.AddBoolean("IsMetal", true);
+					IOMaterialInstance.AddBoolean("IsMetal", true);
 
 					float glossiness = 0.5F; // MetalFinishType.SemiPolished
 					switch ((MetalFinishType) GetIntegerPropertyValue(RenderingAsset, Metal.MetalFinish, (int) MetalFinishType.SemiPolished))
@@ -713,10 +713,10 @@ namespace DatasmithRevitExporter
 					}
 
 					// Control the Unreal material Roughness.
-					IOMasterMaterial.AddFloat("Glossiness", glossiness);
+					IOMaterialInstance.AddFloat("Glossiness", glossiness);
 
 					// Control the Unreal material Emissive Color.
-					IOMasterMaterial.AddFloat("SelfIlluminationLuminance", 0.0F);
+					IOMaterialInstance.AddFloat("SelfIlluminationLuminance", 0.0F);
 
 					if ((MetalPerforationsType) GetIntegerPropertyValue(RenderingAsset, Metal.MetalPerforations, (int) MetalPerforationsType.None) == MetalPerforationsType.Custom)
 					{
@@ -730,18 +730,18 @@ namespace DatasmithRevitExporter
 							float cutoutMapUVScaleY  = 1.0F / GetTexturePropertyDistance(RenderingAsset, Metal.MetalPerforationsShader, UnifiedBitmap.TextureRealWorldScaleY, 1.0F);
 							float cutoutMapUVWAngle  = GetTexturePropertyAngle(RenderingAsset, Metal.MetalPerforationsShader, UnifiedBitmap.TextureWAngle);
 
-							IOMasterMaterial.SetMaterialType(FDatasmithFacadeMasterMaterial.EMasterMaterialType.CutOut);
+							IOMaterialInstance.SetMaterialType(FDatasmithFacadeMaterialInstance.EMaterialInstanceType.CutOut);
 
 							// Control the Unreal material Opacity Mask.
 							FDatasmithFacadeTexture TextureElement = FDatasmithFacadeMaterialsUtils.CreateSimpleTextureElement(cutoutMapPath);
 							TextureElement.SetTextureMode(FDatasmithFacadeTexture.ETextureMode.Other);
 							IOCollectedTextures.Add(TextureElement);
-							IOMasterMaterial.AddTexture("CutoutOpacityMap", TextureElement);
-							IOMasterMaterial.AddFloat("CutoutMap_UVOffsetX", cutoutMapUVOffsetX);
-							IOMasterMaterial.AddFloat("CutoutMap_UVOffsetY", cutoutMapUVOffsetY);
-							IOMasterMaterial.AddFloat("CutoutMap_UVScaleX",  cutoutMapUVScaleX);
-							IOMasterMaterial.AddFloat("CutoutMap_UVScaleY",  cutoutMapUVScaleY);
-							IOMasterMaterial.AddFloat("CutoutMap_UVWAngle",  cutoutMapUVWAngle);
+							IOMaterialInstance.AddTexture("CutoutOpacityMap", TextureElement);
+							IOMaterialInstance.AddFloat("CutoutMap_UVOffsetX", cutoutMapUVOffsetX);
+							IOMaterialInstance.AddFloat("CutoutMap_UVOffsetY", cutoutMapUVOffsetY);
+							IOMaterialInstance.AddFloat("CutoutMap_UVScaleX",  cutoutMapUVScaleX);
+							IOMaterialInstance.AddFloat("CutoutMap_UVScaleY",  cutoutMapUVScaleY);
+							IOMaterialInstance.AddFloat("CutoutMap_UVWAngle",  cutoutMapUVWAngle);
 						}
 					}
 
@@ -759,37 +759,37 @@ namespace DatasmithRevitExporter
 							float bumpMapUVWAngle  = GetTexturePropertyAngle(RenderingAsset, Metal.MetalPatternShader, UnifiedBitmap.TextureWAngle);
 
 							// Control the Unreal material Normal.
-							IOMasterMaterial.AddFloat("BumpAmount", bumpAmount);
+							IOMaterialInstance.AddFloat("BumpAmount", bumpAmount);
 							FDatasmithFacadeTexture TextureElement = FDatasmithFacadeMaterialsUtils.CreateSimpleTextureElement(bumpMapPath);
 							TextureElement.SetTextureMode(FDatasmithFacadeTexture.ETextureMode.Bump);
 							IOCollectedTextures.Add(TextureElement);
-							IOMasterMaterial.AddTexture("BumpMap", TextureElement);
-							IOMasterMaterial.AddFloat("BumpMap_UVOffsetX", bumpMapUVOffsetX);
-							IOMasterMaterial.AddFloat("BumpMap_UVOffsetY", bumpMapUVOffsetY);
-							IOMasterMaterial.AddFloat("BumpMap_UVScaleX",  bumpMapUVScaleX);
-							IOMasterMaterial.AddFloat("BumpMap_UVScaleY",  bumpMapUVScaleY);
-							IOMasterMaterial.AddFloat("BumpMap_UVWAngle",  bumpMapUVWAngle);
+							IOMaterialInstance.AddTexture("BumpMap", TextureElement);
+							IOMaterialInstance.AddFloat("BumpMap_UVOffsetX", bumpMapUVOffsetX);
+							IOMaterialInstance.AddFloat("BumpMap_UVOffsetY", bumpMapUVOffsetY);
+							IOMaterialInstance.AddFloat("BumpMap_UVScaleX",  bumpMapUVScaleX);
+							IOMaterialInstance.AddFloat("BumpMap_UVScaleY",  bumpMapUVScaleY);
+							IOMaterialInstance.AddFloat("BumpMap_UVWAngle",  bumpMapUVWAngle);
 						}
 						else
 						{
-							IOMasterMaterial.AddFloat("BumpAmount", 0.0F);
+							IOMaterialInstance.AddFloat("BumpAmount", 0.0F);
 						}
 					}
 					else
 					{
-						IOMasterMaterial.AddFloat("BumpAmount", 0.0F);
+						IOMaterialInstance.AddFloat("BumpAmount", 0.0F);
 					}
 				}
 				break;
 
 				case "MetallicPaint":
 				{
-					IOMasterMaterial.SetMaterialType(FDatasmithFacadeMasterMaterial.EMasterMaterialType.Opaque);
+					IOMaterialInstance.SetMaterialType(FDatasmithFacadeMaterialInstance.EMaterialInstanceType.Opaque);
 
 					Color color = GetColorPropertyValue(RenderingAsset, MetallicPaint.MetallicpaintBaseColor, sourceMaterialColor);
 
 					// Control the Unreal material Base Color.
-					IOMasterMaterial.AddColor("DiffuseColor", color.Red / 255.0F, color.Green / 255.0F, color.Blue / 255.0F, 1.0F);
+					IOMaterialInstance.AddColor("DiffuseColor", color.Red / 255.0F, color.Green / 255.0F, color.Blue / 255.0F, 1.0F);
 
 					string diffuseMapPath = GetTexturePropertyPath(RenderingAsset, MetallicPaint.MetallicpaintBaseColor);
 
@@ -802,72 +802,72 @@ namespace DatasmithRevitExporter
 						float diffuseMapUVWAngle  = GetTexturePropertyAngle(RenderingAsset, MetallicPaint.MetallicpaintBaseColor, UnifiedBitmap.TextureWAngle);
 
 						// Control the Unreal material Base Color.
-						IOMasterMaterial.AddFloat("DiffuseMapFading", 1.0F);
+						IOMaterialInstance.AddFloat("DiffuseMapFading", 1.0F);
 						FDatasmithFacadeTexture TextureElement = FDatasmithFacadeMaterialsUtils.CreateSimpleTextureElement(diffuseMapPath);
 						TextureElement.SetSRGB(FDatasmithFacadeTexture.EColorSpace.sRGB);
 						TextureElement.SetTextureMode(FDatasmithFacadeTexture.ETextureMode.Diffuse);
 						IOCollectedTextures.Add(TextureElement);
-						IOMasterMaterial.AddTexture("DiffuseMap", TextureElement);                    
-						IOMasterMaterial.AddFloat("DiffuseMap_UVOffsetX", diffuseMapUVOffsetX);
-						IOMasterMaterial.AddFloat("DiffuseMap_UVOffsetY", diffuseMapUVOffsetY);
-						IOMasterMaterial.AddFloat("DiffuseMap_UVScaleX",  diffuseMapUVScaleX);
-						IOMasterMaterial.AddFloat("DiffuseMap_UVScaleY",  diffuseMapUVScaleY);
-						IOMasterMaterial.AddFloat("DiffuseMap_UVWAngle",  diffuseMapUVWAngle);
+						IOMaterialInstance.AddTexture("DiffuseMap", TextureElement);                    
+						IOMaterialInstance.AddFloat("DiffuseMap_UVOffsetX", diffuseMapUVOffsetX);
+						IOMaterialInstance.AddFloat("DiffuseMap_UVOffsetY", diffuseMapUVOffsetY);
+						IOMaterialInstance.AddFloat("DiffuseMap_UVScaleX",  diffuseMapUVScaleX);
+						IOMaterialInstance.AddFloat("DiffuseMap_UVScaleY",  diffuseMapUVScaleY);
+						IOMaterialInstance.AddFloat("DiffuseMap_UVWAngle",  diffuseMapUVWAngle);
 					}
 					else
 					{
-						IOMasterMaterial.AddFloat("DiffuseMapFading", 0.0F);                          
+						IOMaterialInstance.AddFloat("DiffuseMapFading", 0.0F);                          
 					}
 
 					bool  tintEnabled = GetBooleanPropertyValue(RenderingAsset, MetallicPaint.CommonTintToggle, false);
 					Color tintColor   = GetColorPropertyValue(RenderingAsset, MetallicPaint.CommonTintColor, sourceMaterialColor);
 
 					// Control the Unreal material Base Color.
-					IOMasterMaterial.AddBoolean("TintEnabled", tintEnabled);
-					IOMasterMaterial.AddColor("TintColor", tintColor.Red / 255.0F, tintColor.Green / 255.0F, tintColor.Blue / 255.0F, 1.0F);
+					IOMaterialInstance.AddBoolean("TintEnabled", tintEnabled);
+					IOMaterialInstance.AddColor("TintColor", tintColor.Red / 255.0F, tintColor.Green / 255.0F, tintColor.Blue / 255.0F, 1.0F);
 
 					// Control the Unreal material Metallic.
-					IOMasterMaterial.AddBoolean("IsMetal", true);
+					IOMaterialInstance.AddBoolean("IsMetal", true);
 
 					// Control the Unreal material Roughness.
-					IOMasterMaterial.AddFloat("Glossiness", 0.9F);
+					IOMaterialInstance.AddFloat("Glossiness", 0.9F);
 
 					// Control the Unreal material Emissive Color.
-					IOMasterMaterial.AddFloat("SelfIlluminationLuminance", 0.0F);
+					IOMaterialInstance.AddFloat("SelfIlluminationLuminance", 0.0F);
 
 					// Control the Unreal material Normal.
-					IOMasterMaterial.AddFloat("BumpAmount", 0.0F);
+					IOMaterialInstance.AddFloat("BumpAmount", 0.0F);
 				}
 				break;
 
 				case "Mirror":
 				{
-					IOMasterMaterial.SetMaterialType(FDatasmithFacadeMasterMaterial.EMasterMaterialType.Opaque);
+					IOMaterialInstance.SetMaterialType(FDatasmithFacadeMaterialInstance.EMaterialInstanceType.Opaque);
 
 					Color color = GetColorPropertyValue(RenderingAsset, Mirror.MirrorTintcolor, sourceMaterialColor);
 
 					// Control the Unreal material Base Color.
-					IOMasterMaterial.AddColor("DiffuseColor", color.Red / 255.0F, color.Green / 255.0F, color.Blue / 255.0F, 1.0F);
-					IOMasterMaterial.AddFloat("DiffuseMapFading", 0.0F);
+					IOMaterialInstance.AddColor("DiffuseColor", color.Red / 255.0F, color.Green / 255.0F, color.Blue / 255.0F, 1.0F);
+					IOMaterialInstance.AddFloat("DiffuseMapFading", 0.0F);
 
 					bool  tintEnabled = GetBooleanPropertyValue(RenderingAsset, Mirror.CommonTintToggle, false);
 					Color tintColor   = GetColorPropertyValue(RenderingAsset, Mirror.CommonTintColor, sourceMaterialColor);
 
 					// Control the Unreal material Base Color.
-					IOMasterMaterial.AddBoolean("TintEnabled", tintEnabled);
-					IOMasterMaterial.AddColor("TintColor", tintColor.Red / 255.0F, tintColor.Green / 255.0F, tintColor.Blue / 255.0F, 1.0F);
+					IOMaterialInstance.AddBoolean("TintEnabled", tintEnabled);
+					IOMaterialInstance.AddColor("TintColor", tintColor.Red / 255.0F, tintColor.Green / 255.0F, tintColor.Blue / 255.0F, 1.0F);
 
 					// Control the Unreal material Metallic.
-					IOMasterMaterial.AddBoolean("IsMetal", true);
+					IOMaterialInstance.AddBoolean("IsMetal", true);
 
 					// Control the Unreal material Roughness.
-					IOMasterMaterial.AddFloat("Glossiness", 1.0F);
+					IOMaterialInstance.AddFloat("Glossiness", 1.0F);
 
 					// Control the Unreal material Emissive Color.
-					IOMasterMaterial.AddFloat("SelfIlluminationLuminance", 0.0F);
+					IOMaterialInstance.AddFloat("SelfIlluminationLuminance", 0.0F);
 
 					// Control the Unreal material Normal.
-					IOMasterMaterial.AddFloat("BumpAmount", 0.0F);
+					IOMaterialInstance.AddFloat("BumpAmount", 0.0F);
 				}
 				break;
 
@@ -876,7 +876,7 @@ namespace DatasmithRevitExporter
 					Color color = GetColorPropertyValue(RenderingAsset, PlasticVinyl.PlasticvinylColor, sourceMaterialColor);
 
 					// Control the Unreal material Base Color.
-					IOMasterMaterial.AddColor("DiffuseColor", color.Red / 255.0F, color.Green / 255.0F, color.Blue / 255.0F, 1.0F);
+					IOMaterialInstance.AddColor("DiffuseColor", color.Red / 255.0F, color.Green / 255.0F, color.Blue / 255.0F, 1.0F);
 
 					string diffuseMapPath = GetTexturePropertyPath(RenderingAsset, PlasticVinyl.PlasticvinylColor);
 
@@ -889,32 +889,32 @@ namespace DatasmithRevitExporter
 						float diffuseMapUVWAngle  = GetTexturePropertyAngle(RenderingAsset, PlasticVinyl.PlasticvinylColor, UnifiedBitmap.TextureWAngle);
 
 						// Control the Unreal material Base Color.
-						IOMasterMaterial.AddFloat("DiffuseMapFading", 1.0F);
+						IOMaterialInstance.AddFloat("DiffuseMapFading", 1.0F);
 						FDatasmithFacadeTexture TextureElement = FDatasmithFacadeMaterialsUtils.CreateSimpleTextureElement(diffuseMapPath);
 						TextureElement.SetSRGB(FDatasmithFacadeTexture.EColorSpace.sRGB);
 						TextureElement.SetTextureMode(FDatasmithFacadeTexture.ETextureMode.Diffuse);
 						IOCollectedTextures.Add(TextureElement);
-						IOMasterMaterial.AddTexture("DiffuseMap", TextureElement);                    
-						IOMasterMaterial.AddFloat("DiffuseMap_UVOffsetX", diffuseMapUVOffsetX);
-						IOMasterMaterial.AddFloat("DiffuseMap_UVOffsetY", diffuseMapUVOffsetY);
-						IOMasterMaterial.AddFloat("DiffuseMap_UVScaleX",  diffuseMapUVScaleX);
-						IOMasterMaterial.AddFloat("DiffuseMap_UVScaleY",  diffuseMapUVScaleY);
-						IOMasterMaterial.AddFloat("DiffuseMap_UVWAngle",  diffuseMapUVWAngle);
+						IOMaterialInstance.AddTexture("DiffuseMap", TextureElement);                    
+						IOMaterialInstance.AddFloat("DiffuseMap_UVOffsetX", diffuseMapUVOffsetX);
+						IOMaterialInstance.AddFloat("DiffuseMap_UVOffsetY", diffuseMapUVOffsetY);
+						IOMaterialInstance.AddFloat("DiffuseMap_UVScaleX",  diffuseMapUVScaleX);
+						IOMaterialInstance.AddFloat("DiffuseMap_UVScaleY",  diffuseMapUVScaleY);
+						IOMaterialInstance.AddFloat("DiffuseMap_UVWAngle",  diffuseMapUVWAngle);
 					}
 					else
 					{
-						IOMasterMaterial.AddFloat("DiffuseMapFading", 0.0F);                          
+						IOMaterialInstance.AddFloat("DiffuseMapFading", 0.0F);                          
 					}
 
 					bool  tintEnabled = GetBooleanPropertyValue(RenderingAsset, PlasticVinyl.CommonTintToggle, false);
 					Color tintColor   = GetColorPropertyValue(RenderingAsset, PlasticVinyl.CommonTintColor, sourceMaterialColor);
 
 					// Control the Unreal material Base Color.
-					IOMasterMaterial.AddBoolean("TintEnabled", tintEnabled);
-					IOMasterMaterial.AddColor("TintColor", tintColor.Red / 255.0F, tintColor.Green / 255.0F, tintColor.Blue / 255.0F, 1.0F);
+					IOMaterialInstance.AddBoolean("TintEnabled", tintEnabled);
+					IOMaterialInstance.AddColor("TintColor", tintColor.Red / 255.0F, tintColor.Green / 255.0F, tintColor.Blue / 255.0F, 1.0F);
 
 					// Control the Unreal material Emissive Color.
-					IOMasterMaterial.AddFloat("SelfIlluminationLuminance", 0.0F);
+					IOMaterialInstance.AddFloat("SelfIlluminationLuminance", 0.0F);
 
 					if ((PlasticvinylBumpType) GetIntegerPropertyValue(RenderingAsset, PlasticVinyl.PlasticvinylBump, (int) PlasticvinylBumpType.None) == PlasticvinylBumpType.Custom)
 					{
@@ -930,44 +930,44 @@ namespace DatasmithRevitExporter
 							float bumpMapUVWAngle  = GetTexturePropertyAngle(RenderingAsset, PlasticVinyl.PlasticvinylBumpMap, UnifiedBitmap.TextureWAngle);
 
 							// Control the Unreal material Normal.
-							IOMasterMaterial.AddFloat("BumpAmount", bumpAmount);
+							IOMaterialInstance.AddFloat("BumpAmount", bumpAmount);
 							FDatasmithFacadeTexture TextureElement = FDatasmithFacadeMaterialsUtils.CreateSimpleTextureElement(bumpMapPath);
 							TextureElement.SetTextureMode(FDatasmithFacadeTexture.ETextureMode.Bump);
 							IOCollectedTextures.Add(TextureElement);
-							IOMasterMaterial.AddTexture("BumpMap", TextureElement);
-							IOMasterMaterial.AddFloat("BumpMap_UVOffsetX", bumpMapUVOffsetX);
-							IOMasterMaterial.AddFloat("BumpMap_UVOffsetY", bumpMapUVOffsetY);
-							IOMasterMaterial.AddFloat("BumpMap_UVScaleX",  bumpMapUVScaleX);
-							IOMasterMaterial.AddFloat("BumpMap_UVScaleY",  bumpMapUVScaleY);
-							IOMasterMaterial.AddFloat("BumpMap_UVWAngle",  bumpMapUVWAngle);
+							IOMaterialInstance.AddTexture("BumpMap", TextureElement);
+							IOMaterialInstance.AddFloat("BumpMap_UVOffsetX", bumpMapUVOffsetX);
+							IOMaterialInstance.AddFloat("BumpMap_UVOffsetY", bumpMapUVOffsetY);
+							IOMaterialInstance.AddFloat("BumpMap_UVScaleX",  bumpMapUVScaleX);
+							IOMaterialInstance.AddFloat("BumpMap_UVScaleY",  bumpMapUVScaleY);
+							IOMaterialInstance.AddFloat("BumpMap_UVWAngle",  bumpMapUVWAngle);
 						}
 						else
 						{
-							IOMasterMaterial.AddFloat("BumpAmount", 0.0F);
+							IOMaterialInstance.AddFloat("BumpAmount", 0.0F);
 						}
 					}
 					else
 					{
-						IOMasterMaterial.AddFloat("BumpAmount", 0.0F);
+						IOMaterialInstance.AddFloat("BumpAmount", 0.0F);
 					}
 
 					if ((PlasticvinylType) GetIntegerPropertyValue(RenderingAsset, PlasticVinyl.PlasticvinylType, (int) PlasticvinylType.Plasticsolid) == PlasticvinylType.Plastictransparent)
 					{
-						IOMasterMaterial.SetMaterialType(FDatasmithFacadeMasterMaterial.EMasterMaterialType.Transparent);
+						IOMaterialInstance.SetMaterialType(FDatasmithFacadeMaterialInstance.EMaterialInstanceType.Transparent);
 
 						// Control the Unreal material Opacity.
-						IOMasterMaterial.AddFloat("Transparency", 0.5F);
-						IOMasterMaterial.AddFloat("TransparencyMapFading", 0.0F);
+						IOMaterialInstance.AddFloat("Transparency", 0.5F);
+						IOMaterialInstance.AddFloat("TransparencyMapFading", 0.0F);
 
 						// Control the Unreal material Refraction.
-						IOMasterMaterial.AddFloat("RefractionIndex", 1.0F);
+						IOMaterialInstance.AddFloat("RefractionIndex", 1.0F);
 					}
 					else
 					{
-						IOMasterMaterial.SetMaterialType(FDatasmithFacadeMasterMaterial.EMasterMaterialType.Opaque);
+						IOMaterialInstance.SetMaterialType(FDatasmithFacadeMaterialInstance.EMaterialInstanceType.Opaque);
 
 						// Control the Unreal material Metallic.
-						IOMasterMaterial.AddBoolean("IsMetal", false);
+						IOMaterialInstance.AddBoolean("IsMetal", false);
 
 						float glossiness = 0.9F; // PlasticvinylApplicationType.Polished
 						switch ((PlasticvinylApplicationType) GetIntegerPropertyValue(RenderingAsset, PlasticVinyl.PlasticvinylApplication, (int) PlasticvinylApplicationType.Polished))
@@ -984,14 +984,14 @@ namespace DatasmithRevitExporter
 						}                        
 
 						// Control the Unreal material Roughness.
-						IOMasterMaterial.AddFloat("Glossiness", glossiness);
+						IOMaterialInstance.AddFloat("Glossiness", glossiness);
 					}
 				}
 				break;
 
 				case "SolidGlass":
 				{
-					IOMasterMaterial.SetMaterialType(FDatasmithFacadeMasterMaterial.EMasterMaterialType.Transparent);
+					IOMaterialInstance.SetMaterialType(FDatasmithFacadeMaterialInstance.EMaterialInstanceType.Transparent);
 
 					// TODO: Should use the SolidGlass.SolidglassTransmittance to select a predefined color value.
 					Color color           = sourceMaterialColor;
@@ -1003,7 +1003,7 @@ namespace DatasmithRevitExporter
 					}
 
 					// Control the Unreal material Base Color.
-					IOMasterMaterial.AddColor("DiffuseColor", color.Red / 255.0F, color.Green / 255.0F, color.Blue / 255.0F, 1.0F);
+					IOMaterialInstance.AddColor("DiffuseColor", color.Red / 255.0F, color.Green / 255.0F, color.Blue / 255.0F, 1.0F);
 
 					if (!string.IsNullOrEmpty(diffuseMapPath))
 					{
@@ -1014,38 +1014,38 @@ namespace DatasmithRevitExporter
 						float diffuseMapUVWAngle  = GetTexturePropertyAngle(RenderingAsset, SolidGlass.SolidglassTransmittanceCustomColor, UnifiedBitmap.TextureWAngle);
 
 						// Control the Unreal material Base Color.
-						IOMasterMaterial.AddFloat("DiffuseMapFading", 1.0F);
+						IOMaterialInstance.AddFloat("DiffuseMapFading", 1.0F);
 						FDatasmithFacadeTexture TextureElement = FDatasmithFacadeMaterialsUtils.CreateSimpleTextureElement(diffuseMapPath);
 						TextureElement.SetSRGB(FDatasmithFacadeTexture.EColorSpace.sRGB);
 						TextureElement.SetTextureMode(FDatasmithFacadeTexture.ETextureMode.Diffuse);
 						IOCollectedTextures.Add(TextureElement);
-						IOMasterMaterial.AddTexture("DiffuseMap", TextureElement);                    
-						IOMasterMaterial.AddFloat("DiffuseMap_UVOffsetX", diffuseMapUVOffsetX);
-						IOMasterMaterial.AddFloat("DiffuseMap_UVOffsetY", diffuseMapUVOffsetY);
-						IOMasterMaterial.AddFloat("DiffuseMap_UVScaleX",  diffuseMapUVScaleX);
-						IOMasterMaterial.AddFloat("DiffuseMap_UVScaleY",  diffuseMapUVScaleY);
-						IOMasterMaterial.AddFloat("DiffuseMap_UVWAngle",  diffuseMapUVWAngle);
+						IOMaterialInstance.AddTexture("DiffuseMap", TextureElement);                    
+						IOMaterialInstance.AddFloat("DiffuseMap_UVOffsetX", diffuseMapUVOffsetX);
+						IOMaterialInstance.AddFloat("DiffuseMap_UVOffsetY", diffuseMapUVOffsetY);
+						IOMaterialInstance.AddFloat("DiffuseMap_UVScaleX",  diffuseMapUVScaleX);
+						IOMaterialInstance.AddFloat("DiffuseMap_UVScaleY",  diffuseMapUVScaleY);
+						IOMaterialInstance.AddFloat("DiffuseMap_UVWAngle",  diffuseMapUVWAngle);
 					}
 					else
 					{
-						IOMasterMaterial.AddFloat("DiffuseMapFading", 0.0F);                          
+						IOMaterialInstance.AddFloat("DiffuseMapFading", 0.0F);                          
 					}
 
 					bool  tintEnabled = GetBooleanPropertyValue(RenderingAsset, SolidGlass.CommonTintToggle, false);
 					Color tintColor   = GetColorPropertyValue(RenderingAsset, SolidGlass.CommonTintColor, sourceMaterialColor);
 
 					// Control the Unreal material Base Color.
-					IOMasterMaterial.AddBoolean("TintEnabled", tintEnabled);
-					IOMasterMaterial.AddColor("TintColor", tintColor.Red / 255.0F, tintColor.Green / 255.0F, tintColor.Blue / 255.0F, 1.0F);
+					IOMaterialInstance.AddBoolean("TintEnabled", tintEnabled);
+					IOMaterialInstance.AddColor("TintColor", tintColor.Red / 255.0F, tintColor.Green / 255.0F, tintColor.Blue / 255.0F, 1.0F);
 
 					// Control the Unreal material Emissive Color.
-					IOMasterMaterial.AddFloat("SelfIlluminationLuminance", 0.0F);
+					IOMaterialInstance.AddFloat("SelfIlluminationLuminance", 0.0F);
 
 					float transparency = 1.0F - GetFloatPropertyValue(RenderingAsset, SolidGlass.SolidglassReflectance, 0.0F);
 
 					// Control the Unreal material Opacity.
-					IOMasterMaterial.AddFloat("Transparency", transparency);
-					IOMasterMaterial.AddFloat("TransparencyMapFading", 0.0F);
+					IOMaterialInstance.AddFloat("Transparency", transparency);
+					IOMaterialInstance.AddFloat("TransparencyMapFading", 0.0F);
 
 					if ((SolidglassBumpEnableType) GetIntegerPropertyValue(RenderingAsset, SolidGlass.SolidglassBumpEnable, (int) SolidglassBumpEnableType.None) == SolidglassBumpEnableType.Custom)
 					{
@@ -1061,38 +1061,38 @@ namespace DatasmithRevitExporter
 							float bumpMapUVWAngle  = GetTexturePropertyAngle(RenderingAsset, SolidGlass.SolidglassBumpMap, UnifiedBitmap.TextureWAngle);
 
 							// Control the Unreal material Normal.
-							IOMasterMaterial.AddFloat("BumpAmount", bumpAmount);
+							IOMaterialInstance.AddFloat("BumpAmount", bumpAmount);
 							FDatasmithFacadeTexture TextureElement = FDatasmithFacadeMaterialsUtils.CreateSimpleTextureElement(bumpMapPath);
 							TextureElement.SetTextureMode(FDatasmithFacadeTexture.ETextureMode.Bump);
 							IOCollectedTextures.Add(TextureElement);
-							IOMasterMaterial.AddTexture("BumpMap", TextureElement);
-							IOMasterMaterial.AddFloat("BumpMap_UVOffsetX", bumpMapUVOffsetX);
-							IOMasterMaterial.AddFloat("BumpMap_UVOffsetY", bumpMapUVOffsetY);
-							IOMasterMaterial.AddFloat("BumpMap_UVScaleX",  bumpMapUVScaleX);
-							IOMasterMaterial.AddFloat("BumpMap_UVScaleY",  bumpMapUVScaleY);
-							IOMasterMaterial.AddFloat("BumpMap_UVWAngle",  bumpMapUVWAngle);
+							IOMaterialInstance.AddTexture("BumpMap", TextureElement);
+							IOMaterialInstance.AddFloat("BumpMap_UVOffsetX", bumpMapUVOffsetX);
+							IOMaterialInstance.AddFloat("BumpMap_UVOffsetY", bumpMapUVOffsetY);
+							IOMaterialInstance.AddFloat("BumpMap_UVScaleX",  bumpMapUVScaleX);
+							IOMaterialInstance.AddFloat("BumpMap_UVScaleY",  bumpMapUVScaleY);
+							IOMaterialInstance.AddFloat("BumpMap_UVWAngle",  bumpMapUVWAngle);
 						}
 						else
 						{
-							IOMasterMaterial.AddFloat("BumpAmount", 0.0F);
+							IOMaterialInstance.AddFloat("BumpAmount", 0.0F);
 						}
 					}
 					else
 					{
-						IOMasterMaterial.AddFloat("BumpAmount", 0.0F);
+						IOMaterialInstance.AddFloat("BumpAmount", 0.0F);
 					}
 
 					// Control the Unreal material Refraction.
-					IOMasterMaterial.AddFloat("RefractionIndex", 1.0F);
+					IOMaterialInstance.AddFloat("RefractionIndex", 1.0F);
 				}
 				break;
 				
 				case "Stone":
 				{
-					IOMasterMaterial.SetMaterialType(FDatasmithFacadeMasterMaterial.EMasterMaterialType.Opaque);
+					IOMaterialInstance.SetMaterialType(FDatasmithFacadeMaterialInstance.EMaterialInstanceType.Opaque);
 
 					// Control the Unreal material Base Color.
-					IOMasterMaterial.AddColor("DiffuseColor", sourceMaterialColor.Red / 255.0F, sourceMaterialColor.Green / 255.0F, sourceMaterialColor.Blue / 255.0F, 1.0F);
+					IOMaterialInstance.AddColor("DiffuseColor", sourceMaterialColor.Red / 255.0F, sourceMaterialColor.Green / 255.0F, sourceMaterialColor.Blue / 255.0F, 1.0F);
 
 					string diffuseMapPath = GetTexturePropertyPath(RenderingAsset, Stone.StoneColor);
 
@@ -1105,32 +1105,32 @@ namespace DatasmithRevitExporter
 						float diffuseMapUVWAngle  = GetTexturePropertyAngle(RenderingAsset, Stone.StoneColor, UnifiedBitmap.TextureWAngle);
 
 						// Control the Unreal material Base Color.
-						IOMasterMaterial.AddFloat("DiffuseMapFading", 1.0F);
+						IOMaterialInstance.AddFloat("DiffuseMapFading", 1.0F);
 						FDatasmithFacadeTexture TextureElement = FDatasmithFacadeMaterialsUtils.CreateSimpleTextureElement(diffuseMapPath);
 						TextureElement.SetSRGB(FDatasmithFacadeTexture.EColorSpace.sRGB);
 						TextureElement.SetTextureMode(FDatasmithFacadeTexture.ETextureMode.Diffuse);
 						IOCollectedTextures.Add(TextureElement);
-						IOMasterMaterial.AddTexture("DiffuseMap", TextureElement);                    
-						IOMasterMaterial.AddFloat("DiffuseMap_UVOffsetX", diffuseMapUVOffsetX);
-						IOMasterMaterial.AddFloat("DiffuseMap_UVOffsetY", diffuseMapUVOffsetY);
-						IOMasterMaterial.AddFloat("DiffuseMap_UVScaleX",  diffuseMapUVScaleX);
-						IOMasterMaterial.AddFloat("DiffuseMap_UVScaleY",  diffuseMapUVScaleY);
-						IOMasterMaterial.AddFloat("DiffuseMap_UVWAngle",  diffuseMapUVWAngle);
+						IOMaterialInstance.AddTexture("DiffuseMap", TextureElement);                    
+						IOMaterialInstance.AddFloat("DiffuseMap_UVOffsetX", diffuseMapUVOffsetX);
+						IOMaterialInstance.AddFloat("DiffuseMap_UVOffsetY", diffuseMapUVOffsetY);
+						IOMaterialInstance.AddFloat("DiffuseMap_UVScaleX",  diffuseMapUVScaleX);
+						IOMaterialInstance.AddFloat("DiffuseMap_UVScaleY",  diffuseMapUVScaleY);
+						IOMaterialInstance.AddFloat("DiffuseMap_UVWAngle",  diffuseMapUVWAngle);
 					}
 					else
 					{
-						IOMasterMaterial.AddFloat("DiffuseMapFading", 0.0F);                          
+						IOMaterialInstance.AddFloat("DiffuseMapFading", 0.0F);                          
 					}
 
 					bool  tintEnabled = GetBooleanPropertyValue(RenderingAsset, Stone.CommonTintToggle, false);
 					Color tintColor   = GetColorPropertyValue(RenderingAsset, Stone.CommonTintColor, sourceMaterialColor);
 
 					// Control the Unreal material Base Color.
-					IOMasterMaterial.AddBoolean("TintEnabled", tintEnabled);
-					IOMasterMaterial.AddColor("TintColor", tintColor.Red / 255.0F, tintColor.Green / 255.0F, tintColor.Blue / 255.0F, 1.0F);
+					IOMaterialInstance.AddBoolean("TintEnabled", tintEnabled);
+					IOMaterialInstance.AddColor("TintColor", tintColor.Red / 255.0F, tintColor.Green / 255.0F, tintColor.Blue / 255.0F, 1.0F);
 
 					// Control the Unreal material Metallic.
-					IOMasterMaterial.AddBoolean("IsMetal", false);
+					IOMaterialInstance.AddBoolean("IsMetal", false);
 
 					float glossiness = 0.35F; // StoneApplicationType.Matte
 					switch ((StoneApplicationType) GetIntegerPropertyValue(RenderingAsset, Stone.StoneApplication, (int) StoneApplicationType.Matte))
@@ -1150,10 +1150,10 @@ namespace DatasmithRevitExporter
 					}                        
 
 					// Control the Unreal material Roughness.
-					IOMasterMaterial.AddFloat("Glossiness", glossiness);
+					IOMaterialInstance.AddFloat("Glossiness", glossiness);
 
 					// Control the Unreal material Emissive Color.
-					IOMasterMaterial.AddFloat("SelfIlluminationLuminance", 0.0F);
+					IOMaterialInstance.AddFloat("SelfIlluminationLuminance", 0.0F);
 
 					if ((StonePatternType) GetIntegerPropertyValue(RenderingAsset, Stone.StonePattern, (int) StonePatternType.None) == StonePatternType.Custom)
 					{
@@ -1169,48 +1169,48 @@ namespace DatasmithRevitExporter
 							float bumpMapUVWAngle  = GetTexturePropertyAngle(RenderingAsset, Stone.StonePatternMap, UnifiedBitmap.TextureWAngle);
 
 							// Control the Unreal material Normal.
-							IOMasterMaterial.AddFloat("BumpAmount", bumpAmount);
+							IOMaterialInstance.AddFloat("BumpAmount", bumpAmount);
 							FDatasmithFacadeTexture TextureElement = FDatasmithFacadeMaterialsUtils.CreateSimpleTextureElement(bumpMapPath);
 							TextureElement.SetTextureMode(FDatasmithFacadeTexture.ETextureMode.Bump);
 							IOCollectedTextures.Add(TextureElement);
-							IOMasterMaterial.AddTexture("BumpMap", TextureElement);
-							IOMasterMaterial.AddFloat("BumpMap_UVOffsetX", bumpMapUVOffsetX);
-							IOMasterMaterial.AddFloat("BumpMap_UVOffsetY", bumpMapUVOffsetY);
-							IOMasterMaterial.AddFloat("BumpMap_UVScaleX",  bumpMapUVScaleX);
-							IOMasterMaterial.AddFloat("BumpMap_UVScaleY",  bumpMapUVScaleY);
-							IOMasterMaterial.AddFloat("BumpMap_UVWAngle",  bumpMapUVWAngle);
+							IOMaterialInstance.AddTexture("BumpMap", TextureElement);
+							IOMaterialInstance.AddFloat("BumpMap_UVOffsetX", bumpMapUVOffsetX);
+							IOMaterialInstance.AddFloat("BumpMap_UVOffsetY", bumpMapUVOffsetY);
+							IOMaterialInstance.AddFloat("BumpMap_UVScaleX",  bumpMapUVScaleX);
+							IOMaterialInstance.AddFloat("BumpMap_UVScaleY",  bumpMapUVScaleY);
+							IOMaterialInstance.AddFloat("BumpMap_UVWAngle",  bumpMapUVWAngle);
 						}
 						else
 						{
-							IOMasterMaterial.AddFloat("BumpAmount", 0.0F);
+							IOMaterialInstance.AddFloat("BumpAmount", 0.0F);
 						}
 					}
 					else
 					{
-						IOMasterMaterial.AddFloat("BumpAmount", 0.0F);
+						IOMaterialInstance.AddFloat("BumpAmount", 0.0F);
 					}
 				}
 				break;
 
 				case "WallPaint":
 				{
-					IOMasterMaterial.SetMaterialType(FDatasmithFacadeMasterMaterial.EMasterMaterialType.Opaque);
+					IOMaterialInstance.SetMaterialType(FDatasmithFacadeMaterialInstance.EMaterialInstanceType.Opaque);
 
 					Color color = GetColorPropertyValue(RenderingAsset, WallPaint.WallpaintColor, sourceMaterialColor);
 
 					// Control the Unreal material Base Color.
-					IOMasterMaterial.AddColor("DiffuseColor", color.Red / 255.0F, color.Green / 255.0F, color.Blue / 255.0F, 1.0F);
-					IOMasterMaterial.AddFloat("DiffuseMapFading", 0.0F);
+					IOMaterialInstance.AddColor("DiffuseColor", color.Red / 255.0F, color.Green / 255.0F, color.Blue / 255.0F, 1.0F);
+					IOMaterialInstance.AddFloat("DiffuseMapFading", 0.0F);
 
 					bool  tintEnabled = GetBooleanPropertyValue(RenderingAsset, WallPaint.CommonTintToggle, false);
 					Color tintColor   = GetColorPropertyValue(RenderingAsset, WallPaint.CommonTintColor, sourceMaterialColor);
 
 					// Control the Unreal material Base Color.
-					IOMasterMaterial.AddBoolean("TintEnabled", tintEnabled);
-					IOMasterMaterial.AddColor("TintColor", tintColor.Red / 255.0F, tintColor.Green / 255.0F, tintColor.Blue / 255.0F, 1.0F);
+					IOMaterialInstance.AddBoolean("TintEnabled", tintEnabled);
+					IOMaterialInstance.AddColor("TintColor", tintColor.Red / 255.0F, tintColor.Green / 255.0F, tintColor.Blue / 255.0F, 1.0F);
 
 					// Control the Unreal material Metallic.
-					IOMasterMaterial.AddBoolean("IsMetal", false);
+					IOMaterialInstance.AddBoolean("IsMetal", false);
 
 					float glossiness = 0.6F; // WallpaintFinishType.Pearl
 					switch ((WallpaintFinishType) GetIntegerPropertyValue(RenderingAsset, WallPaint.WallpaintFinish, (int) WallpaintFinishType.Pearl))
@@ -1236,13 +1236,13 @@ namespace DatasmithRevitExporter
 					}                        
 
 					// Control the Unreal material Roughness.
-					IOMasterMaterial.AddFloat("Glossiness", glossiness);
+					IOMaterialInstance.AddFloat("Glossiness", glossiness);
 
 					// Control the Unreal material Emissive Color.
-					IOMasterMaterial.AddFloat("SelfIlluminationLuminance", 0.0F);
+					IOMaterialInstance.AddFloat("SelfIlluminationLuminance", 0.0F);
 
 					// Control the Unreal material Normal.
-					IOMasterMaterial.AddFloat("BumpAmount", 0.0F);
+					IOMaterialInstance.AddFloat("BumpAmount", 0.0F);
 				}
 				break;
 
@@ -1251,7 +1251,7 @@ namespace DatasmithRevitExporter
 					Color color = GetColorPropertyValue(RenderingAsset, Generic.GenericDiffuse, sourceMaterialColor);
 
 					// Control the Unreal material Base Color.
-					IOMasterMaterial.AddColor("DiffuseColor", color.Red / 255.0F, color.Green / 255.0F, color.Blue / 255.0F, 1.0F);
+					IOMaterialInstance.AddColor("DiffuseColor", color.Red / 255.0F, color.Green / 255.0F, color.Blue / 255.0F, 1.0F);
 
 					string diffuseMapPath = GetTexturePropertyPath(RenderingAsset, Generic.GenericDiffuse);
 
@@ -1265,38 +1265,38 @@ namespace DatasmithRevitExporter
 						float diffuseMapUVWAngle  = GetTexturePropertyAngle(RenderingAsset, Generic.GenericDiffuse, UnifiedBitmap.TextureWAngle);
 
 						// Control the Unreal material Base Color.
-						IOMasterMaterial.AddFloat("DiffuseMapFading", diffuseImageFade);
+						IOMaterialInstance.AddFloat("DiffuseMapFading", diffuseImageFade);
 						FDatasmithFacadeTexture TextureElement = FDatasmithFacadeMaterialsUtils.CreateSimpleTextureElement(diffuseMapPath);
 						TextureElement.SetSRGB(FDatasmithFacadeTexture.EColorSpace.sRGB);
 						TextureElement.SetTextureMode(FDatasmithFacadeTexture.ETextureMode.Diffuse);
 						IOCollectedTextures.Add(TextureElement);
-						IOMasterMaterial.AddTexture("DiffuseMap", TextureElement);                    
-						IOMasterMaterial.AddFloat("DiffuseMap_UVOffsetX", diffuseMapUVOffsetX);
-						IOMasterMaterial.AddFloat("DiffuseMap_UVOffsetY", diffuseMapUVOffsetY);
-						IOMasterMaterial.AddFloat("DiffuseMap_UVScaleX",  diffuseMapUVScaleX);
-						IOMasterMaterial.AddFloat("DiffuseMap_UVScaleY",  diffuseMapUVScaleY);
-						IOMasterMaterial.AddFloat("DiffuseMap_UVWAngle",  diffuseMapUVWAngle);
+						IOMaterialInstance.AddTexture("DiffuseMap", TextureElement);                    
+						IOMaterialInstance.AddFloat("DiffuseMap_UVOffsetX", diffuseMapUVOffsetX);
+						IOMaterialInstance.AddFloat("DiffuseMap_UVOffsetY", diffuseMapUVOffsetY);
+						IOMaterialInstance.AddFloat("DiffuseMap_UVScaleX",  diffuseMapUVScaleX);
+						IOMaterialInstance.AddFloat("DiffuseMap_UVScaleY",  diffuseMapUVScaleY);
+						IOMaterialInstance.AddFloat("DiffuseMap_UVWAngle",  diffuseMapUVWAngle);
 					}
 					else
 					{
-						IOMasterMaterial.AddFloat("DiffuseMapFading", 0.0F);                          
+						IOMaterialInstance.AddFloat("DiffuseMapFading", 0.0F);                          
 					}
 
 					bool  tintEnabled = GetBooleanPropertyValue(RenderingAsset, Generic.CommonTintToggle, false);
 					Color tintColor   = GetColorPropertyValue(RenderingAsset, Generic.CommonTintColor, sourceMaterialColor);
 
 					// Control the Unreal material Base Color.
-					IOMasterMaterial.AddBoolean("TintEnabled", tintEnabled);
-					IOMasterMaterial.AddColor("TintColor", tintColor.Red / 255.0F, tintColor.Green / 255.0F, tintColor.Blue / 255.0F, 1.0F);
+					IOMaterialInstance.AddBoolean("TintEnabled", tintEnabled);
+					IOMaterialInstance.AddColor("TintColor", tintColor.Red / 255.0F, tintColor.Green / 255.0F, tintColor.Blue / 255.0F, 1.0F);
 
 					float selfIlluminationLuminance			= GetFloatPropertyValue(RenderingAsset, Generic.GenericSelfIllumLuminance, 0.0F);
 					float selfIlluminationColorTemperature	= GetFloatPropertyValue(RenderingAsset, Generic.GenericSelfIllumColorTemperature, 0.0f);
 					Color selfIlluminationFilterColor		= GetColorPropertyValue(RenderingAsset, Generic.GenericSelfIllumFilterMap, new Color(255, 255, 255));
 				
 					// Control the Unreal material Emissive Color.
-					IOMasterMaterial.AddFloat("SelfIlluminationLuminance", selfIlluminationLuminance);
-					IOMasterMaterial.AddFloat("SelfIlluminationColorTemperature", selfIlluminationColorTemperature);
-					IOMasterMaterial.AddColor("SelfIlluminationFilter", selfIlluminationFilterColor.Red / 255.0F, selfIlluminationFilterColor.Green / 255.0F, selfIlluminationFilterColor.Blue / 255.0F, 1.0F);
+					IOMaterialInstance.AddFloat("SelfIlluminationLuminance", selfIlluminationLuminance);
+					IOMaterialInstance.AddFloat("SelfIlluminationColorTemperature", selfIlluminationColorTemperature);
+					IOMaterialInstance.AddColor("SelfIlluminationFilter", selfIlluminationFilterColor.Red / 255.0F, selfIlluminationFilterColor.Green / 255.0F, selfIlluminationFilterColor.Blue / 255.0F, 1.0F);
 
 					string selfIlluminationMapPath = GetTexturePropertyPath(RenderingAsset, Generic.GenericSelfIllumFilterMap);
 
@@ -1309,20 +1309,20 @@ namespace DatasmithRevitExporter
 						float selfIlluminationMapUVWAngle  = GetTexturePropertyAngle(RenderingAsset, Generic.GenericSelfIllumFilterMap, UnifiedBitmap.TextureWAngle);
 
 						// Control the Unreal material Emissive Color.
-						IOMasterMaterial.AddBoolean("SelfIlluminationMapEnable", true);
+						IOMaterialInstance.AddBoolean("SelfIlluminationMapEnable", true);
 						FDatasmithFacadeTexture TextureElement = FDatasmithFacadeMaterialsUtils.CreateSimpleTextureElement(selfIlluminationMapPath);
 						TextureElement.SetTextureMode(FDatasmithFacadeTexture.ETextureMode.Other);
 						IOCollectedTextures.Add(TextureElement);
-						IOMasterMaterial.AddTexture("SelfIlluminationMap", TextureElement);
-						IOMasterMaterial.AddFloat("SelfIlluminationMap_UVOffsetX", selfIlluminationMapUVOffsetX);
-						IOMasterMaterial.AddFloat("SelfIlluminationMap_UVOffsetY", selfIlluminationMapUVOffsetY);
-						IOMasterMaterial.AddFloat("SelfIlluminationMap_UVScaleX",  selfIlluminationMapUVScaleX);
-						IOMasterMaterial.AddFloat("SelfIlluminationMap_UVScaleY",  selfIlluminationMapUVScaleY);
-						IOMasterMaterial.AddFloat("SelfIlluminationMap_UVWAngle",  selfIlluminationMapUVWAngle);
+						IOMaterialInstance.AddTexture("SelfIlluminationMap", TextureElement);
+						IOMaterialInstance.AddFloat("SelfIlluminationMap_UVOffsetX", selfIlluminationMapUVOffsetX);
+						IOMaterialInstance.AddFloat("SelfIlluminationMap_UVOffsetY", selfIlluminationMapUVOffsetY);
+						IOMaterialInstance.AddFloat("SelfIlluminationMap_UVScaleX",  selfIlluminationMapUVScaleX);
+						IOMaterialInstance.AddFloat("SelfIlluminationMap_UVScaleY",  selfIlluminationMapUVScaleY);
+						IOMaterialInstance.AddFloat("SelfIlluminationMap_UVWAngle",  selfIlluminationMapUVWAngle);
 					}
 					else
 					{
-						IOMasterMaterial.AddBoolean("SelfIlluminationMapEnable", false);
+						IOMaterialInstance.AddBoolean("SelfIlluminationMapEnable", false);
 					}
 
 					string bumpMapPath = GetTexturePropertyPath(RenderingAsset, Generic.GenericBumpMap);
@@ -1337,30 +1337,30 @@ namespace DatasmithRevitExporter
 						float bumpMapUVWAngle  = GetTexturePropertyAngle(RenderingAsset, Generic.GenericBumpMap, UnifiedBitmap.TextureWAngle);
 
 						// Control the Unreal material Normal.
-						IOMasterMaterial.AddFloat("BumpAmount", bumpAmount);
+						IOMaterialInstance.AddFloat("BumpAmount", bumpAmount);
 						FDatasmithFacadeTexture TextureElement = FDatasmithFacadeMaterialsUtils.CreateSimpleTextureElement(bumpMapPath);
 						TextureElement.SetTextureMode(FDatasmithFacadeTexture.ETextureMode.Bump);
 						IOCollectedTextures.Add(TextureElement);
-						IOMasterMaterial.AddTexture("BumpMap", TextureElement);
-						IOMasterMaterial.AddFloat("BumpMap_UVOffsetX", bumpMapUVOffsetX);
-						IOMasterMaterial.AddFloat("BumpMap_UVOffsetY", bumpMapUVOffsetY);
-						IOMasterMaterial.AddFloat("BumpMap_UVScaleX",  bumpMapUVScaleX);
-						IOMasterMaterial.AddFloat("BumpMap_UVScaleY",  bumpMapUVScaleY);
-						IOMasterMaterial.AddFloat("BumpMap_UVWAngle",  bumpMapUVWAngle);
+						IOMaterialInstance.AddTexture("BumpMap", TextureElement);
+						IOMaterialInstance.AddFloat("BumpMap_UVOffsetX", bumpMapUVOffsetX);
+						IOMaterialInstance.AddFloat("BumpMap_UVOffsetY", bumpMapUVOffsetY);
+						IOMaterialInstance.AddFloat("BumpMap_UVScaleX",  bumpMapUVScaleX);
+						IOMaterialInstance.AddFloat("BumpMap_UVScaleY",  bumpMapUVScaleY);
+						IOMaterialInstance.AddFloat("BumpMap_UVWAngle",  bumpMapUVWAngle);
 					}
 					else
 					{
-						IOMasterMaterial.AddFloat("BumpAmount", 0.0F);
+						IOMaterialInstance.AddFloat("BumpAmount", 0.0F);
 					}
 
 					float transparency = GetFloatPropertyValue(RenderingAsset, Generic.GenericTransparency, 0.0F);
 
 					if (transparency > 0.0F)
 					{
-						IOMasterMaterial.SetMaterialType(FDatasmithFacadeMasterMaterial.EMasterMaterialType.Transparent);
+						IOMaterialInstance.SetMaterialType(FDatasmithFacadeMaterialInstance.EMaterialInstanceType.Transparent);
 
 						// Control the Unreal material Opacity.
-						IOMasterMaterial.AddFloat("Transparency", transparency);
+						IOMaterialInstance.AddFloat("Transparency", transparency);
 
 						string transparencyMapPath = GetTexturePropertyPath(RenderingAsset, Generic.GenericTransparency);
 
@@ -1374,40 +1374,40 @@ namespace DatasmithRevitExporter
 							float transparencyMapUVWAngle  = GetTexturePropertyAngle(RenderingAsset, Generic.GenericTransparency, UnifiedBitmap.TextureWAngle);
 
 							// Control the Unreal material Opacity.
-							IOMasterMaterial.AddFloat("TransparencyMapFading", transparencyImageFade);
+							IOMaterialInstance.AddFloat("TransparencyMapFading", transparencyImageFade);
 							FDatasmithFacadeTexture TextureElement = FDatasmithFacadeMaterialsUtils.CreateSimpleTextureElement(transparencyMapPath);
 							TextureElement.SetTextureMode(FDatasmithFacadeTexture.ETextureMode.Other);
 							IOCollectedTextures.Add(TextureElement);
-							IOMasterMaterial.AddTexture("TransparencyMap", TextureElement);
-							IOMasterMaterial.AddFloat("TransparencyMap_UVOffsetX", transparencyMapUVOffsetX);
-							IOMasterMaterial.AddFloat("TransparencyMap_UVOffsetY", transparencyMapUVOffsetY);
-							IOMasterMaterial.AddFloat("TransparencyMap_UVScaleX",  transparencyMapUVScaleX);
-							IOMasterMaterial.AddFloat("TransparencyMap_UVScaleY",  transparencyMapUVScaleY);
-							IOMasterMaterial.AddFloat("TransparencyMap_UVWAngle",  transparencyMapUVWAngle);
+							IOMaterialInstance.AddTexture("TransparencyMap", TextureElement);
+							IOMaterialInstance.AddFloat("TransparencyMap_UVOffsetX", transparencyMapUVOffsetX);
+							IOMaterialInstance.AddFloat("TransparencyMap_UVOffsetY", transparencyMapUVOffsetY);
+							IOMaterialInstance.AddFloat("TransparencyMap_UVScaleX",  transparencyMapUVScaleX);
+							IOMaterialInstance.AddFloat("TransparencyMap_UVScaleY",  transparencyMapUVScaleY);
+							IOMaterialInstance.AddFloat("TransparencyMap_UVWAngle",  transparencyMapUVWAngle);
 						}
 						else
 						{
-							IOMasterMaterial.AddFloat("TransparencyMapFading", 0.0F);
+							IOMaterialInstance.AddFloat("TransparencyMapFading", 0.0F);
 						}
 
 						float refractionIndex = GetFloatPropertyValue(RenderingAsset, Generic.GenericRefractionIndex, 1.0F);
 
 						// Control the Unreal material Refraction.
-						IOMasterMaterial.AddFloat("RefractionIndex", refractionIndex);
+						IOMaterialInstance.AddFloat("RefractionIndex", refractionIndex);
 					}
 					else
 					{
-						IOMasterMaterial.SetMaterialType(FDatasmithFacadeMasterMaterial.EMasterMaterialType.Opaque);
+						IOMaterialInstance.SetMaterialType(FDatasmithFacadeMaterialInstance.EMaterialInstanceType.Opaque);
 
 						bool isMetal = GetBooleanPropertyValue(RenderingAsset, Generic.GenericIsMetal, false);
 
 						// Control the Unreal material Metallic.
-						IOMasterMaterial.AddBoolean("IsMetal", isMetal);
+						IOMaterialInstance.AddBoolean("IsMetal", isMetal);
 
 						float glossiness = GetFloatPropertyValue(RenderingAsset, Generic.GenericGlossiness, InMaterial.Smoothness / 100.0F);
 
 						// Control the Unreal material Roughness.
-						IOMasterMaterial.AddFloat("Glossiness", glossiness);
+						IOMaterialInstance.AddFloat("Glossiness", glossiness);
 
 						string cutoutMapPath = GetTexturePropertyPath(RenderingAsset, Generic.GenericCutoutOpacity);
 
@@ -1419,18 +1419,18 @@ namespace DatasmithRevitExporter
 							float cutoutMapUVScaleY  = 1.0F / GetTexturePropertyDistance(RenderingAsset, Generic.GenericCutoutOpacity, UnifiedBitmap.TextureRealWorldScaleY, 1.0F);
 							float cutoutMapUVWAngle  = GetTexturePropertyAngle(RenderingAsset, Generic.GenericCutoutOpacity, UnifiedBitmap.TextureWAngle);
 
-							IOMasterMaterial.SetMaterialType(FDatasmithFacadeMasterMaterial.EMasterMaterialType.CutOut);
+							IOMaterialInstance.SetMaterialType(FDatasmithFacadeMaterialInstance.EMaterialInstanceType.CutOut);
 
 							// Control the Unreal material Opacity Mask.
 							FDatasmithFacadeTexture TextureElement = FDatasmithFacadeMaterialsUtils.CreateSimpleTextureElement(cutoutMapPath);
 							TextureElement.SetTextureMode(FDatasmithFacadeTexture.ETextureMode.Other);
 							IOCollectedTextures.Add(TextureElement);
-							IOMasterMaterial.AddTexture("CutoutOpacityMap", TextureElement);
-							IOMasterMaterial.AddFloat("CutoutMap_UVOffsetX", cutoutMapUVOffsetX);
-							IOMasterMaterial.AddFloat("CutoutMap_UVOffsetY", cutoutMapUVOffsetY);
-							IOMasterMaterial.AddFloat("CutoutMap_UVScaleX",  cutoutMapUVScaleX);
-							IOMasterMaterial.AddFloat("CutoutMap_UVScaleY",  cutoutMapUVScaleY);
-							IOMasterMaterial.AddFloat("CutoutMap_UVWAngle",  cutoutMapUVWAngle);
+							IOMaterialInstance.AddTexture("CutoutOpacityMap", TextureElement);
+							IOMaterialInstance.AddFloat("CutoutMap_UVOffsetX", cutoutMapUVOffsetX);
+							IOMaterialInstance.AddFloat("CutoutMap_UVOffsetY", cutoutMapUVOffsetY);
+							IOMaterialInstance.AddFloat("CutoutMap_UVScaleX",  cutoutMapUVScaleX);
+							IOMaterialInstance.AddFloat("CutoutMap_UVScaleY",  cutoutMapUVScaleY);
+							IOMaterialInstance.AddFloat("CutoutMap_UVWAngle",  cutoutMapUVWAngle);
 						}
 					}
 				}
@@ -1441,88 +1441,88 @@ namespace DatasmithRevitExporter
 #if REVIT_API_2020
 				case "AdvancedGlazing":
 				{
-					IOMasterMaterial.AddBoolean("IsPBR", true);
-					IOMasterMaterial.SetMaterialType(FDatasmithFacadeMasterMaterial.EMasterMaterialType.Transparent);
+					IOMaterialInstance.AddBoolean("IsPBR", true);
+					IOMaterialInstance.SetMaterialType(FDatasmithFacadeMaterialInstance.EMaterialInstanceType.Transparent);
 
 					Color TranspColor = GetColorPropertyValue(RenderingAsset, AdvancedGlazing.GlazingTransmissionColor, sourceMaterialColor);
 					float Transparency = LightnessFromColor(TranspColor);
-					IOMasterMaterial.AddFloat("Transparency", Transparency);
-					IOMasterMaterial.AddFloat("TransparencyMapFading", 0.0F);
+					IOMaterialInstance.AddFloat("Transparency", Transparency);
+					IOMaterialInstance.AddFloat("TransparencyMapFading", 0.0F);
 
-					IOMasterMaterial.AddBoolean("TintEnabled", true);
-					IOMasterMaterial.AddColor("TintColor", TranspColor.Red / 255.0F, TranspColor.Green / 255.0F, TranspColor.Blue / 255.0F, 1.0F);
+					IOMaterialInstance.AddBoolean("TintEnabled", true);
+					IOMaterialInstance.AddColor("TintColor", TranspColor.Red / 255.0F, TranspColor.Green / 255.0F, TranspColor.Blue / 255.0F, 1.0F);
 
-					ExportRougness(IOMasterMaterial, IOCollectedTextures, RenderingAsset, AdvancedGlazing.GlazingTransmissionRoughness);
-					ExportCutout(IOMasterMaterial, IOCollectedTextures, RenderingAsset, AdvancedGlazing.SurfaceCutout);
-					ExportNormalMap(IOMasterMaterial, IOCollectedTextures, RenderingAsset, AdvancedGlazing.SurfaceNormal);
+					ExportRougness(IOMaterialInstance, IOCollectedTextures, RenderingAsset, AdvancedGlazing.GlazingTransmissionRoughness);
+					ExportCutout(IOMaterialInstance, IOCollectedTextures, RenderingAsset, AdvancedGlazing.SurfaceCutout);
+					ExportNormalMap(IOMaterialInstance, IOCollectedTextures, RenderingAsset, AdvancedGlazing.SurfaceNormal);
 				}
 				break;
 #endif
 
 				case "AdvancedLayered":
 				{
-					IOMasterMaterial.AddBoolean("IsPBR", true);
-					ExportDiffuse(IOMasterMaterial, IOCollectedTextures, RenderingAsset, AdvancedLayered.LayeredBottomF0, sourceMaterialColor);
-					ExportRougness(IOMasterMaterial, IOCollectedTextures, RenderingAsset, AdvancedLayered.LayeredRoughness);
-					ExportCutout(IOMasterMaterial, IOCollectedTextures, RenderingAsset, AdvancedLayered.SurfaceCutout);
-					ExportNormalMap(IOMasterMaterial, IOCollectedTextures, RenderingAsset, AdvancedLayered.LayeredNormal);
+					IOMaterialInstance.AddBoolean("IsPBR", true);
+					ExportDiffuse(IOMaterialInstance, IOCollectedTextures, RenderingAsset, AdvancedLayered.LayeredBottomF0, sourceMaterialColor);
+					ExportRougness(IOMaterialInstance, IOCollectedTextures, RenderingAsset, AdvancedLayered.LayeredRoughness);
+					ExportCutout(IOMaterialInstance, IOCollectedTextures, RenderingAsset, AdvancedLayered.SurfaceCutout);
+					ExportNormalMap(IOMaterialInstance, IOCollectedTextures, RenderingAsset, AdvancedLayered.LayeredNormal);
 				}
 				break;
 
 				case "AdvancedTransparent":
 				{
-					IOMasterMaterial.AddBoolean("IsPBR", true);
-					IOMasterMaterial.SetMaterialType(FDatasmithFacadeMasterMaterial.EMasterMaterialType.Transparent);
+					IOMaterialInstance.AddBoolean("IsPBR", true);
+					IOMaterialInstance.SetMaterialType(FDatasmithFacadeMaterialInstance.EMaterialInstanceType.Transparent);
 
 					Color TranspColor = GetColorPropertyValue(RenderingAsset, AdvancedTransparent.TransparentColor, sourceMaterialColor);
 
-					IOMasterMaterial.AddBoolean("TintEnabled", true);
-					IOMasterMaterial.AddColor("TintColor", TranspColor.Red / 255.0F, TranspColor.Green / 255.0F, TranspColor.Blue / 255.0F, 1.0F);
+					IOMaterialInstance.AddBoolean("TintEnabled", true);
+					IOMaterialInstance.AddColor("TintColor", TranspColor.Red / 255.0F, TranspColor.Green / 255.0F, TranspColor.Blue / 255.0F, 1.0F);
 
 					float Transparency = LightnessFromColor(TranspColor);
-					IOMasterMaterial.AddFloat("Transparency", Transparency);
-					IOMasterMaterial.AddFloat("TransparencyMapFading", 0.0F);
+					IOMaterialInstance.AddFloat("Transparency", Transparency);
+					IOMaterialInstance.AddFloat("TransparencyMapFading", 0.0F);
 
 					// Control the Unreal material Refraction.
 					float RefractionIndex = GetFloatPropertyValue(RenderingAsset, AdvancedTransparent.TransparentIor, 1.0f);
-					IOMasterMaterial.AddFloat("RefractionIndex", RefractionIndex);
+					IOMaterialInstance.AddFloat("RefractionIndex", RefractionIndex);
 
-					ExportRougness(IOMasterMaterial, IOCollectedTextures, RenderingAsset, AdvancedTransparent.SurfaceRoughness);
-					ExportCutout(IOMasterMaterial, IOCollectedTextures, RenderingAsset, AdvancedTransparent.SurfaceCutout);
-					ExportNormalMap(IOMasterMaterial, IOCollectedTextures, RenderingAsset, AdvancedTransparent.SurfaceNormal);
+					ExportRougness(IOMaterialInstance, IOCollectedTextures, RenderingAsset, AdvancedTransparent.SurfaceRoughness);
+					ExportCutout(IOMaterialInstance, IOCollectedTextures, RenderingAsset, AdvancedTransparent.SurfaceCutout);
+					ExportNormalMap(IOMaterialInstance, IOCollectedTextures, RenderingAsset, AdvancedTransparent.SurfaceNormal);
 				}
 				break;
 
 				case "AdvancedMetal":
 				{
-					IOMasterMaterial.AddBoolean("IsPBR", true);
-					IOMasterMaterial.SetMaterialType(FDatasmithFacadeMasterMaterial.EMasterMaterialType.Opaque);
-					IOMasterMaterial.AddBoolean("IsMetal", true);
-					ExportDiffuse(IOMasterMaterial, IOCollectedTextures, RenderingAsset, AdvancedMetal.MetalF0, sourceMaterialColor);
-					ExportRougness(IOMasterMaterial, IOCollectedTextures, RenderingAsset, AdvancedMetal.SurfaceRoughness);
-					ExportCutout(IOMasterMaterial, IOCollectedTextures, RenderingAsset, AdvancedMetal.SurfaceCutout);
-					ExportNormalMap(IOMasterMaterial, IOCollectedTextures, RenderingAsset, AdvancedMetal.SurfaceNormal);
+					IOMaterialInstance.AddBoolean("IsPBR", true);
+					IOMaterialInstance.SetMaterialType(FDatasmithFacadeMaterialInstance.EMaterialInstanceType.Opaque);
+					IOMaterialInstance.AddBoolean("IsMetal", true);
+					ExportDiffuse(IOMaterialInstance, IOCollectedTextures, RenderingAsset, AdvancedMetal.MetalF0, sourceMaterialColor);
+					ExportRougness(IOMaterialInstance, IOCollectedTextures, RenderingAsset, AdvancedMetal.SurfaceRoughness);
+					ExportCutout(IOMaterialInstance, IOCollectedTextures, RenderingAsset, AdvancedMetal.SurfaceCutout);
+					ExportNormalMap(IOMaterialInstance, IOCollectedTextures, RenderingAsset, AdvancedMetal.SurfaceNormal);
 				}
 				break;
 
 				case "AdvancedOpaque":
 				{
-					IOMasterMaterial.AddBoolean("IsPBR", true);
-					IOMasterMaterial.SetMaterialType(FDatasmithFacadeMasterMaterial.EMasterMaterialType.Opaque);
+					IOMaterialInstance.AddBoolean("IsPBR", true);
+					IOMaterialInstance.SetMaterialType(FDatasmithFacadeMaterialInstance.EMaterialInstanceType.Opaque);
 
-					ExportDiffuse(IOMasterMaterial, IOCollectedTextures, RenderingAsset, AdvancedOpaque.OpaqueAlbedo, sourceMaterialColor);
-					ExportCutout(IOMasterMaterial, IOCollectedTextures, RenderingAsset, AdvancedOpaque.SurfaceCutout);
-					ExportRougness(IOMasterMaterial, IOCollectedTextures, RenderingAsset, AdvancedOpaque.SurfaceRoughness);
-					ExportNormalMap(IOMasterMaterial, IOCollectedTextures, RenderingAsset, AdvancedOpaque.SurfaceNormal);
+					ExportDiffuse(IOMaterialInstance, IOCollectedTextures, RenderingAsset, AdvancedOpaque.OpaqueAlbedo, sourceMaterialColor);
+					ExportCutout(IOMaterialInstance, IOCollectedTextures, RenderingAsset, AdvancedOpaque.SurfaceCutout);
+					ExportRougness(IOMaterialInstance, IOCollectedTextures, RenderingAsset, AdvancedOpaque.SurfaceRoughness);
+					ExportNormalMap(IOMaterialInstance, IOCollectedTextures, RenderingAsset, AdvancedOpaque.SurfaceNormal);
 
 					bool emission = GetBooleanPropertyValue(RenderingAsset, AdvancedOpaque.OpaqueEmission, false);
 					if (emission)
 					{
-						ExportEmission(IOMasterMaterial, IOCollectedTextures, RenderingAsset, AdvancedOpaque.OpaqueLuminance, AdvancedOpaque.OpaqueLuminanceModifier);
+						ExportEmission(IOMaterialInstance, IOCollectedTextures, RenderingAsset, AdvancedOpaque.OpaqueLuminance, AdvancedOpaque.OpaqueLuminanceModifier);
 					}
 					else
 					{
-						IOMasterMaterial.AddBoolean("SelfIlluminationMapEnable", false);
+						IOMaterialInstance.AddBoolean("SelfIlluminationMapEnable", false);
 					}
 				}
 				break;
@@ -1539,7 +1539,7 @@ namespace DatasmithRevitExporter
 		}
 
 		private FDatasmithFacadeTexture ExportTexture(
-			FDatasmithFacadeMasterMaterial IOMasterMaterial, Asset RenderingAsset, string AssetPropertyName, 
+			FDatasmithFacadeMaterialInstance IOMaterialInstance, Asset RenderingAsset, string AssetPropertyName, 
 			string MapParamName, string TextureParamPrefix)
 		{
 			string MapPath = GetTexturePropertyPath(RenderingAsset, AssetPropertyName);
@@ -1554,12 +1554,12 @@ namespace DatasmithRevitExporter
 
 				// Control the Unreal material Normal.
 				FDatasmithFacadeTexture TextureElement = FDatasmithFacadeMaterialsUtils.CreateSimpleTextureElement(MapPath);
-				IOMasterMaterial.AddTexture(MapParamName, TextureElement);
-				IOMasterMaterial.AddFloat($"{TextureParamPrefix}_UVOffsetX", UVOffsetX);
-				IOMasterMaterial.AddFloat($"{TextureParamPrefix}_UVOffsetY", UVOffsetY);
-				IOMasterMaterial.AddFloat($"{TextureParamPrefix}_UVScaleX", UVScaleX);
-				IOMasterMaterial.AddFloat($"{TextureParamPrefix}_UVScaleY", UVScaleY);
-				IOMasterMaterial.AddFloat($"{TextureParamPrefix}_UVWAngle", UVWAngle);
+				IOMaterialInstance.AddTexture(MapParamName, TextureElement);
+				IOMaterialInstance.AddFloat($"{TextureParamPrefix}_UVOffsetX", UVOffsetX);
+				IOMaterialInstance.AddFloat($"{TextureParamPrefix}_UVOffsetY", UVOffsetY);
+				IOMaterialInstance.AddFloat($"{TextureParamPrefix}_UVScaleX", UVScaleX);
+				IOMaterialInstance.AddFloat($"{TextureParamPrefix}_UVScaleY", UVScaleY);
+				IOMaterialInstance.AddFloat($"{TextureParamPrefix}_UVWAngle", UVWAngle);
 
 				return TextureElement;
 			}
@@ -1567,9 +1567,9 @@ namespace DatasmithRevitExporter
 			return null;
 		}
 
-		private void ExportNormalMap(FDatasmithFacadeMasterMaterial IOMasterMaterial, List<FDatasmithFacadeTexture> IOCollectedTextures, Asset RenderingAsset, string AssetPropertyName)
+		private void ExportNormalMap(FDatasmithFacadeMaterialInstance IOMaterialInstance, List<FDatasmithFacadeTexture> IOCollectedTextures, Asset RenderingAsset, string AssetPropertyName)
 		{
-			FDatasmithFacadeTexture TextureElement = ExportTexture(IOMasterMaterial, RenderingAsset, AssetPropertyName, "NormalMap", "NormalMap");
+			FDatasmithFacadeTexture TextureElement = ExportTexture(IOMaterialInstance, RenderingAsset, AssetPropertyName, "NormalMap", "NormalMap");
 			if (TextureElement != null)
 			{
 				IOCollectedTextures.Add(TextureElement);
@@ -1577,74 +1577,74 @@ namespace DatasmithRevitExporter
 			}
 		}
 
-		private void ExportDiffuse(FDatasmithFacadeMasterMaterial IOMasterMaterial, List<FDatasmithFacadeTexture> IOCollectedTextures, Asset RenderingAsset, string AssetProperty, Color DefaultColor)
+		private void ExportDiffuse(FDatasmithFacadeMaterialInstance IOMaterialInstance, List<FDatasmithFacadeTexture> IOCollectedTextures, Asset RenderingAsset, string AssetProperty, Color DefaultColor)
 		{
 			Color DiffuseColor = GetColorPropertyValue(RenderingAsset, AssetProperty, DefaultColor);
 
 			// Control the Unreal material Base Color.
-			IOMasterMaterial.AddColor("DiffuseColor", DiffuseColor.Red / 255.0F, DiffuseColor.Green / 255.0F, DiffuseColor.Blue / 255.0F, 1.0F);
+			IOMaterialInstance.AddColor("DiffuseColor", DiffuseColor.Red / 255.0F, DiffuseColor.Green / 255.0F, DiffuseColor.Blue / 255.0F, 1.0F);
 
-			FDatasmithFacadeTexture TextureElement = ExportTexture(IOMasterMaterial, RenderingAsset, AssetProperty, "DiffuseMap", "DiffuseMap");
+			FDatasmithFacadeTexture TextureElement = ExportTexture(IOMaterialInstance, RenderingAsset, AssetProperty, "DiffuseMap", "DiffuseMap");
 			if (TextureElement != null)
 			{
 				IOCollectedTextures.Add(TextureElement);
 				TextureElement.SetSRGB(FDatasmithFacadeTexture.EColorSpace.sRGB);
 				TextureElement.SetTextureMode(FDatasmithFacadeTexture.ETextureMode.Diffuse);
-				IOMasterMaterial.AddFloat("DiffuseMapFading", 1.0F);
+				IOMaterialInstance.AddFloat("DiffuseMapFading", 1.0F);
 			}
 			else
 			{
-				IOMasterMaterial.AddFloat("DiffuseMapFading", 0.0F);
+				IOMaterialInstance.AddFloat("DiffuseMapFading", 0.0F);
 			}
 		}
 
-		private void ExportEmission(FDatasmithFacadeMasterMaterial IOMasterMaterial, List<FDatasmithFacadeTexture> IOCollectedTextures, Asset RenderingAsset, string LuminanceParam, string ColorParam)
+		private void ExportEmission(FDatasmithFacadeMaterialInstance IOMaterialInstance, List<FDatasmithFacadeTexture> IOCollectedTextures, Asset RenderingAsset, string LuminanceParam, string ColorParam)
 		{
 			float EmissionLuminance = GetFloatPropertyValue(RenderingAsset, LuminanceParam, 1.0f);
 			Color EmissionFilterColor = GetColorPropertyValue(RenderingAsset, ColorParam, new Color(255, 255, 255));
 
 			// Control the Unreal material Emissive Color.
-			IOMasterMaterial.AddFloat("SelfIlluminationLuminance", EmissionLuminance);
-			IOMasterMaterial.AddColor("SelfIlluminationFilter", EmissionFilterColor.Red / 255.0F, EmissionFilterColor.Green / 255.0F, EmissionFilterColor.Blue / 255.0F, 1.0F);
+			IOMaterialInstance.AddFloat("SelfIlluminationLuminance", EmissionLuminance);
+			IOMaterialInstance.AddColor("SelfIlluminationFilter", EmissionFilterColor.Red / 255.0F, EmissionFilterColor.Green / 255.0F, EmissionFilterColor.Blue / 255.0F, 1.0F);
 
-			FDatasmithFacadeTexture TextureElement = ExportTexture(IOMasterMaterial, RenderingAsset, ColorParam, "SelfIlluminationMap", "SelfIlluminationMap");
+			FDatasmithFacadeTexture TextureElement = ExportTexture(IOMaterialInstance, RenderingAsset, ColorParam, "SelfIlluminationMap", "SelfIlluminationMap");
 			if (TextureElement != null)
 			{
 				IOCollectedTextures.Add(TextureElement);
 				TextureElement.SetTextureMode(FDatasmithFacadeTexture.ETextureMode.Other);
-				IOMasterMaterial.AddBoolean("SelfIlluminationMapEnable", true);
+				IOMaterialInstance.AddBoolean("SelfIlluminationMapEnable", true);
 			}
 			else
 			{
-				IOMasterMaterial.AddBoolean("SelfIlluminationMapEnable", false);
+				IOMaterialInstance.AddBoolean("SelfIlluminationMapEnable", false);
 			}
 		}
 
-		private void ExportRougness(FDatasmithFacadeMasterMaterial IOMasterMaterial, List<FDatasmithFacadeTexture> IOCollectedTextures, Asset RenderingAsset, string AssetProperty)
+		private void ExportRougness(FDatasmithFacadeMaterialInstance IOMaterialInstance, List<FDatasmithFacadeTexture> IOCollectedTextures, Asset RenderingAsset, string AssetProperty)
 		{
-			FDatasmithFacadeTexture TextureElement = ExportTexture(IOMasterMaterial, RenderingAsset, AssetProperty, "RoughnessMap", "RoughnessMap");
+			FDatasmithFacadeTexture TextureElement = ExportTexture(IOMaterialInstance, RenderingAsset, AssetProperty, "RoughnessMap", "RoughnessMap");
 			if (TextureElement != null)
 			{
 				IOCollectedTextures.Add(TextureElement);
 				TextureElement.SetTextureMode(FDatasmithFacadeTexture.ETextureMode.Other);
-				IOMasterMaterial.AddBoolean("RoughnessMapEnable", true);
+				IOMaterialInstance.AddBoolean("RoughnessMapEnable", true);
 			}
 			else
 			{
-				IOMasterMaterial.AddBoolean("RoughnessMapEnable", false);
+				IOMaterialInstance.AddBoolean("RoughnessMapEnable", false);
 				float Rougness = GetFloatPropertyValue(RenderingAsset, AssetProperty, 0f);
-				IOMasterMaterial.AddFloat("Rougness", Rougness);
+				IOMaterialInstance.AddFloat("Rougness", Rougness);
 			}
 		}
 
-		private void ExportCutout(FDatasmithFacadeMasterMaterial IOMasterMaterial, List<FDatasmithFacadeTexture> IOCollectedTextures, Asset RenderingAsset, string AssetProperty)
+		private void ExportCutout(FDatasmithFacadeMaterialInstance IOMaterialInstance, List<FDatasmithFacadeTexture> IOCollectedTextures, Asset RenderingAsset, string AssetProperty)
 		{
-			FDatasmithFacadeTexture TextureElement = ExportTexture(IOMasterMaterial, RenderingAsset, AssetProperty, "CutoutOpacityMap", "CutoutMap");
+			FDatasmithFacadeTexture TextureElement = ExportTexture(IOMaterialInstance, RenderingAsset, AssetProperty, "CutoutOpacityMap", "CutoutMap");
 			if (TextureElement != null)
 			{
 				IOCollectedTextures.Add(TextureElement);
 				TextureElement.SetTextureMode(FDatasmithFacadeTexture.ETextureMode.Other);
-				IOMasterMaterial.SetMaterialType(FDatasmithFacadeMasterMaterial.EMasterMaterialType.CutOut);
+				IOMaterialInstance.SetMaterialType(FDatasmithFacadeMaterialInstance.EMaterialInstanceType.CutOut);
 			}
 		}
 
@@ -1652,41 +1652,41 @@ namespace DatasmithRevitExporter
 			Color                          InMaterialColor,
 			float                          InMaterialTransparency, // in range 0.0-1.0
 			float                          InMaterialGlossiness,   // in range 0.0-1.0
-			FDatasmithFacadeMasterMaterial IOMasterMaterial
+			FDatasmithFacadeMaterialInstance IOMaterialInstance
 		)
 		{
 			// Control the Unreal material Base Color.
 			Color MaterialColor = InMaterialColor.IsValid ? InMaterialColor : new Color(255, 255, 255);
-			IOMasterMaterial.AddColor("DiffuseColor", MaterialColor.Red / 255.0F, MaterialColor.Green / 255.0F, MaterialColor.Blue / 255.0F, 1.0F);
-			IOMasterMaterial.AddFloat("DiffuseMapFading", 0.0F);
-			IOMasterMaterial.AddBoolean("TintEnabled", false);
+			IOMaterialInstance.AddColor("DiffuseColor", MaterialColor.Red / 255.0F, MaterialColor.Green / 255.0F, MaterialColor.Blue / 255.0F, 1.0F);
+			IOMaterialInstance.AddFloat("DiffuseMapFading", 0.0F);
+			IOMaterialInstance.AddBoolean("TintEnabled", false);
 
 			// Control the Unreal material Emissive Color.
-			IOMasterMaterial.AddFloat("SelfIlluminationLuminance", 0.0F);
+			IOMaterialInstance.AddFloat("SelfIlluminationLuminance", 0.0F);
 
 			// Control the Unreal material Normal.
-			IOMasterMaterial.AddFloat("BumpAmount", 0.0F);
+			IOMaterialInstance.AddFloat("BumpAmount", 0.0F);
 
 			if (InMaterialTransparency > 0.0F)
 			{
-				IOMasterMaterial.SetMaterialType(FDatasmithFacadeMasterMaterial.EMasterMaterialType.Transparent);
+				IOMaterialInstance.SetMaterialType(FDatasmithFacadeMaterialInstance.EMaterialInstanceType.Transparent);
 
 				// Control the Unreal material Opacity.
-				IOMasterMaterial.AddFloat("Transparency", InMaterialTransparency);
-				IOMasterMaterial.AddFloat("TransparencyMapFading", 0.0F);
+				IOMaterialInstance.AddFloat("Transparency", InMaterialTransparency);
+				IOMaterialInstance.AddFloat("TransparencyMapFading", 0.0F);
 
 				// Control the Unreal material Refraction.
-				IOMasterMaterial.AddFloat("RefractionIndex", 1.0F);
+				IOMaterialInstance.AddFloat("RefractionIndex", 1.0F);
 			}
 			else
 			{
-				IOMasterMaterial.SetMaterialType(FDatasmithFacadeMasterMaterial.EMasterMaterialType.Opaque);
+				IOMaterialInstance.SetMaterialType(FDatasmithFacadeMaterialInstance.EMaterialInstanceType.Opaque);
 
 				// Control the Unreal material Metallic.
-				IOMasterMaterial.AddBoolean("IsMetal", false);
+				IOMaterialInstance.AddBoolean("IsMetal", false);
 
 				// Control the Unreal material Roughness.
-				IOMasterMaterial.AddFloat("Glossiness", InMaterialGlossiness);
+				IOMaterialInstance.AddFloat("Glossiness", InMaterialGlossiness);
 			}
 		}
 
diff --git a/Engine/Source/Programs/Enterprise/Datasmith/DatasmithSketchUpRubyExporter/Private/DatasmithSketchUpExportContext.h b/Engine/Source/Programs/Enterprise/Datasmith/DatasmithSketchUpRubyExporter/Private/DatasmithSketchUpExportContext.h
--- a/Engine/Source/Programs/Enterprise/Datasmith/DatasmithSketchUpRubyExporter/Private/DatasmithSketchUpExportContext.h
+++ b/Engine/Source/Programs/Enterprise/Datasmith/DatasmithSketchUpRubyExporter/Private/DatasmithSketchUpExportContext.h
@@ -18,7 +18,7 @@ class FDatasmithSceneExporter;
 
 class IDatasmithActorElement;
 class IDatasmithCameraActorElement;
-class IDatasmithMasterMaterialElement;
+class IDatasmithMaterialInstanceElement;
 class IDatasmithMeshElement;
 class IDatasmithMetaDataElement;
 class IDatasmithScene;
