[Non-inclusivity] Removed banned words from Datasmith API and code using the API

diff --git a/Engine/Source/Programs/Enterprise/Datasmith/DatasmithSolidworksExporter/Private/DatasmithExporter.cs b/Engine/Source/Programs/Enterprise/Datasmith/DatasmithSolidworksExporter/Private/DatasmithExporter.cs
--- a/Engine/Source/Programs/Enterprise/Datasmith/DatasmithSolidworksExporter/Private/DatasmithExporter.cs
+++ b/Engine/Source/Programs/Enterprise/Datasmith/DatasmithSolidworksExporter/Private/DatasmithExporter.cs
@@ -33,7 +33,7 @@ namespace DatasmithSolidworks
 	{
 		private Dictionary<string, Tuple<EActorType, FDatasmithFacadeActor>> ExportedActorsMap = new Dictionary<string, Tuple<EActorType, FDatasmithFacadeActor>>();
 		private ConcurrentDictionary<string, Tuple<FDatasmithFacadeMeshElement, FDatasmithFacadeMesh>> ExportedMeshesMap = new ConcurrentDictionary<string, Tuple<FDatasmithFacadeMeshElement, FDatasmithFacadeMesh>>();
-		private ConcurrentDictionary<int, FDatasmithFacadeMasterMaterial> ExportedMaterialsMap = new ConcurrentDictionary<int, FDatasmithFacadeMasterMaterial>();
+		private ConcurrentDictionary<int, FDatasmithFacadeMaterialInstance> ExportedMaterialsMap = new ConcurrentDictionary<int, FDatasmithFacadeMaterialInstance>();
 		private ConcurrentDictionary<string, FDatasmithFacadeTexture> ExportedTexturesMap = new ConcurrentDictionary<string, FDatasmithFacadeTexture>();
 		private Dictionary<string, FDatasmithFacadeActorBinding> ExportedActorBindingsMap = new Dictionary<string, FDatasmithFacadeActorBinding>();
 		private Dictionary<string, FDatasmithFacadeVariant> ExportedVariantsMap = new Dictionary<string, FDatasmithFacadeVariant>();
@@ -242,7 +242,7 @@ namespace DatasmithSolidworks
 				{
 					if (!MeshAddedMaterials.Contains(Triangle.MaterialID))
 					{
-						FDatasmithFacadeMasterMaterial Material = null;
+						FDatasmithFacadeMaterialInstance Material = null;
 						ExportedMaterialsMap.TryGetValue(Triangle.MaterialID, out Material);
 
 						if (Material != null)
@@ -399,7 +399,7 @@ namespace DatasmithSolidworks
 						foreach (var MatID in UniqueMaterialsSet)
 						{
 							FMaterial Material = Materials.GetMaterial(MatID);
-							FDatasmithFacadeMasterMaterial DatasmithMaterial = null;
+							FDatasmithFacadeMaterialInstance DatasmithMaterial = null;
 							if (Material != null && ExportedMaterialsMap.TryGetValue(Material.ID, out DatasmithMaterial))
 							{
 								Binding.AddMaterialCapture(DatasmithMaterial);
@@ -540,11 +540,11 @@ namespace DatasmithSolidworks
 		public void ExportMaterials(ConcurrentDictionary<int, FMaterial> InMaterialsMap)
 		{
 			ConcurrentBag<FDatasmithFacadeTexture> CreatedTextures = new ConcurrentBag<FDatasmithFacadeTexture>();
-			ConcurrentBag<FDatasmithFacadeMasterMaterial> CreatedMaterials = new ConcurrentBag<FDatasmithFacadeMasterMaterial>();
+			ConcurrentBag<FDatasmithFacadeMaterialInstance> CreatedMaterials = new ConcurrentBag<FDatasmithFacadeMaterialInstance>();
 			Parallel.ForEach(InMaterialsMap, MatKVP =>
 			{
 				List<FDatasmithFacadeTexture> NewMaterialTextures = null;
-				FDatasmithFacadeMasterMaterial NewMaterial = null;
+				FDatasmithFacadeMaterialInstance NewMaterial = null;
 				if (CreateAndCacheMaterial(MatKVP.Value, out NewMaterialTextures, out NewMaterial))
 				{
 					CreatedMaterials.Add(NewMaterial);
@@ -556,7 +556,7 @@ namespace DatasmithSolidworks
 				}
 			});
 			// Adding stuff to a datasmith scene cannot be multithreaded!
-			foreach (FDatasmithFacadeMasterMaterial Mat in CreatedMaterials)
+			foreach (FDatasmithFacadeMaterialInstance Mat in CreatedMaterials)
 			{
 				DatasmithScene.AddMaterial(Mat);
 			}
@@ -566,7 +566,7 @@ namespace DatasmithSolidworks
 			}
 		}
 
-		private bool CreateAndCacheMaterial(FMaterial InMaterial, out List<FDatasmithFacadeTexture> OutCreatedTextures, out FDatasmithFacadeMasterMaterial OutCreatedMaterial)
+		private bool CreateAndCacheMaterial(FMaterial InMaterial, out List<FDatasmithFacadeTexture> OutCreatedTextures, out FDatasmithFacadeMaterialInstance OutCreatedMaterial)
 		{
 			OutCreatedTextures = null;
 			OutCreatedMaterial = null;
@@ -615,28 +615,28 @@ namespace DatasmithSolidworks
 			float G = Mult * InMaterial.PrimaryColor.G * 1.0f / 255.0f;
 			float B = Mult * InMaterial.PrimaryColor.B * 1.0f / 255.0f;
 
-			FDatasmithFacadeMasterMaterial MasterMaterial = new FDatasmithFacadeMasterMaterial(InMaterial.Name);
+			FDatasmithFacadeMaterialInstance MaterialInstance = new FDatasmithFacadeMaterialInstance(InMaterial.Name);
 
 			OutCreatedTextures = new List<FDatasmithFacadeTexture>();
-			OutCreatedMaterial = MasterMaterial;
+			OutCreatedMaterial = MaterialInstance;
 
-			MasterMaterial.SetMaterialType(FDatasmithFacadeMasterMaterial.EMasterMaterialType.Opaque);
-			MasterMaterial.AddColor("TintColor", R, G, B, 1.0F);
-			MasterMaterial.AddFloat("RoughnessAmount", Roughness);
+			MaterialInstance.SetMaterialType(FDatasmithFacadeMaterialInstance.EMaterialInstanceType.Opaque);
+			MaterialInstance.AddColor("TintColor", R, G, B, 1.0F);
+			MaterialInstance.AddFloat("RoughnessAmount", Roughness);
 
 			if (InMaterial.Transparency > 0.0)
 			{
-				MasterMaterial.SetMaterialType(FDatasmithFacadeMasterMaterial.EMasterMaterialType.Transparent);
-				MasterMaterial.AddFloat("Metalness", Metallic);
+				MaterialInstance.SetMaterialType(FDatasmithFacadeMaterialInstance.EMaterialInstanceType.Transparent);
+				MaterialInstance.AddFloat("Metalness", Metallic);
 
-				MasterMaterial.AddColor("OpacityAndRefraction",
+				MaterialInstance.AddColor("OpacityAndRefraction",
 					0.25f,                              // Opacity
 					1.0f,                               // Refraction
 					0.0f,                               // Refraction Exponent
 					1f - (float)InMaterial.Transparency // Fresnel Opacity
 				);
 
-				FDatasmithFacadeTexture NormalMap = ExportNormalMap(InMaterial, MasterMaterial, "NormalMap");
+				FDatasmithFacadeTexture NormalMap = ExportNormalMap(InMaterial, MaterialInstance, "NormalMap");
 
 				if (NormalMap != null)
 				{
@@ -645,17 +645,17 @@ namespace DatasmithSolidworks
 			}
 			else
 			{
-				MasterMaterial.AddFloat("MetallicAmount", Metallic);
+				MaterialInstance.AddFloat("MetallicAmount", Metallic);
 
 				if (InMaterial.Emission > 0.0)
 				{
-					MasterMaterial.SetMaterialType(FDatasmithFacadeMasterMaterial.EMasterMaterialType.Emissive);
-					MasterMaterial.AddFloat("LuminanceAmount", (float)InMaterial.Emission);
-					MasterMaterial.AddColor("LuminanceFilter", R, G, B, 1.0f);
+					MaterialInstance.SetMaterialType(FDatasmithFacadeMaterialInstance.EMaterialInstanceType.Emissive);
+					MaterialInstance.AddFloat("LuminanceAmount", (float)InMaterial.Emission);
+					MaterialInstance.AddColor("LuminanceFilter", R, G, B, 1.0f);
 				}
 
-				FDatasmithFacadeTexture DiffuseMap = ExportDiffuseMap(InMaterial, MasterMaterial, "ColorMap");
-				FDatasmithFacadeTexture NormalMap = ExportNormalMap(InMaterial, MasterMaterial, "NormalMap");
+				FDatasmithFacadeTexture DiffuseMap = ExportDiffuseMap(InMaterial, MaterialInstance, "ColorMap");
+				FDatasmithFacadeTexture NormalMap = ExportNormalMap(InMaterial, MaterialInstance, "NormalMap");
 
 				if (DiffuseMap != null)
 				{
@@ -667,7 +667,7 @@ namespace DatasmithSolidworks
 				}
 			}
 
-			ExportedMaterialsMap[InMaterial.ID] = MasterMaterial;
+			ExportedMaterialsMap[InMaterial.ID] = MaterialInstance;
 
 			return true;
 		}
@@ -721,7 +721,7 @@ namespace DatasmithSolidworks
 			DatasmithScene.AddLevelSequence(LevelSeq);
 		}
 
-		private FDatasmithFacadeTexture ExportDiffuseMap(FMaterial InMaterial, FDatasmithFacadeMasterMaterial InMasterMaterial, string InParamName)
+		private FDatasmithFacadeTexture ExportDiffuseMap(FMaterial InMaterial, FDatasmithFacadeMaterialInstance InMaterialInstance, string InParamName)
 		{
 			if (!string.IsNullOrEmpty(InMaterial.Texture) && !File.Exists(InMaterial.Texture))
 			{
@@ -746,14 +746,14 @@ namespace DatasmithSolidworks
 					TextureElement.SetTextureMode(TextureMode);
 					ExportedTexturesMap.TryAdd(InMaterial.Texture, TextureElement);
 
-					InMasterMaterial.AddTexture(InParamName, TextureElement);
+					InMaterialInstance.AddTexture(InParamName, TextureElement);
 				}
 			}
 
 			return TextureElement;
 		}
 
-		private FDatasmithFacadeTexture ExportNormalMap(FMaterial InMaterial, FDatasmithFacadeMasterMaterial InMasterMaterial, string InParamName)
+		private FDatasmithFacadeTexture ExportNormalMap(FMaterial InMaterial, FDatasmithFacadeMaterialInstance InMaterialInstance, string InParamName)
 		{
 			if (!string.IsNullOrEmpty(InMaterial.BumpTextureFileName) && !File.Exists(InMaterial.BumpTextureFileName))
 			{
@@ -778,7 +778,7 @@ namespace DatasmithSolidworks
 					TextureElement.SetTextureMode(TextureMode);
 					ExportedTexturesMap.TryAdd(InMaterial.BumpTextureFileName, TextureElement);
 
-					InMasterMaterial.AddTexture(InParamName, TextureElement);
+					InMaterialInstance.AddTexture(InParamName, TextureElement);
 				}
 			}
 
diff --git a/Engine/Source/Runtime/Datasmith/DatasmithCore/Private/DatasmithSceneElementsImpl.cpp b/Engine/Source/Runtime/Datasmith/DatasmithCore/Private/DatasmithSceneElementsImpl.cpp
--- a/Engine/Source/Runtime/Datasmith/DatasmithCore/Private/DatasmithSceneElementsImpl.cpp
+++ b/Engine/Source/Runtime/Datasmith/DatasmithCore/Private/DatasmithSceneElementsImpl.cpp
@@ -377,10 +377,10 @@ const TSharedPtr< IDatasmithShaderElement >& FDatasmithMaterialElementImpl::GetS
 	return Shaders[InIndex];
 }
 
-FDatasmithMasterMaterialElementImpl::FDatasmithMasterMaterialElementImpl(const TCHAR* InName)
-	: FDatasmithBaseMaterialElementImpl(InName, EDatasmithElementType::MasterMaterial)
-	, MaterialType( EDatasmithMasterMaterialType::Auto )
-	, Quality( EDatasmithMasterMaterialQuality::High )
+FDatasmithMaterialIntanceElementImpl::FDatasmithMaterialIntanceElementImpl(const TCHAR* InName)
+	: FDatasmithBaseMaterialElementImpl(InName, EDatasmithElementType::MaterialInstance)
+	, MaterialType( EDatasmithReferenceMaterialType::Auto )
+	, Quality( EDatasmithReferenceMaterialQuality::High )
 {
 	RegisterReferenceProxy(Properties, "Properties");
 
@@ -389,12 +389,12 @@ FDatasmithMasterMaterialElementImpl::FDatasmithMasterMaterialElementImpl(const T
 	Store.RegisterParameter(CustomMaterialPathName, "CustomMaterialPathName");
 }
 
-const TSharedPtr< IDatasmithKeyValueProperty >& FDatasmithMasterMaterialElementImpl::GetProperty( int32 InIndex ) const
+const TSharedPtr< IDatasmithKeyValueProperty >& FDatasmithMaterialIntanceElementImpl::GetProperty( int32 InIndex ) const
 {
 	return Properties.IsValidIndex( InIndex ) ? Properties[InIndex] : FDatasmithKeyValuePropertyImpl::NullPropertyPtr;
 }
 
-const TSharedPtr< IDatasmithKeyValueProperty >& FDatasmithMasterMaterialElementImpl::GetPropertyByName( const TCHAR* InName ) const
+const TSharedPtr< IDatasmithKeyValueProperty >& FDatasmithMaterialIntanceElementImpl::GetPropertyByName( const TCHAR* InName ) const
 {
 	int32 Index = Properties.View().IndexOfByPredicate([InName](const TSharedPtr<IDatasmithKeyValueProperty>& Property){
 		return Property.IsValid() && FCString::Stricmp(Property->GetName(), InName) == 0;
@@ -402,7 +402,7 @@ const TSharedPtr< IDatasmithKeyValueProperty >& FDatasmithMasterMaterialElementI
 	return GetProperty(Index);
 }
 
-void FDatasmithMasterMaterialElementImpl::AddProperty( const TSharedPtr< IDatasmithKeyValueProperty >& InProperty )
+void FDatasmithMaterialIntanceElementImpl::AddProperty( const TSharedPtr< IDatasmithKeyValueProperty >& InProperty )
 {
 	if (!InProperty.IsValid())
 	{
@@ -420,7 +420,7 @@ void FDatasmithMasterMaterialElementImpl::AddProperty( const TSharedPtr< IDatasm
 	}
 }
 
-FMD5Hash FDatasmithMasterMaterialElementImpl::CalculateElementHash(bool bForce)
+FMD5Hash FDatasmithMaterialIntanceElementImpl::CalculateElementHash(bool bForce)
 {
 	if (ElementHash.IsValid() && !bForce)
 	{
diff --git a/Engine/Source/Runtime/Datasmith/DatasmithCore/Private/DatasmithSceneElementsImpl.h b/Engine/Source/Runtime/Datasmith/DatasmithCore/Private/DatasmithSceneElementsImpl.h
--- a/Engine/Source/Runtime/Datasmith/DatasmithCore/Private/DatasmithSceneElementsImpl.h
+++ b/Engine/Source/Runtime/Datasmith/DatasmithCore/Private/DatasmithSceneElementsImpl.h
@@ -1377,16 +1377,16 @@ private:// #ue_directlink_reflect
 	TArray< TSharedPtr< IDatasmithShaderElement > > Shaders;
 };
 
-class FDatasmithMasterMaterialElementImpl : public FDatasmithBaseMaterialElementImpl< IDatasmithMasterMaterialElement >
+class FDatasmithMaterialIntanceElementImpl : public FDatasmithBaseMaterialElementImpl< IDatasmithMaterialInstanceElement >
 {
 public:
-	FDatasmithMasterMaterialElementImpl(const TCHAR* InName);
+	FDatasmithMaterialIntanceElementImpl(const TCHAR* InName);
 
-	virtual EDatasmithMasterMaterialType GetMaterialType() const override { return MaterialType; }
-	virtual void SetMaterialType( EDatasmithMasterMaterialType InType ) override { MaterialType = InType; }
+	virtual EDatasmithReferenceMaterialType GetMaterialType() const override { return MaterialType; }
+	virtual void SetMaterialType( EDatasmithReferenceMaterialType InType ) override { MaterialType = InType; }
 
-	virtual EDatasmithMasterMaterialQuality GetQuality() const override { return Quality; }
-	virtual void SetQuality( EDatasmithMasterMaterialQuality InQuality ) override { Quality = InQuality; }
+	virtual EDatasmithReferenceMaterialQuality GetQuality() const override { return Quality; }
+	virtual void SetQuality( EDatasmithReferenceMaterialQuality InQuality ) override { Quality = InQuality; }
 
 	virtual const TCHAR* GetCustomMaterialPathName() const override { return *(FString&)CustomMaterialPathName; }
 	virtual void SetCustomMaterialPathName( const TCHAR* InPathName ) override { CustomMaterialPathName = InPathName; }
@@ -1401,8 +1401,8 @@ public:
 private:
 	TDatasmithReferenceArrayProxy<IDatasmithKeyValueProperty> Properties;
 
-	TReflected<EDatasmithMasterMaterialType, uint8> MaterialType;
-	TReflected<EDatasmithMasterMaterialQuality, uint8> Quality;
+	TReflected<EDatasmithReferenceMaterialType, uint8> MaterialType;
+	TReflected<EDatasmithReferenceMaterialQuality, uint8> Quality;
 
 	TReflected<FString> CustomMaterialPathName;
 };
diff --git a/Engine/Source/Runtime/Datasmith/DatasmithCore/Private/DatasmithSceneFactory.cpp b/Engine/Source/Runtime/Datasmith/DatasmithCore/Private/DatasmithSceneFactory.cpp
--- a/Engine/Source/Runtime/Datasmith/DatasmithCore/Private/DatasmithSceneFactory.cpp
+++ b/Engine/Source/Runtime/Datasmith/DatasmithCore/Private/DatasmithSceneFactory.cpp
@@ -69,8 +69,8 @@ TSharedPtr< IDatasmithElement > FDatasmithSceneFactory::CreateElement( EDatasmit
 		return CreateShader( InName );
 	case EDatasmithElementType::Material:
 		return CreateMaterial( InName );
-	case EDatasmithElementType::MasterMaterial:
-		return CreateMasterMaterial( InName );
+	case EDatasmithElementType::MaterialInstance:
+		return CreateMaterialInstance( InName );
 	case EDatasmithElementType::UEPbrMaterial:
 		return CreateUEPbrMaterial( InName );
 	case EDatasmithElementType::MaterialExpression:
@@ -224,9 +224,9 @@ TSharedRef< IDatasmithMaterialElement > FDatasmithSceneFactory::CreateMaterial(
 	return MakeShared< FDatasmithMaterialElementImpl >( InName );
 }
 
-TSharedRef< IDatasmithMasterMaterialElement > FDatasmithSceneFactory::CreateMasterMaterial( const TCHAR* InName )
+TSharedRef< IDatasmithMaterialInstanceElement > FDatasmithSceneFactory::CreateMaterialInstance( const TCHAR* InName )
 {
-	return MakeShared< FDatasmithMasterMaterialElementImpl >( InName );
+	return MakeShared< FDatasmithMaterialIntanceElementImpl >( InName );
 }
 
 TSharedRef< IDatasmithUEPbrMaterialElement > FDatasmithSceneFactory::CreateUEPbrMaterial( const TCHAR* InName )
diff --git a/Engine/Source/Runtime/Datasmith/DatasmithCore/Private/DatasmithSceneXmlReader.cpp b/Engine/Source/Runtime/Datasmith/DatasmithCore/Private/DatasmithSceneXmlReader.cpp
--- a/Engine/Source/Runtime/Datasmith/DatasmithCore/Private/DatasmithSceneXmlReader.cpp
+++ b/Engine/Source/Runtime/Datasmith/DatasmithCore/Private/DatasmithSceneXmlReader.cpp
@@ -1325,15 +1325,16 @@ bool FDatasmithSceneXmlReader::ParseXmlFile(TSharedRef< IDatasmithScene >& OutSc
 
 			Objects.Add( Material->GetName(), Material );
 		}
-		// MASTER MATERIALS
-		else if (Nodes[i]->GetTag() == DATASMITH_MASTERMATERIALNAME)
+		// MATERIAL INSTANCES
+		// Support legacy udatasmith files which have the banned word
+		else if (Nodes[i]->GetTag().Equals(TEXT("Mas" "terMaterial")) || Nodes[i]->GetTag() == DATASMITH_MATERIALINSTANCENAME)
 		{
-			TSharedPtr< IDatasmithMasterMaterialElement > MasterMaterial = FDatasmithSceneFactory::CreateMasterMaterial(*Nodes[i]->GetAttribute(TEXT("name")));
+			TSharedPtr< IDatasmithMaterialInstanceElement > ReferenceMaterial = FDatasmithSceneFactory::CreateMaterialInstance(*Nodes[i]->GetAttribute(TEXT("name")));
 
-			ParseMasterMaterial(Nodes[i], MasterMaterial);
-			OutScene->AddMaterial(MasterMaterial);
+			ParseMaterialInstance(Nodes[i], ReferenceMaterial);
+			OutScene->AddMaterial(ReferenceMaterial);
 
-			Objects.Add( MasterMaterial->GetName(), MasterMaterial );
+			Objects.Add( ReferenceMaterial->GetName(), ReferenceMaterial );
 		}
 		// DECAL MATERIALS
 		else if (Nodes[i]->GetTag() == DATASMITH_DECALMATERIALNAME)
@@ -1608,24 +1609,24 @@ void FDatasmithSceneXmlReader::ParseMaterial(FXmlNode* InNode, TSharedPtr< IData
 	}
 }
 
-void FDatasmithSceneXmlReader::ParseMasterMaterial(FXmlNode* InNode, TSharedPtr< IDatasmithMasterMaterialElement >& OutElement) const
+void FDatasmithSceneXmlReader::ParseMaterialInstance(FXmlNode* InNode, TSharedPtr< IDatasmithMaterialInstanceElement >& OutElement) const
 {
 	ParseElement( InNode, OutElement.ToSharedRef() );
 
 	for ( const FXmlAttribute& Attribute : InNode->GetAttributes() )
 	{
-		if (Attribute.GetTag() == DATASMITH_MASTERMATERIALTYPE)
+		if (Attribute.GetTag() == DATASMITH_MATERIALINSTANCENAME)
 		{
-			EDatasmithMasterMaterialType MaterialType = (EDatasmithMasterMaterialType)FMath::Clamp( ValueFromString<int32>( Attribute.GetValue() ), 0, (int32)EDatasmithMasterMaterialType::Count - 1 );
+			EDatasmithReferenceMaterialType MaterialType = (EDatasmithReferenceMaterialType)FMath::Clamp( ValueFromString<int32>( Attribute.GetValue() ), 0, (int32)EDatasmithReferenceMaterialType::Count - 1 );
 
 			OutElement->SetMaterialType( MaterialType );
 		}
-		else if (Attribute.GetTag() == DATASMITH_MASTERMATERIALQUALITY)
+		else if (Attribute.GetTag() == DATASMITH_MATERIALINSTANCEQUALITY)
 		{
-			EDatasmithMasterMaterialQuality Quality = (EDatasmithMasterMaterialQuality)FMath::Clamp( ValueFromString<int32>( Attribute.GetValue() ), 0, (int32)EDatasmithMasterMaterialQuality::Count - 1 );
+			EDatasmithReferenceMaterialQuality Quality = (EDatasmithReferenceMaterialQuality)FMath::Clamp( ValueFromString<int32>( Attribute.GetValue() ), 0, (int32)EDatasmithReferenceMaterialQuality::Count - 1 );
 			OutElement->SetQuality( Quality );
 		}
-		else if (Attribute.GetTag() == DATASMITH_MASTERMATERIALPATHNAME)
+		else if (Attribute.GetTag() == DATASMITH_MATERIALINSTANCEPATHNAME)
 		{
 			OutElement->SetCustomMaterialPathName( *Attribute.GetValue() );
 		}
diff --git a/Engine/Source/Runtime/Datasmith/DatasmithCore/Private/DatasmithSceneXmlWriter.cpp b/Engine/Source/Runtime/Datasmith/DatasmithCore/Private/DatasmithSceneXmlWriter.cpp
--- a/Engine/Source/Runtime/Datasmith/DatasmithCore/Private/DatasmithSceneXmlWriter.cpp
+++ b/Engine/Source/Runtime/Datasmith/DatasmithCore/Private/DatasmithSceneXmlWriter.cpp
@@ -74,7 +74,7 @@ public:
 	static void WriteTextureElement(const TSharedPtr< IDatasmithTextureElement >& TextureElement, FArchive& Archive, int32 Indent);
 
 	static void WriteMaterialElement( TSharedPtr< IDatasmithMaterialElement >& MaterialElement, FArchive& Archive, int32 Indent);
-	static void WriteMasterMaterialElement(TSharedPtr< IDatasmithMasterMaterialElement >& MasterMaterialElement, FArchive& Archive, int32 Indent);
+	static void WriteMaterialInstanceElement(TSharedPtr< IDatasmithMaterialInstanceElement >& MaterialInstanceElement, FArchive& Archive, int32 Indent);
 	static void WriteDecalMaterialElement(TSharedPtr< IDatasmithDecalMaterialElement >& DecalMaterialElement, FArchive& Archive, int32 Indent);
 
 	static void WriteUEPbrMaterialElement( const TSharedRef< IDatasmithUEPbrMaterialElement >& MaterialElement, FArchive& Archive, int32 Indent );
@@ -1378,28 +1378,28 @@ void FDatasmithSceneXmlWriterImpl::WriteMaterialElement( TSharedPtr< IDatasmithM
 	SerializeToArchive( Archive, XmlString );
 }
 
-void FDatasmithSceneXmlWriterImpl::WriteMasterMaterialElement(TSharedPtr< IDatasmithMasterMaterialElement >& MasterMaterialElement, FArchive& Archive, int32 Indent)
+void FDatasmithSceneXmlWriterImpl::WriteMaterialInstanceElement(TSharedPtr< IDatasmithMaterialInstanceElement >& MaterialInstanceElement, FArchive& Archive, int32 Indent)
 {
 	WriteIndent(Archive, Indent);
 
 	FString XmlString = TEXT("<");
-	XmlString += FString(DATASMITH_MASTERMATERIALNAME) + TEXT(" name=\"") + SanitizeXMLText(FDatasmithUtils::SanitizeFileName(MasterMaterialElement->GetName())) + TEXT("\" ");
-	XmlString += TEXT(" label=\"") + SanitizeXMLText(MasterMaterialElement->GetLabel()) + TEXT("\" ");
-	XmlString += FString(DATASMITH_MASTERMATERIALTYPE) + TEXT("=\"") + FString::FromInt( (int32)MasterMaterialElement->GetMaterialType() ) + TEXT("\" ");
-	XmlString += FString(DATASMITH_MASTERMATERIALQUALITY) + TEXT("=\"") + FString::FromInt( (int32)MasterMaterialElement->GetQuality() ) + TEXT("\" ");
+	XmlString += FString(DATASMITH_MATERIALINSTANCENAME) + TEXT(" name=\"") + SanitizeXMLText(FDatasmithUtils::SanitizeFileName(MaterialInstanceElement->GetName())) + TEXT("\" ");
+	XmlString += TEXT(" label=\"") + SanitizeXMLText(MaterialInstanceElement->GetLabel()) + TEXT("\" ");
+	XmlString += FString(DATASMITH_MATERIALINSTANCETYPE) + TEXT("=\"") + FString::FromInt( (int32)MaterialInstanceElement->GetMaterialType() ) + TEXT("\" ");
+	XmlString += FString(DATASMITH_MATERIALINSTANCEQUALITY) + TEXT("=\"") + FString::FromInt( (int32)MaterialInstanceElement->GetQuality() ) + TEXT("\" ");
 
-	if ( MasterMaterialElement->GetMaterialType() == EDatasmithMasterMaterialType::Custom )
+	if ( MaterialInstanceElement->GetMaterialType() == EDatasmithReferenceMaterialType::Custom )
 	{
-		XmlString += FString(DATASMITH_MASTERMATERIALPATHNAME) + TEXT("=\"") + SanitizeXMLText( MasterMaterialElement->GetCustomMaterialPathName() ) + TEXT("\" ");
+		XmlString += FString(DATASMITH_MATERIALINSTANCEPATHNAME) + TEXT("=\"") + SanitizeXMLText( MaterialInstanceElement->GetCustomMaterialPathName() ) + TEXT("\" ");
 	}
 
 	XmlString += FString(TEXT(">")) + LINE_TERMINATOR;
 	SerializeToArchive( Archive, XmlString );
 
-	WriteKeyValueProperties( *MasterMaterialElement, Archive, Indent );
+	WriteKeyValueProperties( *MaterialInstanceElement, Archive, Indent );
 
 	WriteIndent(Archive, Indent);
-	XmlString = TEXT("</") + FString(DATASMITH_MASTERMATERIALNAME) + TEXT(">") + LINE_TERMINATOR;
+	XmlString = TEXT("</") + FString(DATASMITH_MATERIALINSTANCENAME) + TEXT(">") + LINE_TERMINATOR;
 	SerializeToArchive( Archive, XmlString );
 }
 
@@ -1990,10 +1990,10 @@ void FDatasmithSceneXmlWriter::Serialize( TSharedRef< IDatasmithScene > Datasmit
 			TSharedPtr< IDatasmithMaterialElement > MaterialElement = StaticCastSharedPtr< IDatasmithMaterialElement >( BaseMaterialElement );
 			FDatasmithSceneXmlWriterImpl::WriteMaterialElement( MaterialElement, Archive, 1 );
 		}
-		else if ( BaseMaterialElement->IsA( EDatasmithElementType::MasterMaterial ) )
+		else if ( BaseMaterialElement->IsA( EDatasmithElementType::MaterialInstance ) )
 		{
-			TSharedPtr< IDatasmithMasterMaterialElement > MasterMaterialElement = StaticCastSharedPtr< IDatasmithMasterMaterialElement >( BaseMaterialElement );
-			FDatasmithSceneXmlWriterImpl::WriteMasterMaterialElement( MasterMaterialElement, Archive, 1 );
+			TSharedPtr< IDatasmithMaterialInstanceElement > MaterialInstanceElement = StaticCastSharedPtr< IDatasmithMaterialInstanceElement >( BaseMaterialElement );
+			FDatasmithSceneXmlWriterImpl::WriteMaterialInstanceElement( MaterialInstanceElement, Archive, 1 );
 		}
 		else if ( BaseMaterialElement->IsA( EDatasmithElementType::DecalMaterial ) )
 		{
diff --git a/Engine/Source/Runtime/Datasmith/DatasmithCore/Private/DatasmithUtils.cpp b/Engine/Source/Runtime/Datasmith/DatasmithCore/Private/DatasmithUtils.cpp
--- a/Engine/Source/Runtime/Datasmith/DatasmithCore/Private/DatasmithUtils.cpp
+++ b/Engine/Source/Runtime/Datasmith/DatasmithCore/Private/DatasmithUtils.cpp
@@ -807,7 +807,7 @@ namespace DatasmithSceneUtilsImpl
 		return EDatasmithTextureMode::Diffuse;
 	};
 
-	void CheckMasterMaterialTextures( IDatasmithScene& Scene )
+	void CheckMaterialInstanceTextures( IDatasmithScene& Scene )
 	{
 		TSet<FString> ProcessedTextures;
 
@@ -818,9 +818,9 @@ namespace DatasmithSceneUtilsImpl
 		{
 			const TSharedPtr< IDatasmithBaseMaterialElement >& BaseMaterial = Scene.GetMaterial( MaterialIndex );
 
-			if ( BaseMaterial->IsA( EDatasmithElementType::MasterMaterial ) )
+			if ( BaseMaterial->IsA( EDatasmithElementType::MaterialInstance ) )
 			{
-				const TSharedPtr< IDatasmithMasterMaterialElement >& Material = StaticCastSharedPtr< IDatasmithMasterMaterialElement >( BaseMaterial );
+				const TSharedPtr< IDatasmithMaterialInstanceElement >& Material = StaticCastSharedPtr< IDatasmithMaterialInstanceElement >( BaseMaterial );
 
 				for ( int32 i = 0; i < Material->GetPropertiesCount(); ++i )
 				{
@@ -1042,7 +1042,7 @@ namespace DatasmithSceneUtilsImpl
 			}
 		}
 
-		void ScanMasterMaterialElement(IDatasmithMasterMaterialElement* MaterialElement)
+		void ScanMaterialInstanceElement(IDatasmithMaterialInstanceElement* MaterialElement)
 		{
 			for ( int32 Index = 0; Index < MaterialElement->GetPropertiesCount(); ++Index )
 			{
@@ -1309,9 +1309,9 @@ namespace DatasmithSceneUtilsImpl
 				{
 					ScanPbrMaterialElement(static_cast< IDatasmithUEPbrMaterialElement* >( MaterialElement.Get() ));
 				}
-				else if ( MaterialElement->IsA( EDatasmithElementType::MasterMaterial ) )
+				else if ( MaterialElement->IsA( EDatasmithElementType::MaterialInstance ) )
 				{
-					ScanMasterMaterialElement(static_cast< IDatasmithMasterMaterialElement* >( MaterialElement.Get() ));
+					ScanMaterialInstanceElement(static_cast< IDatasmithMaterialInstanceElement* >( MaterialElement.Get() ));
 				}
 				else if (MaterialElement->IsA(EDatasmithElementType::DecalMaterial))
 				{
@@ -1578,7 +1578,7 @@ void FDatasmithSceneUtils::CleanUpScene(TSharedRef<IDatasmithScene> Scene, bool
 		FixIesTextures(*Scene, Scene->GetActor(Index));
 	}
 
-	CheckMasterMaterialTextures(*Scene);
+	CheckMaterialInstanceTextures(*Scene);
 
 	CleanUpEnvironments(Scene);
 
diff --git a/Engine/Source/Runtime/Datasmith/DatasmithCore/Private/DirectLink/DatasmithDirectLinkTools.cpp b/Engine/Source/Runtime/Datasmith/DatasmithCore/Private/DirectLink/DatasmithDirectLinkTools.cpp
--- a/Engine/Source/Runtime/Datasmith/DatasmithCore/Private/DirectLink/DatasmithDirectLinkTools.cpp
+++ b/Engine/Source/Runtime/Datasmith/DatasmithCore/Private/DirectLink/DatasmithDirectLinkTools.cpp
@@ -34,7 +34,7 @@ const TCHAR* GetElementTypeName(const IDatasmithElement* Element)
 	DS_ELEMENT_TYPE(EDatasmithElementType::MaterialId)
 	DS_ELEMENT_TYPE(EDatasmithElementType::Texture)
 	DS_ELEMENT_TYPE(EDatasmithElementType::KeyValueProperty)
-	DS_ELEMENT_TYPE(EDatasmithElementType::MasterMaterial)
+	DS_ELEMENT_TYPE(EDatasmithElementType::MaterialInstance)
 	DS_ELEMENT_TYPE(EDatasmithElementType::BaseMaterial)
 	DS_ELEMENT_TYPE(EDatasmithElementType::Shader)
 	DS_ELEMENT_TYPE(EDatasmithElementType::Camera)
diff --git a/Engine/Source/Runtime/Datasmith/DatasmithCore/Public/DatasmithDefinitions.h b/Engine/Source/Runtime/Datasmith/DatasmithCore/Public/DatasmithDefinitions.h
--- a/Engine/Source/Runtime/Datasmith/DatasmithCore/Public/DatasmithDefinitions.h
+++ b/Engine/Source/Runtime/Datasmith/DatasmithCore/Public/DatasmithDefinitions.h
@@ -2,6 +2,7 @@
 #pragma once
 
 #include "CoreTypes.h"
+#include "Misc/CoreMiscDefines.h"
 #include "Misc/EnumClassFlags.h"
 
 #ifdef WITH_COREUOBJECT
@@ -27,7 +28,8 @@ enum class EDatasmithElementType : uint64
 	Camera                         = 1ull << 10,
 	Shader                         = 1ull << 11,
 	BaseMaterial                   = 1ull << 12,
-	MasterMaterial                 = 1ull << 13,
+	MasterMaterial UE_DEPRECATED(5.2, "EDatasmithElementType::MasterMaterial will not be supported in 5.2. Please use EDatasmithElementType::MaterialInstance instead.") = 1ull << 13,
+	MaterialInstance               = 1ull << 13,
 	KeyValueProperty               = 1ull << 14,
 	Texture                        = 1ull << 15,
 	MaterialId                     = 1ull << 16,
@@ -225,14 +227,14 @@ enum class EDatasmithMaterialMode
 	MixedMetal
 };
 
-enum class EDatasmithMasterMaterialType : uint8
+enum class EDatasmithReferenceMaterialType : uint8
 {
-	/** Let Datasmith figure which master material to use */
+	/** Let Datasmith figure which reference material to use */
 	Auto,
 	Opaque,
 	Transparent,
 	ClearCoat,
-	/** Instantiate a master material from a specified one */
+	/** Instantiate a reference material from a specified one */
 	Custom,
 	/** Material has a transparent cutout map */
 	CutOut,
@@ -242,7 +244,7 @@ enum class EDatasmithMasterMaterialType : uint8
 	Count
 };
 
-enum class EDatasmithMasterMaterialQuality : uint8
+enum class EDatasmithReferenceMaterialQuality : uint8
 {
 	High,
 	Low,
@@ -559,10 +561,10 @@ static const TCHAR* KeyValuePropertyTypeStrings[] = { TEXT("String"), TEXT("Colo
 #define DATASMITH_PARENTMATERIALLABEL			TEXT("ParentLabel")
 #define DATASMITH_UEPBRMATERIALNAME				TEXT("UEPbrMaterial")
 
-#define DATASMITH_MASTERMATERIALNAME			TEXT("MasterMaterial")
-#define DATASMITH_MASTERMATERIALTYPE			TEXT("Type")
-#define DATASMITH_MASTERMATERIALQUALITY			TEXT("Quality")
-#define DATASMITH_MASTERMATERIALPATHNAME		TEXT("PathName")
+#define DATASMITH_MATERIALINSTANCENAME			TEXT("MaterialInstance")
+#define DATASMITH_MATERIALINSTANCETYPE			TEXT("Type")
+#define DATASMITH_MATERIALINSTANCEQUALITY		TEXT("Quality")
+#define DATASMITH_MATERIALINSTANCEPATHNAME		TEXT("PathName")
 
 #define DATASMITH_TEXTURENAME					TEXT("Texture")
 #define DATASMITH_TEXTURECOMPNAME				TEXT("Texturecomp")
diff --git a/Engine/Source/Runtime/Datasmith/DatasmithCore/Public/DatasmithSceneFactory.h b/Engine/Source/Runtime/Datasmith/DatasmithCore/Public/DatasmithSceneFactory.h
--- a/Engine/Source/Runtime/Datasmith/DatasmithCore/Public/DatasmithSceneFactory.h
+++ b/Engine/Source/Runtime/Datasmith/DatasmithCore/Public/DatasmithSceneFactory.h
@@ -56,7 +56,15 @@ public:
 
 	static TSharedRef< IDatasmithMaterialElement > CreateMaterial( const TCHAR* InName );
 
-	static TSharedRef< IDatasmithMasterMaterialElement > CreateMasterMaterial( const TCHAR* InName );
+PRAGMA_DISABLE_DEPRECATION_WARNINGS
+	UE_DEPRECATED(5.2, "FDatasmithSceneFactory::CreateMasterMaterial will not be supported in 5.2. Please use FDatasmithSceneFactory::CreateMaterialInstance instead.")
+	static TSharedRef< IDatasmithMasterMaterialElement > CreateMasterMaterial(const TCHAR* InName)
+	{
+		return CreateMaterialInstance(InName);
+	}
+PRAGMA_ENABLE_DEPRECATION_WARNINGS
+
+	static TSharedRef< IDatasmithMaterialInstanceElement > CreateMaterialInstance(const TCHAR* InName);
 
 	static TSharedRef< IDatasmithUEPbrMaterialElement > CreateUEPbrMaterial( const TCHAR* InName );
 
diff --git a/Engine/Source/Runtime/Datasmith/DatasmithCore/Public/DatasmithSceneXmlReader.h b/Engine/Source/Runtime/Datasmith/DatasmithCore/Public/DatasmithSceneXmlReader.h
--- a/Engine/Source/Runtime/Datasmith/DatasmithCore/Public/DatasmithSceneXmlReader.h
+++ b/Engine/Source/Runtime/Datasmith/DatasmithCore/Public/DatasmithSceneXmlReader.h
@@ -65,7 +65,7 @@ private:
 	void ParseColor(FXmlNode* InNode, FLinearColor& OutColor) const;
 	void ParseComp(FXmlNode* InNode, TSharedPtr< IDatasmithCompositeTexture >& OutCompTexture, bool bInIsNormal = false) const;
 	void ParseMaterial(FXmlNode* InNode, TSharedPtr< IDatasmithMaterialElement >& OutElement) const;
-	void ParseMasterMaterial(FXmlNode* InNode, TSharedPtr< IDatasmithMasterMaterialElement >& OutElement) const;
+	void ParseMaterialInstance(FXmlNode* InNode, TSharedPtr< IDatasmithMaterialInstanceElement >& OutElement) const;
 	void ParseDecalMaterial(FXmlNode* InNode, TSharedPtr< IDatasmithDecalMaterialElement >& OutElement) const;
 	void ParseUEPbrMaterial(FXmlNode* InNode, TSharedPtr< IDatasmithUEPbrMaterialElement >& OutElement) const;
 	void ParseCustomActor(FXmlNode* InNode, TSharedPtr< IDatasmithCustomActorElement >& OutElement) const;
diff --git a/Engine/Source/Runtime/Datasmith/DatasmithCore/Public/IDatasmithSceneElements.h b/Engine/Source/Runtime/Datasmith/DatasmithCore/Public/IDatasmithSceneElements.h
--- a/Engine/Source/Runtime/Datasmith/DatasmithCore/Public/IDatasmithSceneElements.h
+++ b/Engine/Source/Runtime/Datasmith/DatasmithCore/Public/IDatasmithSceneElements.h
@@ -10,6 +10,7 @@
 #include "Math/Color.h"
 #include "Math/Quat.h"
 #include "Math/Vector.h"
+#include "Misc/CoreMiscDefines.h"
 #include "Misc/SecureHash.h"
 #include "Templates/SharedPointer.h"
 
@@ -641,16 +642,16 @@ public:
 	virtual const TSharedPtr< IDatasmithShaderElement >& GetShader(int32 InIndex) const = 0;
 };
 
-class DATASMITHCORE_API IDatasmithMasterMaterialElement : public IDatasmithBaseMaterialElement
+class DATASMITHCORE_API IDatasmithMaterialInstanceElement : public IDatasmithBaseMaterialElement
 {
 public:
-	virtual ~IDatasmithMasterMaterialElement() {}
+	virtual ~IDatasmithMaterialInstanceElement() {}
 
-	virtual EDatasmithMasterMaterialType GetMaterialType() const = 0;
-	virtual void SetMaterialType(EDatasmithMasterMaterialType InType) = 0;
+	virtual EDatasmithReferenceMaterialType GetMaterialType() const = 0;
+	virtual void SetMaterialType(EDatasmithReferenceMaterialType InType) = 0;
 
-	virtual EDatasmithMasterMaterialQuality GetQuality() const = 0;
-	virtual void SetQuality(EDatasmithMasterMaterialQuality InQuality) = 0;
+	virtual EDatasmithReferenceMaterialQuality GetQuality() const = 0;
+	virtual void SetQuality(EDatasmithReferenceMaterialQuality InQuality) = 0;
 
 	/** Only used when the material type is set to Custom. The path name to an existing material to instantiate. */
 	virtual const TCHAR* GetCustomMaterialPathName() const = 0;
@@ -1785,3 +1786,6 @@ public:
 	/** Attach the actor to the scene root. Detach the actor if it was already attached. */
 	virtual void AttachActorToSceneRoot(const TSharedPtr< IDatasmithActorElement >& Child, EDatasmithActorAttachmentRule AttachmentRule) = 0;
 };
+
+
+using IDatasmithMasterMaterialElement UE_DEPRECATED(5.2, "IDatasmithMasterMaterialElement will not be supported in 5.2. Please use IDatasmithMaterialInstanceElement instead.") = IDatasmithMaterialInstanceElement;
