[3ds Max] Version displayed with CL #

diff --git a/Engine/Source/Programs/Enterprise/Datasmith/DatasmithMaxExporter/Private/DatasmithMaxDirectLinkMaxScript.cpp b/Engine/Source/Programs/Enterprise/Datasmith/DatasmithMaxExporter/Private/DatasmithMaxDirectLinkMaxScript.cpp
--- a/Engine/Source/Programs/Enterprise/Datasmith/DatasmithMaxExporter/Private/DatasmithMaxDirectLinkMaxScript.cpp
+++ b/Engine/Source/Programs/Enterprise/Datasmith/DatasmithMaxExporter/Private/DatasmithMaxDirectLinkMaxScript.cpp
@@ -87,7 +87,7 @@ Value* OnLoad_cf(Value **arg_list, int count)
 		MessagesDialog = Ui::CreateMessagesWindow();
 	}
 
-	LogInfo(FString::Printf(TEXT("Initialized Datasmith SDK v.%s"), *FDatasmithUtils::GetEnterpriseVersionAsString()));
+	LogInfo(FString::Printf(TEXT("Initialized Datasmith SDK v.%s"), *FDatasmithUtils::GetEnterpriseVersionAsString(true)));
 
 	return bool_result(bResult);
 }
diff --git a/Engine/Source/Programs/Enterprise/Datasmith/DatasmithMaxExporter/Private/UI/Datasmith3dsMaxUI.cpp b/Engine/Source/Programs/Enterprise/Datasmith/DatasmithMaxExporter/Private/UI/Datasmith3dsMaxUI.cpp
--- a/Engine/Source/Programs/Enterprise/Datasmith/DatasmithMaxExporter/Private/UI/Datasmith3dsMaxUI.cpp
+++ b/Engine/Source/Programs/Enterprise/Datasmith/DatasmithMaxExporter/Private/UI/Datasmith3dsMaxUI.cpp
@@ -247,7 +247,7 @@ public:
 	FText GetWindowTitle()
 	{
 		FFormatNamedArguments WindowTitleArgs;
-		WindowTitleArgs.Add(TEXT("DatasmithSDKVersion"), FText::FromString(FDatasmithUtils::GetEnterpriseVersionAsString()));
+		WindowTitleArgs.Add(TEXT("DatasmithSDKVersion"), FText::FromString(FDatasmithUtils::GetEnterpriseVersionAsString(true)));
 		return FText::Format(LOCTEXT("DatasmithMessagesWindowTitle", "Datasmith Messages v.{DatasmithSDKVersion}"), WindowTitleArgs);
 	}
 
