DatasmithSDK: Fixed unacceptable word

diff --git a/Engine/Source/Programs/Enterprise/Datasmith/DatasmithSDK/Documentation/build/DatasmithSDK_config.txt b/Engine/Source/Programs/Enterprise/Datasmith/DatasmithSDK/Documentation/build/DatasmithSDK_config.txt
--- a/Engine/Source/Programs/Enterprise/Datasmith/DatasmithSDK/Documentation/build/DatasmithSDK_config.txt
+++ b/Engine/Source/Programs/Enterprise/Datasmith/DatasmithSDK/Documentation/build/DatasmithSDK_config.txt
@@ -1405,7 +1405,7 @@ CHM_FILE               =
 HHC_LOCATION           =
 
 # The GENERATE_CHI flag controls if a separate .chi index file is generated
-# (YES) or that it should be included in the master .chm file (NO).
+# (YES) or that it should be included in the main .chm file (NO).
 # The default value is: NO.
 # This tag requires that the tag GENERATE_HTMLHELP is set to YES.
 
