[3ds Max] Fixed texture sync when multiple texmaps can use same image file

- to allow the fix refactored texture element creation for texmaps. Removed legacy code duplication of texmap conversion code and texture element creation.

- also added more statistics and for baked textures - total/average stats and individual texmap

diff --git a/Engine/Source/Programs/Enterprise/Datasmith/DatasmithMaxExporter/Private/DatasmithMaxCorona.cpp b/Engine/Source/Programs/Enterprise/Datasmith/DatasmithMaxExporter/Private/DatasmithMaxCorona.cpp
--- a/Engine/Source/Programs/Enterprise/Datasmith/DatasmithMaxExporter/Private/DatasmithMaxCorona.cpp
+++ b/Engine/Source/Programs/Enterprise/Datasmith/DatasmithMaxExporter/Private/DatasmithMaxCorona.cpp
@@ -67,62 +67,6 @@ FString FDatasmithMaxMatWriter::DumpBitmapCorona(TSharedPtr<IDatasmithCompositeT
 	return *Base;
 }
 
-void FDatasmithMaxMatWriter::GetCoronaTexmap(TSharedRef< IDatasmithScene > DatasmithScene, BitmapTex* InBitmapTex, TArray<TSharedPtr< IDatasmithTextureElement >>* OutTextureElements)
-{
-	FString Path = TEXT("");
-
-	int NumParamBlocks = InBitmapTex->NumParamBlocks();
-	
-	for (int j = 0; j < NumParamBlocks; j++)
-	{
-		IParamBlock2* ParamBlock2 = InBitmapTex->GetParamBlockByID((short)j);
-		// The the descriptor to 'decode'
-		ParamBlockDesc2* ParamBlockDesc = ParamBlock2->GetDesc();
-		// Loop through all the defined parameters therein
-		for (int i = 0; i < ParamBlockDesc->count; i++)
-		{
-			const ParamDef& ParamDefinition = ParamBlockDesc->paramdefs[i];
-
-			if (FCString::Stricmp(ParamDefinition.int_name, TEXT("filename")) == 0)
-			{
-				Path = FDatasmithMaxSceneExporter::GetActualPath(ParamBlock2->GetStr(ParamDefinition.ID, GetCOREInterface()->GetTime()));
-				continue;
-			}
-		}
-		ParamBlock2->ReleaseDesc();
-	}
-
-	if (Path.IsEmpty())
-	{
-		return;
-	}
-
-	float Gamma = GetCoronaTexmapGamma(InBitmapTex);
-
-	FString BaseName = FPaths::GetBaseFilename(Path);
-	FString Base = BaseName + FString("_") + FString::SanitizeFloat(Gamma).Replace(TEXT("."), TEXT("_")) + TextureSuffix;
-
-	for (int i = 0; i < DatasmithScene->GetTexturesCount(); i++)
-	{
-		if (DatasmithScene->GetTexture(i)->GetFile() == Path && DatasmithScene->GetTexture(i)->GetName() == Base)
-		{
-			return;
-		}
-	}
-
-	TSharedPtr< IDatasmithTextureElement > TextureElement = FDatasmithSceneFactory::CreateTexture(*Base);
-	if (gammaMgr.IsEnabled())
-	{
-		TextureElement->SetRGBCurve(Gamma / 2.2f);
-	}
-	TextureElement->SetFile(*Path);
-	DatasmithScene->AddTexture(TextureElement);
-	if (OutTextureElements)
-	{
-		OutTextureElements->Add(TextureElement);
-	}
-}
-
 void FDatasmithMaxMatWriter::ExportCoronaMaterial(TSharedRef< IDatasmithScene > DatasmithScene, TSharedPtr< IDatasmithMaterialElement >& MaterialElement, Mtl* Material)
 {
 	TSharedPtr< IDatasmithShaderElement > MaterialShader = FDatasmithSceneFactory::CreateShader((TCHAR*)Material->GetName().data());
diff --git a/Engine/Source/Programs/Enterprise/Datasmith/DatasmithMaxExporter/Private/DatasmithMaxDirectLinkMaterial.cpp b/Engine/Source/Programs/Enterprise/Datasmith/DatasmithMaxExporter/Private/DatasmithMaxDirectLinkMaterial.cpp
--- a/Engine/Source/Programs/Enterprise/Datasmith/DatasmithMaxExporter/Private/DatasmithMaxDirectLinkMaterial.cpp
+++ b/Engine/Source/Programs/Enterprise/Datasmith/DatasmithMaxExporter/Private/DatasmithMaxDirectLinkMaterial.cpp
@@ -12,6 +12,9 @@
 
 
 #include "DatasmithSceneFactory.h"
+#include "DatasmithExportOptions.h"
+
+#include "Misc/Paths.h"
 
 
 #include "Logging/LogMacros.h"
@@ -241,6 +244,9 @@ void FMaterialsCollectionTracker::Reset()
 
 	UsedTextureToMaterialTracker.Reset();
 	UsedTextureToDatasmithElement.Reset();
+	TextureElementToTexmap.Reset();
+
+	TextureElementsAddedToScene.Reset();
 }
 
 void FMaterialsCollectionTracker::UpdateMaterial(FMaterialTracker* MaterialTracker)
@@ -271,7 +277,7 @@ void FMaterialsCollectionTracker::ConvertMaterial(Mtl* Material, TSharedRef<IDat
 
 	SCENE_UPDATE_STAT_INC(UpdateMaterials, Total);
 
-	TMap<Texmap*, TSet<TSharedPtr<IDatasmithTextureElement>>> TexmapsUsedByMaterial;
+	TSet<Texmap*> TexmapsUsedByMaterial;
 	FMaterialConversionContext MaterialConversionContext = {TexmapsUsedByMaterial, *this};
 	TGuardValue MaterialConversionContextGuard(FDatasmithMaxMaterialsToUEPbrManager::Context, &MaterialConversionContext);
 
@@ -285,18 +291,14 @@ void FMaterialsCollectionTracker::ConvertMaterial(Mtl* Material, TSharedRef<IDat
 	}
 
 	// Tie texture used by an actual material to tracked material
-	for (TPair<Texmap*, TSet<TSharedPtr<IDatasmithTextureElement>>> TexmapAndDatasmithElements : TexmapsUsedByMaterial)
+	for (Texmap* Tex : TexmapsUsedByMaterial)
 	{
-		Texmap* Tex = TexmapAndDatasmithElements.Key;
 		for (FMaterialTracker* MaterialTracker : UsedMaterialToMaterialTracker[Material])
 		{
 			MaterialTracker->AddActualTexture(Tex);
 			UsedTextureToMaterialTracker.FindOrAdd(Tex).Add(MaterialTracker);
 		}
 
-		// Some material converters create texture elements during material conversion(like bakeable) - record created texture elements
-		UsedTextureToDatasmithElement.FindOrAdd(Tex).Append(TexmapAndDatasmithElements.Value);
-
 		TexmapsConverted.Add(Tex);
 	}
 
@@ -333,26 +335,86 @@ void FMaterialsCollectionTracker::RemoveConvertedMaterial(FMaterialTracker& Mate
 		}
 	}
 
-	for (Texmap* Texture: MaterialTracker.GetActualTexmaps())
+	for (Texmap* Tex: MaterialTracker.GetActualTexmaps())
 	{
-		TSet<FMaterialTracker*>& MaterialTrackersForTexture = UsedTextureToMaterialTracker[Texture];
+		TSet<FMaterialTracker*>& MaterialTrackersForTexture = UsedTextureToMaterialTracker[Tex];
 		MaterialTrackersForTexture.Remove(&MaterialTracker);
 
 		if (!MaterialTrackersForTexture.Num()) // No tracked materials are using this texture anymore
 		{
-			UsedTextureToMaterialTracker.Remove(Texture);
+			UsedTextureToMaterialTracker.Remove(Tex);
 
-			for (const TSharedPtr<IDatasmithTextureElement>& TextureElement: UsedTextureToDatasmithElement[Texture])
+			for (const TSharedPtr<IDatasmithTextureElement>& TextureElement: UsedTextureToDatasmithElement[Tex])
 			{
-				SceneTracker.RemoveTexture(TextureElement);
+				TSet<Texmap*>& Texmaps = TextureElementToTexmap[TextureElement];
+				Texmaps.Remove(Tex);
+
+				if (Texmaps.IsEmpty())  // This was the last texmap that produced this element
+				{
+					RemoveTextureElement(TextureElement);
+					TextureElementToTexmap.Remove(TextureElement);
+				}
 			}
-			UsedTextureToDatasmithElement.Remove(Texture);
+			UsedTextureToDatasmithElement.Remove(Tex);
 		}
 	}
 
 	MaterialTracker.ResetActualMaterialAndTextures();
 }
 
+void FMaterialsCollectionTracker::UpdateTexmap(Texmap* Texmap)
+{
+	if (UsedTextureToDatasmithElement.Contains(Texmap))
+	{
+		// Don't update texmap that wasn't released - this means that it doesn't need update
+		// Texmap is released when every material that uses it is invalidated or removed
+		// When Texmap wasn't released means that some materials using it weren't invalidated which implies texmap is up to date(or material would have received a change event)
+		return;
+	}
+
+	TArray<TSharedPtr<IDatasmithTextureElement>> TextureElements;
+	if (TSharedPtr<ITexmapToTextureElementConverter>* Found = TexmapConverters.Find(Texmap))
+	{
+		TSharedPtr<ITexmapToTextureElementConverter> Converter = *Found;
+		TSharedPtr<IDatasmithTextureElement> TextureElement = Converter->Convert(*this, Converter->TextureElementName);
+		if (TextureElement)
+		{
+			TextureElements.Add(TextureElement);
+			AddTextureElement(TextureElement);
+		}
+	}
+
+	UsedTextureToDatasmithElement.FindOrAdd(Texmap).Append(TextureElements);
+
+	for (TSharedPtr<IDatasmithTextureElement> TextureElement : TextureElements)
+	{
+		TextureElementToTexmap.FindOrAdd(TextureElement).Add(Texmap);
+	}
+}
+
+void FMaterialsCollectionTracker::AddTextureElement(const TSharedPtr<IDatasmithTextureElement>& TextureElement)
+{
+	if (TextureElementsAddedToScene.Contains(TextureElement))
+	{
+		return;
+	}
+
+	SceneTracker.GetDatasmithSceneRef()->AddTexture(TextureElement);
+	TextureElementsAddedToScene.Add(TextureElement);
+}
+
+void FMaterialsCollectionTracker::RemoveTextureElement(const TSharedPtr<IDatasmithTextureElement>& TextureElement)
+{
+	TextureElementsAddedToScene.Remove(TextureElement);
+	SceneTracker.RemoveTexture(TextureElement);
+}
+
+void FMaterialsCollectionTracker::AddTexmapForConversion(Texmap* Texmap, const FString& DesiredTextureElementName, const TSharedPtr<ITexmapToTextureElementConverter>& Converter)
+{
+	Converter->TextureElementName = DesiredTextureElementName;
+	TexmapConverters.Add(Texmap, Converter);
+}
+
 FMaterialTracker* FMaterialsCollectionTracker::AddMaterial(Mtl* Material)
 {
 	if (FMaterialTrackerHandle* HandlePtr = MaterialTrackers.Find(Material))
diff --git a/Engine/Source/Programs/Enterprise/Datasmith/DatasmithMaxExporter/Private/DatasmithMaxExporterWithDirectLink.cpp b/Engine/Source/Programs/Enterprise/Datasmith/DatasmithMaxExporter/Private/DatasmithMaxExporterWithDirectLink.cpp
--- a/Engine/Source/Programs/Enterprise/Datasmith/DatasmithMaxExporter/Private/DatasmithMaxExporterWithDirectLink.cpp
+++ b/Engine/Source/Programs/Enterprise/Datasmith/DatasmithMaxExporter/Private/DatasmithMaxExporterWithDirectLink.cpp
@@ -1039,6 +1039,11 @@ public:
 		return ExportedScene.GetDatasmithScene().ToSharedRef();
 	}
 
+	virtual const TCHAR* GetAssetsOutputPath() override
+	{
+		return ExportedScene.GetSceneExporter().GetAssetsOutputPath();
+	}
+
 	bool bParseXRefScenes = true;
 	bool bIncludeXRefWhileParsing = false; // note - this should not be set when bParseXRefScenes if disabled
 
@@ -1648,6 +1653,8 @@ public:
 			};
 			PROGRESS_STAGE_COUNTER(ActualMaterialToUpdate.Num());
 
+			MaterialsCollectionTracker.TexmapConverters.Reset();
+
 			for (Mtl* ActualMaterial: ActualMaterialToUpdate)
 			{
 				ProgressCounter.Next();
@@ -1656,7 +1663,6 @@ public:
 					break;	
 				}
 
-
 				MaterialsCollectionTracker.ConvertMaterial(ActualMaterial, ExportedScene.GetDatasmithScene().ToSharedRef(), ExportedScene.GetSceneExporter().GetAssetsOutputPath(), ActualTexmapsToUpdate);
 			}
 
@@ -1684,10 +1690,7 @@ public:
 				}
 				SCENE_UPDATE_STAT_INC(UpdateTextures, Total);
 
-				TArray<TSharedPtr<IDatasmithTextureElement>> TextureElements;
-				FDatasmithMaxMatExport::GetXMLTexture(ExportedScene.GetDatasmithScene().ToSharedRef(), Texture, ExportedScene.GetSceneExporter().GetAssetsOutputPath(), &TextureElements);
-
-				MaterialsCollectionTracker.UsedTextureToDatasmithElement.FindOrAdd(Texture).Append(TextureElements);
+				MaterialsCollectionTracker.UpdateTexmap(Texture);
 			}
 		}
 
@@ -1776,9 +1779,47 @@ public:
 		return FString::Printf(TEXT("Materials: %d updated, %d converted, %d skipped as already converted"), SCENE_UPDATE_STAT_GET(UpdateMaterials, Total), SCENE_UPDATE_STAT_GET(UpdateMaterials, Converted), SCENE_UPDATE_STAT_GET(UpdateMaterials, SkippedAsAlreadyConverted));
 	}
 
+	FString FormatPixelCountApproximately(int64 Value)
+	{
+		if (Value < 2*1024 )
+		{
+			return FString::Printf(TEXT("%lld pixels"), Value);
+		}
+
+		Value /= 1024; // Kilo
+
+		if (Value  < 2*1024 )
+		{
+			return FString::Printf(TEXT("%lldK pixels"), Value);
+		}
+
+		Value /= 1024; // Mega
+
+		if (Value  < 2*1024 )
+		{
+			return FString::Printf(TEXT("%lld Megapixels"), Value);
+		}
+
+		Value /= 1024; // Giga
+
+		return FString::Printf(TEXT("%lld Gigapixels"), Value);
+	}
+
 	FString FormatStatsUpdateTextures()
 	{
-		return FString::Printf(TEXT("Texmaps: %d updated"), SCENE_UPDATE_STAT_GET(UpdateTextures, Total));
+		if (SCENE_UPDATE_STAT_GET(UpdateTextures, Baked) > 0)
+		{
+			return FString::Printf(TEXT("Texmaps: %d updated, baked %d, %s baked total(%s per textmap on average)"), 
+				SCENE_UPDATE_STAT_GET(UpdateTextures, Total),
+				SCENE_UPDATE_STAT_GET(UpdateTextures, Baked),
+				*FormatPixelCountApproximately(SCENE_UPDATE_STAT_GET(UpdateTextures, BakedPixels)),
+				*FormatPixelCountApproximately(SCENE_UPDATE_STAT_GET(UpdateTextures, BakedPixels) / SCENE_UPDATE_STAT_GET(UpdateTextures, Baked)
+				));
+		}
+		else
+		{
+			return FString::Printf(TEXT("Texmaps: %d updated"), SCENE_UPDATE_STAT_GET(UpdateTextures, Total));
+		}
 	}
 
 	bool ExportAnimations(FUpdateProgress::FStage& MainStage)
diff --git a/Engine/Source/Programs/Enterprise/Datasmith/DatasmithMaxExporter/Private/DatasmithMaxScanline.cpp b/Engine/Source/Programs/Enterprise/Datasmith/DatasmithMaxExporter/Private/DatasmithMaxScanline.cpp
--- a/Engine/Source/Programs/Enterprise/Datasmith/DatasmithMaxExporter/Private/DatasmithMaxScanline.cpp
+++ b/Engine/Source/Programs/Enterprise/Datasmith/DatasmithMaxExporter/Private/DatasmithMaxScanline.cpp
@@ -454,97 +454,6 @@ FString FDatasmithMaxMatWriter::DumpAutodeskBitmap(TSharedPtr<IDatasmithComposit
 	return ActualBitmapName;
 }
 
-void FDatasmithMaxMatWriter::GetRegularTexmap(TSharedRef< IDatasmithScene > DatasmithScene, BitmapTex* InBitmapTex, TArray<TSharedPtr< IDatasmithTextureElement >>* OutTextureElements)
-{
-	if (!InBitmapTex)
-	{
-		return;
-	}
-
-	FString Path = FDatasmithMaxMatWriter::GetActualBitmapPath(InBitmapTex);
-	FString ActualBitmapName = FDatasmithMaxMatWriter::GetActualBitmapName(InBitmapTex);
-
-	for (int i = 0; i < DatasmithScene->GetTexturesCount(); i++)
-	{
-		if (DatasmithScene->GetTexture(i)->GetFile() == Path && DatasmithScene->GetTexture(i)->GetName() == ActualBitmapName)
-		{
-			return;
-		}
-	}
-
-	TSharedPtr< IDatasmithTextureElement > TextureElement = FDatasmithSceneFactory::CreateTexture(*ActualBitmapName);
-	if (gammaMgr.IsEnabled())
-	{
-		const float Gamma = FDatasmithMaxMatHelper::GetBitmapGamma(InBitmapTex);
-
-		if (FDatasmithMaxMatHelper::IsSRGB(*InBitmapTex))
-		{
-			TextureElement->SetRGBCurve(Gamma / 2.2f);
-		}
-		else
-		{
-			TextureElement->SetRGBCurve(Gamma);
-		}
-	}
-	
-	if (!FPaths::FileExists(*Path))
-	{
-		DatasmithMaxDirectLink::LogWarning(FString::Printf(TEXT("Bitmap texture '%s' has missing external file '%s'"), InBitmapTex->GetName().data(), *Path));
-	}
-
-	TextureElement->SetFile(*Path);
-	DatasmithScene->AddTexture(TextureElement);
-	if (OutTextureElements)
-	{
-		OutTextureElements->Add(TextureElement);
-	}
-}
-
-void FDatasmithMaxMatWriter::GetAutodeskTexmap(TSharedRef< IDatasmithScene > DatasmithScene, Texmap* InTexMap, TArray<TSharedPtr< IDatasmithTextureElement >>* OutTextureElements)
-{
-	if (PBBitmap* BitmapSourceFile = DatasmithMaxTexmapParser::ParseAutodeskBitmap(InTexMap).SourceFile)
-	{
-		FScopedBitMapPtr ActualBitmap(BitmapSourceFile->bi, BitmapSourceFile->bm);
-		if (!ActualBitmap.Map)
-		{
-			return;
-		}
-
-		FString Path = FDatasmithMaxMatWriter::GetActualBitmapPath(&ActualBitmap.MapInfo);
-		FString ActualBitmapName = FDatasmithMaxMatWriter::GetActualBitmapName(&ActualBitmap.MapInfo);
-
-		for (int i = 0; i < DatasmithScene->GetTexturesCount(); i++)
-		{
-			if (DatasmithScene->GetTexture(i)->GetFile() == Path && DatasmithScene->GetTexture(i)->GetName() == ActualBitmapName)
-			{
-				return;
-			}
-		}
-
-		TSharedPtr< IDatasmithTextureElement > TextureElement = FDatasmithSceneFactory::CreateTexture(*ActualBitmapName);
-		if (gammaMgr.IsEnabled())
-		{
-			const float Gamma = FDatasmithMaxMatHelper::GetBitmapGamma(&ActualBitmap.MapInfo);
-
-			if (FDatasmithMaxMatHelper::IsSRGB(*ActualBitmap.Map))
-			{
-				TextureElement->SetRGBCurve(Gamma / 2.2f);
-			}
-			else
-			{
-				TextureElement->SetRGBCurve(Gamma);
-			}
-		}
-
-		TextureElement->SetFile(*Path);
-		DatasmithScene->AddTexture(TextureElement);
-		if (OutTextureElements)
-		{
-			OutTextureElements->Add(TextureElement);
-		}
-	}
-}
-
 
 enum class EScanlineMaterialMaps
 {
@@ -1225,120 +1134,6 @@ FString FDatasmithMaxMatWriter::DumpFalloff(TSharedRef< IDatasmithScene > Datasm
 	return Result;
 }
 
-void GetBakeableMaximumSize(Texmap* InTexmap, int &Width, int &Height)
-{
-	TArray<Texmap*> Texmaps;
-	Texmaps.Add(InTexmap);
-
-	// Limit baked texture size by subbitmap with largest area(number of pixels)
-	int32 BiggestSubBitmapArea = 0;  // Biggest pixel count
-	int32 BiggestSubBitmapWidth = 0;
-	int32 BiggestSubBitmapHeight = 0;
-
-	while (!Texmaps.IsEmpty())
-	{
-		Texmap* Texmap = Texmaps.Pop(false);
-
-		if (!Texmap)
-		{
-			continue;;
-		}
-
-		if (Texmap->ClassID() == RBITMAPCLASS)
-		{
-			BitmapTex* BitmapTexture = (BitmapTex*)Texmap;
-			Bitmap* ActualBitmap = ((BitmapTex*)Texmap)->GetBitmap(GetCOREInterface()->GetTime());
-			if (ActualBitmap != NULL)
-			{
-				int32 BitmapWidth = ActualBitmap->Width();
-				int32 BitmapHeight = ActualBitmap->Height();
-
-				int32 Area = BitmapWidth*BitmapHeight;
-				if (Area > BiggestSubBitmapArea)
-				{
-					BiggestSubBitmapArea = Area;
-					BiggestSubBitmapWidth = BitmapWidth;
-					BiggestSubBitmapHeight = BitmapHeight;
-				}
-
-			}
-		}
-
-		for (int SubTexmap = 0; SubTexmap < Texmap->NumSubTexmaps(); SubTexmap++)
-		{
-			Texmaps.Add(Texmap->GetSubTexmap(SubTexmap));
-		}
-	}
-
-	if (BiggestSubBitmapArea > 0)
-	{
-		int32 AreaLimit = Width*Height;
-
-		Width = BiggestSubBitmapWidth;
-		Height = BiggestSubBitmapHeight;
-		int32 Area = BiggestSubBitmapArea;
-		while (Area > AreaLimit)
-		{
-			Width = Width / 2;
-			Height = Height / 2;
-			Area = Width * Height;
-		}
-	}
-}
-
-TSharedPtr< IDatasmithTextureElement > FDatasmithMaxMatWriter::AddBakeable(TSharedRef< IDatasmithScene > DatasmithScene, Texmap* InTexmap, const TCHAR* AssetsPath)
-{
-	if ( !InTexmap )
-	{
-		return TSharedPtr< IDatasmithTextureElement >();
-	}
-
-	MSTR ClassName;
-	InTexmap->GetClassName(ClassName);
-	FString FileName = FString(InTexmap->GetName().data()) + FString(ClassName.data()) + FString::FromInt(InTexmap->GetHandleByAnim(InTexmap));
-	FileName = FDatasmithUtils::SanitizeFileName(FileName);
-
-	FString Path = FPaths::Combine(AssetsPath, FileName + TextureBakeFormat);
-
-	for (int i = 0; i < DatasmithScene->GetTexturesCount(); i++)
-	{
-		if (DatasmithScene->GetTexture(i)->GetFile() == Path)
-		{
-			return DatasmithScene->GetTexture(i);
-		}
-	}
-
-	int BakeWidth = FDatasmithExportOptions::MaxTextureSize;
-	int BakeHeight = FDatasmithExportOptions::MaxTextureSize;
-	GetBakeableMaximumSize(InTexmap, BakeWidth, BakeHeight);
-
-	BitmapInfo BitmapInformation;
-	BitmapInformation.SetType(BMM_TRUE_32);
-
-	BitmapInformation.SetWidth(BakeWidth);
-	BitmapInformation.SetHeight(BakeHeight);
-	BitmapInformation.SetGamma(2.2f);
-
-	BitmapInformation.SetName(*Path);
-
-	Bitmap* NewBitmap = TheManager->Create(&BitmapInformation);
-	InTexmap->RenderBitmap(GetCOREInterface()->GetTime(), NewBitmap, 1.0f, 1);
-	
-	NewBitmap->OpenOutput(&BitmapInformation);
-	NewBitmap->Write(&BitmapInformation);
-	NewBitmap->Close(&BitmapInformation);
-	NewBitmap->DeleteThis();
-
-	FString BaseName = FPaths::GetBaseFilename(Path);
-	FString Base = BaseName + TextureSuffix;
-	TSharedPtr< IDatasmithTextureElement > TextureElement = FDatasmithSceneFactory::CreateTexture(*Base);
-	TextureElement->SetRGBCurve(1.0f);
-	TextureElement->SetFile(*Path);
-	DatasmithScene->AddTexture(TextureElement);
-
-	return TextureElement;
-}
-
 FString FDatasmithMaxMatWriter::DumpBakeable(TSharedPtr<IDatasmithCompositeTexture>& CompTex, Texmap* InTexmap, const TCHAR* Prefix, bool bForceInvert, bool bIsGrayscale)
 {
 	MSTR ClassName;
diff --git a/Engine/Source/Programs/Enterprise/Datasmith/DatasmithMaxExporter/Private/DatasmithMaxThea.cpp b/Engine/Source/Programs/Enterprise/Datasmith/DatasmithMaxExporter/Private/DatasmithMaxThea.cpp
--- a/Engine/Source/Programs/Enterprise/Datasmith/DatasmithMaxExporter/Private/DatasmithMaxThea.cpp
+++ b/Engine/Source/Programs/Enterprise/Datasmith/DatasmithMaxExporter/Private/DatasmithMaxThea.cpp
@@ -11,12 +11,6 @@ FString FDatasmithMaxMatWriter::DumpBitmapThea(TSharedPtr<IDatasmithCompositeTex
 	return DumpBitmap(CompTex, Texture, Prefix, bForceInvert, bIsGrayscale);
 }
 
-void FDatasmithMaxMatWriter::GetTheaTexmap(TSharedRef< IDatasmithScene > DatasmithScene, BitmapTex* InBitmapTex, TArray<TSharedPtr< IDatasmithTextureElement >>* OutTextureElements)
-{
-	BitmapTex* Texture = (BitmapTex*)InBitmapTex->GetReference(0);
-	GetRegularTexmap(DatasmithScene, Texture, OutTextureElements);
-}
-
 void FDatasmithMaxMatWriter::ExportTheaSubmaterial(TSharedRef< IDatasmithScene > DatasmithScene, TSharedPtr< IDatasmithShaderElement >& MaterialShader, Mtl* Material, EDSMaterialType MaterialType)
 {
 	int NumParamBlocks = Material->NumParamBlocks();
diff --git a/Engine/Source/Programs/Enterprise/Datasmith/DatasmithMaxExporter/Private/DatasmithMaxVRay.cpp b/Engine/Source/Programs/Enterprise/Datasmith/DatasmithMaxExporter/Private/DatasmithMaxVRay.cpp
--- a/Engine/Source/Programs/Enterprise/Datasmith/DatasmithMaxExporter/Private/DatasmithMaxVRay.cpp
+++ b/Engine/Source/Programs/Enterprise/Datasmith/DatasmithMaxExporter/Private/DatasmithMaxVRay.cpp
@@ -123,63 +123,6 @@ FString FDatasmithMaxMatWriter::DumpVrayHdri(TSharedPtr<IDatasmithCompositeTextu
 	return *BitmapName;
 }
 
-void FDatasmithMaxMatWriter::GetVrayHdri(TSharedRef< IDatasmithScene > DatasmithScene, BitmapTex* InBitmapTex, TArray<TSharedPtr< IDatasmithTextureElement >>* OutTextureElements)
-{
-	FString Path = TEXT("");
-
-	int NumParamBlocks = InBitmapTex->NumParamBlocks();
-
-	for (int j = 0; j < NumParamBlocks; j++)
-	{
-		IParamBlock2* ParamBlock2 = InBitmapTex->GetParamBlockByID((short)j);
-		// The the descriptor to 'decode'
-		ParamBlockDesc2* ParamBlockDesc = ParamBlock2->GetDesc();
-		// Loop through all the defined parameters therein
-		for (int i = 0; i < ParamBlockDesc->count; i++)
-		{
-			const ParamDef& ParamDefinition = ParamBlockDesc->paramdefs[i];
-
-			if (FCString::Stricmp(ParamDefinition.int_name, TEXT("HDRIMapName")) == 0)
-			{
-				Path = FDatasmithMaxSceneExporter::GetActualPath(ParamBlock2->GetStr(ParamDefinition.ID, GetCOREInterface()->GetTime()));
-			}
-		}
-		ParamBlock2->ReleaseDesc();
-	}
-
-
-
-	if (Path.IsEmpty())
-	{
-		return;
-	}
-
-	float Gamma =  FDatasmithMaxMatHelper::GetVrayHdriGamma(InBitmapTex);
-
-	FString BaseName = FPaths::GetBaseFilename(Path);
-	FString Base = BaseName + FString("_") + FString::SanitizeFloat(Gamma).Replace(TEXT("."), TEXT("_")) + TextureSuffix;
-
-	for (int i = 0; i < DatasmithScene->GetTexturesCount(); i++)
-	{
-		if (DatasmithScene->GetTexture(i)->GetFile() == Path && DatasmithScene->GetTexture(i)->GetName() == Base)
-		{
-			return;
-		}
-	}
-
-	TSharedPtr< IDatasmithTextureElement > TextureElement = FDatasmithSceneFactory::CreateTexture(*Base);
-	if (gammaMgr.IsEnabled())
-	{
-		TextureElement->SetRGBCurve(Gamma / 2.2f);
-	}
-	TextureElement->SetFile(*Path);
-	DatasmithScene->AddTexture(TextureElement);
-	if (OutTextureElements)
-	{
-		OutTextureElements->Add(TextureElement);
-	}
-}
-
 void FDatasmithMaxMatWriter::ExportVRayMaterial(TSharedRef< IDatasmithScene > DatasmithScene, TSharedPtr< IDatasmithMaterialElement >& MaterialElement, Mtl* Material)
 {
 	TSharedPtr< IDatasmithShaderElement > MaterialShader = FDatasmithSceneFactory::CreateShader( (TCHAR*)Material->GetName().data() );
diff --git a/Engine/Source/Programs/Enterprise/Datasmith/DatasmithMaxExporter/Private/DatasmithMaxWriter.cpp b/Engine/Source/Programs/Enterprise/Datasmith/DatasmithMaxExporter/Private/DatasmithMaxWriter.cpp
--- a/Engine/Source/Programs/Enterprise/Datasmith/DatasmithMaxExporter/Private/DatasmithMaxWriter.cpp
+++ b/Engine/Source/Programs/Enterprise/Datasmith/DatasmithMaxExporter/Private/DatasmithMaxWriter.cpp
@@ -444,60 +444,6 @@ bool FDatasmithMaxMatHelper::IsSRGB(Bitmap& InBitmap)
 	return true;
 }
 
-// todo: need wa wya to return created texture element(if there was one), maybe use 'last' in DatasmithScene(check that count changed)
-//   verify that only one texture created!
-void FDatasmithMaxMatExport::GetXMLTexture(TSharedRef< IDatasmithScene > DatasmithScene, Texmap* InTexMap, const TCHAR* AssetsPath, TArray<TSharedPtr< IDatasmithTextureElement >>* OutTextureElements)
-{
-	switch (FDatasmithMaxMatHelper::GetTextureClass(InTexMap))
-	{
-	case EDSBitmapType::NotSupported:
-		break;
-	case EDSBitmapType::RegularBitmap:
-		FDatasmithMaxMatWriter::GetRegularTexmap(DatasmithScene, (BitmapTex*)InTexMap, OutTextureElements);
-		break;
-	case EDSBitmapType::AutodeskBitmap:
-		FDatasmithMaxMatWriter::GetAutodeskTexmap(DatasmithScene, InTexMap, OutTextureElements);
-		break; 
-	case EDSBitmapType::TheaBitmap:
-		FDatasmithMaxMatWriter::GetTheaTexmap(DatasmithScene, (BitmapTex*)InTexMap, OutTextureElements);
-		break;
-	case EDSBitmapType::CoronaBitmap:
-		FDatasmithMaxMatWriter::GetCoronaTexmap(DatasmithScene, (BitmapTex*)InTexMap, OutTextureElements);
-		break;
-	case EDSBitmapType::VRayHRDI:
-		FDatasmithMaxMatWriter::GetVrayHdri(DatasmithScene, (BitmapTex*)InTexMap, OutTextureElements);
-		break;
-	case EDSBitmapType::ColorCorrector:
-		GetXMLTexture(DatasmithScene, InTexMap->GetSubTexmap(0), AssetsPath, OutTextureElements);
-		break;
-	case EDSBitmapType::BakeableMap:
-		{
-			TSharedPtr<IDatasmithTextureElement> TextureElement = FDatasmithMaxMatWriter::AddBakeable(DatasmithScene, (BitmapTex*)InTexMap, AssetsPath);
-			if (OutTextureElements)
-			{
-				OutTextureElements->Add(TextureElement);
-			}
-		}
-		break;
-	case EDSBitmapType::NormalMap:
-	case EDSBitmapType::FallOff:
-	case EDSBitmapType::Gradient:
-	case EDSBitmapType::GradientRamp:
-	case EDSBitmapType::Checker:
-	case EDSBitmapType::Cellular:
-	case EDSBitmapType::Mix:
-	case EDSBitmapType::Noise:
-	case EDSBitmapType::CompositeTex:
-		for (int i = 0; i < InTexMap->NumSubTexmaps(); i++)
-		{
-			GetXMLTexture(DatasmithScene, InTexMap->GetSubTexmap(i), AssetsPath, OutTextureElements);
-		}
-		break;
-	default:
-		break;
-	}
-}
-
 bool FDatasmithMaxMatExport::UseFirstSubMapOnly(EDSMaterialType MaterialType, Mtl* Material)
 {
 	switch (MaterialType)
diff --git a/Engine/Source/Programs/Enterprise/Datasmith/DatasmithMaxExporter/Private/MaxMaterialsToUEPbr/DatasmithMaxCoronaTexmapToUEPbr.cpp b/Engine/Source/Programs/Enterprise/Datasmith/DatasmithMaxExporter/Private/MaxMaterialsToUEPbr/DatasmithMaxCoronaTexmapToUEPbr.cpp
--- a/Engine/Source/Programs/Enterprise/Datasmith/DatasmithMaxExporter/Private/MaxMaterialsToUEPbr/DatasmithMaxCoronaTexmapToUEPbr.cpp
+++ b/Engine/Source/Programs/Enterprise/Datasmith/DatasmithMaxExporter/Private/MaxMaterialsToUEPbr/DatasmithMaxCoronaTexmapToUEPbr.cpp
@@ -10,6 +10,7 @@
 #include "MaxMaterialsToUEPbr/DatasmithMaxTexmapToUEPbr.h"
 
 #include "DatasmithMaterialsUtils.h"
+#include "DatasmithSceneFactory.h"
 
 #include "Misc/Paths.h"
 
@@ -216,6 +217,59 @@ bool FDatasmithMaxCoronalBitmapToUEPbr::IsSupported( const FDatasmithMaxMaterial
 	return ( InTexmap && InTexmap->ClassID() == CORONABITMAPCLASS );
 }
 
+float GetCoronaTexmapGamma(BitmapTex* InBitmapTex);
+
+class FCoronaBitmapToTextureElementConverter: public DatasmithMaxDirectLink::ITexmapToTextureElementConverter
+{
+public:
+	virtual TSharedPtr<IDatasmithTextureElement> Convert(DatasmithMaxDirectLink::FMaterialsCollectionTracker& MaterialsTracker, const FString& ActualBitmapName) override
+	{
+		FString Path = TEXT("");
+
+		int NumParamBlocks = Tex->NumParamBlocks();
+		
+		for (int j = 0; j < NumParamBlocks; j++)
+		{
+			IParamBlock2* ParamBlock2 = Tex->GetParamBlockByID((short)j);
+			// The the descriptor to 'decode'
+			ParamBlockDesc2* ParamBlockDesc = ParamBlock2->GetDesc();
+			// Loop through all the defined parameters therein
+			for (int i = 0; i < ParamBlockDesc->count; i++)
+			{
+				const ParamDef& ParamDefinition = ParamBlockDesc->paramdefs[i];
+
+				if (FCString::Stricmp(ParamDefinition.int_name, TEXT("filename")) == 0)
+				{
+					Path = FDatasmithMaxSceneExporter::GetActualPath(ParamBlock2->GetStr(ParamDefinition.ID, GetCOREInterface()->GetTime()));
+					continue;
+				}
+			}
+			ParamBlock2->ReleaseDesc();
+		}
+
+		if (Path.IsEmpty())
+		{
+			return {};
+		}
+
+		float Gamma = GetCoronaTexmapGamma(Tex);
+
+		TSharedPtr< IDatasmithTextureElement > TextureElement = FDatasmithSceneFactory::CreateTexture(*ActualBitmapName);
+		if (gammaMgr.IsEnabled())
+		{
+			TextureElement->SetRGBCurve(Gamma / 2.2f);
+		}
+		TextureElement->SetFile(*Path);
+
+		return TextureElement;
+	}
+	BitmapTex* Tex;
+
+	bool bIsSRGB;
+};
+
+
+
 IDatasmithMaterialExpression* FDatasmithMaxCoronalBitmapToUEPbr::Convert( FDatasmithMaxMaterialsToUEPbr* MaxMaterialToUEPbr, Texmap* InTexmap )
 {
 	DatasmithMaxTexmapParser::FCoronaBitmapParameters CoronaBitmapParameters = DatasmithMaxTexmapParser::ParseCoronaBitmap( InTexmap );
@@ -225,6 +279,13 @@ IDatasmithMaterialExpression* FDatasmithMaxCoronalBitmapToUEPbr::Convert( FDatas
 
 	const bool bIsSRGB = false;
 	const bool bUseAlphaAsMono = false;
+
+	TSharedRef<FCoronaBitmapToTextureElementConverter> Converter = MakeShared<FCoronaBitmapToTextureElementConverter>();
+	Converter->Tex = (BitmapTex*)InTexmap;
+	Converter->bIsSRGB = bIsSRGB;
+
+	MaxMaterialToUEPbr->AddTexmap(InTexmap, ActualBitmapName, Converter);
+
 	return FDatasmithMaxTexmapToUEPbrUtils::ConvertBitMap(MaxMaterialToUEPbr, InTexmap, ActualBitmapName, bUseAlphaAsMono, bIsSRGB);
 }
 
diff --git a/Engine/Source/Programs/Enterprise/Datasmith/DatasmithMaxExporter/Private/MaxMaterialsToUEPbr/DatasmithMaxMaterialsToUEPbr.cpp b/Engine/Source/Programs/Enterprise/Datasmith/DatasmithMaxExporter/Private/MaxMaterialsToUEPbr/DatasmithMaxMaterialsToUEPbr.cpp
--- a/Engine/Source/Programs/Enterprise/Datasmith/DatasmithMaxExporter/Private/MaxMaterialsToUEPbr/DatasmithMaxMaterialsToUEPbr.cpp
+++ b/Engine/Source/Programs/Enterprise/Datasmith/DatasmithMaxExporter/Private/MaxMaterialsToUEPbr/DatasmithMaxMaterialsToUEPbr.cpp
@@ -196,7 +196,7 @@ void FDatasmithMaxMaterialsToUEPbr::AddConvertedMap(const DatasmithMaxTexmapPars
 {
 	if (Context)
 	{
-		Context->TexmapsConverted.FindOrAdd(MapParameter.Map);
+		Context->TexmapsConverted.Add(MapParameter.Map);
 	}
 }
 
diff --git a/Engine/Source/Programs/Enterprise/Datasmith/DatasmithMaxExporter/Private/MaxMaterialsToUEPbr/DatasmithMaxMaterialsToUEPbr.h b/Engine/Source/Programs/Enterprise/Datasmith/DatasmithMaxExporter/Private/MaxMaterialsToUEPbr/DatasmithMaxMaterialsToUEPbr.h
--- a/Engine/Source/Programs/Enterprise/Datasmith/DatasmithMaxExporter/Private/MaxMaterialsToUEPbr/DatasmithMaxMaterialsToUEPbr.h
+++ b/Engine/Source/Programs/Enterprise/Datasmith/DatasmithMaxExporter/Private/MaxMaterialsToUEPbr/DatasmithMaxMaterialsToUEPbr.h
@@ -31,12 +31,12 @@ namespace DatasmithMaxDirectLink
 
 	struct FMaterialConversionContext
 	{
-		// Used to record texmaps that were asked for conversion whil ematerial conversion is performed
-		// recordking also datasmith elements created for those maps.
-		// Baked texmaps currently create datasmith elements during material convertion
-		TMap<Texmap*, TSet<TSharedPtr<IDatasmithTextureElement>>>& TexmapsConverted;
+		// Records texmaps that were asked for conversion while material conversion is performed
+		TSet<Texmap*>& TexmapsConverted;
 		FMaterialsCollectionTracker& MaterialsCollectionTracker;
 	};
+
+	class ITexmapToTextureElementConverter;
 }
 
 
@@ -65,6 +65,15 @@ public:
 	 */
 	IDatasmithMaterialExpression* ConvertTexmap( const DatasmithMaxTexmapParser::FMapParameter& MapParameter );
 
+	// Texmap which needs to have a datasmith texture element created for
+	void AddTexmap(Texmap* Texmap, const FString& DesiredTextureElementName, const TSharedPtr<DatasmithMaxDirectLink::ITexmapToTextureElementConverter>& Converter)
+	{
+		if (Context)
+		{
+			Context->MaterialsCollectionTracker.AddTexmapForConversion(Texmap, DesiredTextureElementName, Converter);
+		}
+	}
+
 	struct FConvertState
 	{
 		TSharedPtr< IDatasmithScene > DatasmithScene;
@@ -81,14 +90,6 @@ public:
 
 	void AddConvertedMap(const DatasmithMaxTexmapParser::FMapParameter& MapParameter);
 
-	void AddConvertedMapDatasmithElement(Texmap* InTexmap, TSharedPtr<IDatasmithTextureElement> TextureElement)
-	{
-		if (Context)
-		{
-			Context->TexmapsConverted.FindOrAdd(InTexmap).Add(TextureElement);
-		}
-	}
-
 	const TCHAR* GetMaterialName(Mtl* Material) const
 	{
 		return Context->MaterialsCollectionTracker.GetMaterialName(Material);
diff --git a/Engine/Source/Programs/Enterprise/Datasmith/DatasmithMaxExporter/Private/MaxMaterialsToUEPbr/DatasmithMaxTexmapToUEPbr.cpp b/Engine/Source/Programs/Enterprise/Datasmith/DatasmithMaxExporter/Private/MaxMaterialsToUEPbr/DatasmithMaxTexmapToUEPbr.cpp
--- a/Engine/Source/Programs/Enterprise/Datasmith/DatasmithMaxExporter/Private/MaxMaterialsToUEPbr/DatasmithMaxTexmapToUEPbr.cpp
+++ b/Engine/Source/Programs/Enterprise/Datasmith/DatasmithMaxExporter/Private/MaxMaterialsToUEPbr/DatasmithMaxTexmapToUEPbr.cpp
@@ -12,6 +12,7 @@
 #include "DatasmithMaxHelper.h"
 
 #include "Misc/Paths.h"
+#include "DatasmithExportOptions.h"
 
 #include "Windows/AllowWindowsPlatformTypes.h"
 MAX_INCLUDES_START
@@ -846,6 +847,48 @@ bool FDatasmithMaxBitmapToUEPbr::IsSupported( const FDatasmithMaxMaterialsToUEPb
 	return false;
 }
 
+class FBitmapToTextureElementConverter: public DatasmithMaxDirectLink::ITexmapToTextureElementConverter
+{
+public:
+	virtual TSharedPtr<IDatasmithTextureElement> Convert(DatasmithMaxDirectLink::FMaterialsCollectionTracker& MaterialsTracker, const FString& ActualBitmapName) override
+	{
+		FString Path = FDatasmithMaxMatWriter::GetActualBitmapPath(Tex);
+
+		if (TSharedPtr<IDatasmithTextureElement>* Found =  MaterialsTracker.BitmapTextureElements.Find(Path))
+		{
+			return *Found;
+		}
+
+		TSharedPtr< IDatasmithTextureElement > TextureElement = FDatasmithSceneFactory::CreateTexture(*ActualBitmapName);
+		MaterialsTracker.BitmapTextureElements.Add(Path, TextureElement);
+
+		if (gammaMgr.IsEnabled())
+		{
+			const float Gamma = FDatasmithMaxMatHelper::GetBitmapGamma(Tex);
+
+			if (FDatasmithMaxMatHelper::IsSRGB(*Tex))
+			{
+				TextureElement->SetRGBCurve(Gamma / 2.2f);
+			}
+			else
+			{
+				TextureElement->SetRGBCurve(Gamma);
+			}
+		}
+		
+		if (!FPaths::FileExists(*Path))
+		{
+			DatasmithMaxDirectLink::LogWarning(FString::Printf(TEXT("Bitmap texture '%s' has missing external file '%s'"), Tex->GetName().data(), *Path));
+		}
+
+		TextureElement->SetFile(*Path);
+		return TextureElement;
+	}
+	BitmapTex* Tex;
+
+	bool bIsSRGB;
+};
+
 IDatasmithMaterialExpression* FDatasmithMaxBitmapToUEPbr::Convert( FDatasmithMaxMaterialsToUEPbr* MaxMaterialToUEPbr, Texmap* InTexmap )
 {
 	BitmapTex* InBitmapTex = (BitmapTex*)InTexmap;
@@ -854,6 +897,12 @@ IDatasmithMaterialExpression* FDatasmithMaxBitmapToUEPbr::Convert( FDatasmithMax
 	bool bUseAlphaAsMono = InBitmapTex->GetAlphaAsMono(true);
 	bool bIsSRGB = FDatasmithMaxMatHelper::IsSRGB(*InBitmapTex);
 
+	TSharedRef<FBitmapToTextureElementConverter> Converter = MakeShared<FBitmapToTextureElementConverter>();
+	Converter->Tex = InBitmapTex;
+	Converter->bIsSRGB = bIsSRGB;
+
+	MaxMaterialToUEPbr->AddTexmap(InTexmap, ActualBitmapName, Converter);
+
 	return FDatasmithMaxTexmapToUEPbrUtils::ConvertBitMap(MaxMaterialToUEPbr, InTexmap, ActualBitmapName, bUseAlphaAsMono, bIsSRGB);
 }
 
@@ -873,16 +922,75 @@ bool FDatasmithMaxAutodeskBitmapToUEPbr::IsSupported(const FDatasmithMaxMaterial
 	return false;
 }
 
+
+class FAutodeskBitmapToTextureElementConverter: public DatasmithMaxDirectLink::ITexmapToTextureElementConverter
+{
+public:
+	virtual TSharedPtr<IDatasmithTextureElement> Convert(DatasmithMaxDirectLink::FMaterialsCollectionTracker& MaterialsTracker, const FString& ActualBitmapName) override
+	{
+		if (PBBitmap* BitmapSourceFile = DatasmithMaxTexmapParser::ParseAutodeskBitmap(Tex).SourceFile)
+		{
+			FScopedBitMapPtr ActualBitmap(BitmapSourceFile->bi, BitmapSourceFile->bm);
+			if (!ActualBitmap.Map)
+			{
+				return {};
+			}
+
+			TSharedRef<IDatasmithScene> DatasmithScene = MaterialsTracker.SceneTracker.GetDatasmithSceneRef();
+			for (int i = 0; i < DatasmithScene->GetTexturesCount(); i++)
+			{
+				if (DatasmithScene->GetTexture(i)->GetFile() == Path && DatasmithScene->GetTexture(i)->GetName() == ActualBitmapName)
+				{
+					return {};
+				}
+			}
+
+			TSharedPtr< IDatasmithTextureElement > TextureElement = FDatasmithSceneFactory::CreateTexture(*ActualBitmapName);
+			if (gammaMgr.IsEnabled())
+			{
+				const float Gamma = FDatasmithMaxMatHelper::GetBitmapGamma(&ActualBitmap.MapInfo);
+
+				if (FDatasmithMaxMatHelper::IsSRGB(*ActualBitmap.Map))
+				{
+					TextureElement->SetRGBCurve(Gamma / 2.2f);
+				}
+				else
+				{
+					TextureElement->SetRGBCurve(Gamma);
+				}
+			}
+
+			TextureElement->SetFile(*Path);
+			return TextureElement;
+		}
+		return {};
+	}
+	Texmap* Tex;
+
+	bool bIsSRGB;
+	FString Path;
+};
+
+
 IDatasmithMaterialExpression* FDatasmithMaxAutodeskBitmapToUEPbr::Convert(FDatasmithMaxMaterialsToUEPbr* MaxMaterialToUEPbr, Texmap* InTexmap)
 {
 	DatasmithMaxTexmapParser::FAutodeskBitmapParameters AutodeskBitmapParameters = DatasmithMaxTexmapParser::ParseAutodeskBitmap(InTexmap);
 	if (AutodeskBitmapParameters.SourceFile)
 	{
 		FScopedBitMapPtr ActualBitmap(AutodeskBitmapParameters.SourceFile->bi, AutodeskBitmapParameters.SourceFile->bm);
+		FString Path = FDatasmithMaxMatWriter::GetActualBitmapPath(&ActualBitmap.MapInfo);
 		FString ActualBitmapName = FDatasmithMaxMatWriter::GetActualBitmapName(&ActualBitmap.MapInfo);
 		bool bUseAlphaAsMono = (ActualBitmap.Map->HasAlpha() != 0);
 		bool bIsSRGB = FDatasmithMaxMatHelper::IsSRGB(*ActualBitmap.Map);
 
+
+		TSharedRef<FAutodeskBitmapToTextureElementConverter> Converter = MakeShared<FAutodeskBitmapToTextureElementConverter>();
+		Converter->Tex = InTexmap;
+		Converter->bIsSRGB = bIsSRGB;
+		Converter->Path = Path;
+
+		MaxMaterialToUEPbr->AddTexmap(InTexmap, ActualBitmapName, Converter);
+
 		return FDatasmithMaxTexmapToUEPbrUtils::ConvertBitMap(MaxMaterialToUEPbr, InTexmap, ActualBitmapName, bUseAlphaAsMono, bIsSRGB);
 	}
 	return nullptr;
@@ -1694,30 +1802,159 @@ bool FDatasmithMaxBakeableToUEPbr::IsSupported( const FDatasmithMaxMaterialsToUE
 	return false;
 }
 
-IDatasmithMaterialExpression* FDatasmithMaxBakeableToUEPbr::Convert( FDatasmithMaxMaterialsToUEPbr* MaxMaterialToUEPbr, Texmap* InTexmap )
+class FBakeableToTextureElementConverter: public DatasmithMaxDirectLink::ITexmapToTextureElementConverter
 {
-	IDatasmithMaterialExpression* BakedExpression = nullptr;
+public:
+	virtual TSharedPtr<IDatasmithTextureElement> Convert(DatasmithMaxDirectLink::FMaterialsCollectionTracker& MaterialsTracker, const FString& ActualBitmapName) override
+	{
+		if ( !Tex )
+		{
+			return {};
+		}
+
+		FString TexmapName = FString(Tex->GetName().data());
+		FString Path = FPaths::Combine(MaterialsTracker.SceneTracker.GetAssetsOutputPath(), FileName + FDatasmithMaxMatWriter::TextureBakeFormat);
+		
+		int BakeWidth = FDatasmithExportOptions::MaxTextureSize;
+		int BakeHeight = FDatasmithExportOptions::MaxTextureSize;
+		GetBakeableMaximumSize(Tex, BakeWidth, BakeHeight);
+
+		DatasmithMaxDirectLink::FSceneUpdateStats& Stats = MaterialsTracker.Stats;
+		SCENE_UPDATE_STAT_INC(UpdateTextures, Baked);
+
+		SCENE_UPDATE_STAT_SET(UpdateTextures, BakedPixels, SCENE_UPDATE_STAT_GET(UpdateTextures, BakedPixels) + BakeWidth*BakeHeight);
 
-	TSharedPtr< IDatasmithTextureElement > BakedTextureElement = FDatasmithMaxMatWriter::AddBakeable( MaxMaterialToUEPbr->ConvertState.DatasmithScene.ToSharedRef(),
-		InTexmap, *MaxMaterialToUEPbr->ConvertState.AssetsPath );
+		DatasmithMaxDirectLink::LogInfo(FString::Printf(TEXT("Rendering Texmap '%s' to %dx%d image, path '%s'..."), *TexmapName, BakeWidth, BakeHeight, *Path));
+		FDateTime TimeStart = FDateTime::UtcNow();
 
-	if ( BakedTextureElement )
+		BitmapInfo BitmapInformation;
+		BitmapInformation.SetType(BMM_TRUE_32);
+
+		BitmapInformation.SetWidth(BakeWidth);
+		BitmapInformation.SetHeight(BakeHeight);
+		BitmapInformation.SetGamma(2.2f);
+
+		BitmapInformation.SetName(*Path);
+
+		Bitmap* NewBitmap = TheManager->Create(&BitmapInformation);
+		Tex->RenderBitmap(GetCOREInterface()->GetTime(), NewBitmap, 1.0f, 1);
+		
+		NewBitmap->OpenOutput(&BitmapInformation);
+		NewBitmap->Write(&BitmapInformation);
+		NewBitmap->Close(&BitmapInformation);
+		NewBitmap->DeleteThis();
+
+		TSharedPtr< IDatasmithTextureElement > TextureElement = FDatasmithSceneFactory::CreateTexture(*ActualBitmapName);
+		TextureElement->SetRGBCurve(1.0f);
+		TextureElement->SetFile(*Path);
+
+		FDateTime TimeFinish = FDateTime::UtcNow();
+		
+		DatasmithMaxDirectLink::LogInfo(FString::Printf(TEXT("done in %s"), *(TimeFinish-TimeStart).ToString()));
+
+		return TextureElement;
+	}
+
+	static void GetBakeableMaximumSize(Texmap* InTexmap, int &Width, int &Height)
 	{
-		MaxMaterialToUEPbr->AddConvertedMapDatasmithElement(InTexmap, BakedTextureElement);
+		TArray<Texmap*> Texmaps;
+		Texmaps.Add(InTexmap);
 
-		IDatasmithMaterialExpression* MaterialExpression = MaxMaterialToUEPbr->ConvertState.MaterialElement->AddMaterialExpression( EDatasmithMaterialExpressionType::Texture );
+		// Limit baked texture size by subbitmap with largest area(number of pixels)
+		int32 BiggestSubBitmapArea = 0;  // Biggest pixel count
+		int32 BiggestSubBitmapWidth = 0;
+		int32 BiggestSubBitmapHeight = 0;
 
-		if ( MaterialExpression )
+		while (!Texmaps.IsEmpty())
 		{
-			IDatasmithMaterialExpressionTexture* MaterialExpressionTexture = static_cast< IDatasmithMaterialExpressionTexture* >( MaterialExpression );
-			MaterialExpressionTexture->SetTexturePathName( BakedTextureElement->GetName() );
+			Texmap* Texmap = Texmaps.Pop(false);
+
+			if (!Texmap)
+			{
+				continue;;
+			}
+
+			if (Texmap->ClassID() == RBITMAPCLASS)
+			{
+				BitmapTex* BitmapTexture = (BitmapTex*)Texmap;
+				Bitmap* ActualBitmap = ((BitmapTex*)Texmap)->GetBitmap(GetCOREInterface()->GetTime());
+				if (ActualBitmap != NULL)
+				{
+					int32 BitmapWidth = ActualBitmap->Width();
+					int32 BitmapHeight = ActualBitmap->Height();
+
+					int32 Area = BitmapWidth*BitmapHeight;
+					if (Area > BiggestSubBitmapArea)
+					{
+						BiggestSubBitmapArea = Area;
+						BiggestSubBitmapWidth = BitmapWidth;
+						BiggestSubBitmapHeight = BitmapHeight;
+					}
+
+				}
+			}
 
-			FDatasmithMaxTexmapToUEPbrUtils::SetupTextureCoordinates( MaxMaterialToUEPbr, MaterialExpressionTexture->GetInputCoordinate(), InTexmap );
+			for (int SubTexmap = 0; SubTexmap < Texmap->NumSubTexmaps(); SubTexmap++)
+			{
+				Texmaps.Add(Texmap->GetSubTexmap(SubTexmap));
+			}
+		}
+
+		if (BiggestSubBitmapArea > 0)
+		{
+			int32 AreaLimit = Width*Height;
 
-			BakedExpression = MaterialExpressionTexture;
+			Width = BiggestSubBitmapWidth;
+			Height = BiggestSubBitmapHeight;
+			int32 Area = BiggestSubBitmapArea;
+			while (Area > AreaLimit)
+			{
+				Width = Width / 2;
+				Height = Height / 2;
+				Area = Width * Height;
+			}
 		}
 	}
 
+	Texmap* Tex;
+	FString FileName;
+};
+
+IDatasmithMaterialExpression* FDatasmithMaxBakeableToUEPbr::Convert( FDatasmithMaxMaterialsToUEPbr* MaxMaterialToUEPbr, Texmap* InTexmap )
+{
+	if ( !InTexmap )
+	{
+		return false;
+	}
+
+	MSTR ClassName;
+	InTexmap->GetClassName(ClassName);
+	FString TexmapName = FString(InTexmap->GetName().data());
+	FString FileName = TexmapName + FString(ClassName.data()) + FString::FromInt(InTexmap->GetHandleByAnim(InTexmap));
+	FileName = FDatasmithUtils::SanitizeFileName(FileName);
+
+	FString Name = FDatasmithUtils::SanitizeObjectName(FileName + FDatasmithMaxMatWriter::TextureSuffix);
+
+	TSharedRef<FBakeableToTextureElementConverter> Converter = MakeShared<FBakeableToTextureElementConverter>();
+	Converter->Tex = InTexmap;
+	Converter->FileName = FileName;
+
+	MaxMaterialToUEPbr->AddTexmap(InTexmap, Name, Converter);
+
+	IDatasmithMaterialExpression* BakedExpression = nullptr;
+
+	IDatasmithMaterialExpression* MaterialExpression = MaxMaterialToUEPbr->ConvertState.MaterialElement->AddMaterialExpression( EDatasmithMaterialExpressionType::Texture );
+
+	if ( MaterialExpression )
+	{
+		IDatasmithMaterialExpressionTexture* MaterialExpressionTexture = static_cast< IDatasmithMaterialExpressionTexture* >( MaterialExpression );
+		MaterialExpressionTexture->SetTexturePathName( *Name );
+
+		FDatasmithMaxTexmapToUEPbrUtils::SetupTextureCoordinates( MaxMaterialToUEPbr, MaterialExpressionTexture->GetInputCoordinate(), InTexmap );
+
+		BakedExpression = MaterialExpressionTexture;
+	}
+
 	return BakedExpression;
 }
 
diff --git a/Engine/Source/Programs/Enterprise/Datasmith/DatasmithMaxExporter/Private/MaxMaterialsToUEPbr/DatasmithMaxVRayTexmapToUEPbr.cpp b/Engine/Source/Programs/Enterprise/Datasmith/DatasmithMaxExporter/Private/MaxMaterialsToUEPbr/DatasmithMaxVRayTexmapToUEPbr.cpp
--- a/Engine/Source/Programs/Enterprise/Datasmith/DatasmithMaxExporter/Private/MaxMaterialsToUEPbr/DatasmithMaxVRayTexmapToUEPbr.cpp
+++ b/Engine/Source/Programs/Enterprise/Datasmith/DatasmithMaxExporter/Private/MaxMaterialsToUEPbr/DatasmithMaxVRayTexmapToUEPbr.cpp
@@ -9,6 +9,7 @@
 #include "DatasmithMaxWriter.h"
 #include "MaxMaterialsToUEPbr/DatasmithMaxTexmapToUEPbr.h"
 
+#include "DatasmithSceneFactory.h"
 #include "Misc/Paths.h"
 
 #include "Windows/AllowWindowsPlatformTypes.h"
@@ -160,6 +161,54 @@ bool FDatasmithMaxVRayHDRITexmapToUEPbr::IsSupported( const FDatasmithMaxMateria
 	return InTexmap ? (bool)( InTexmap->ClassID() == VRAYHDRICLASS ) : false;
 }
 
+class FVRayHDRIToTextureElementConverter: public DatasmithMaxDirectLink::ITexmapToTextureElementConverter
+{
+public:
+	virtual TSharedPtr<IDatasmithTextureElement> Convert(DatasmithMaxDirectLink::FMaterialsCollectionTracker& MaterialsTracker, const FString& ActualBitmapName) override
+	{
+		FString Path = TEXT("");
+
+		int NumParamBlocks = Tex->NumParamBlocks();
+
+		for (int j = 0; j < NumParamBlocks; j++)
+		{
+			IParamBlock2* ParamBlock2 = Tex->GetParamBlockByID((short)j);
+			// The the descriptor to 'decode'
+			ParamBlockDesc2* ParamBlockDesc = ParamBlock2->GetDesc();
+			// Loop through all the defined parameters therein
+			for (int i = 0; i < ParamBlockDesc->count; i++)
+			{
+				const ParamDef& ParamDefinition = ParamBlockDesc->paramdefs[i];
+
+				if (FCString::Stricmp(ParamDefinition.int_name, TEXT("HDRIMapName")) == 0)
+				{
+					Path = FDatasmithMaxSceneExporter::GetActualPath(ParamBlock2->GetStr(ParamDefinition.ID, GetCOREInterface()->GetTime()));
+				}
+			}
+			ParamBlock2->ReleaseDesc();
+		}
+
+
+		if (Path.IsEmpty())
+		{
+			return {};
+		}
+
+		float Gamma =  FDatasmithMaxMatHelper::GetVrayHdriGamma(Tex);
+
+		TSharedPtr< IDatasmithTextureElement > TextureElement = FDatasmithSceneFactory::CreateTexture(*ActualBitmapName);
+		if (gammaMgr.IsEnabled())
+		{
+			TextureElement->SetRGBCurve(Gamma / 2.2f);
+		}
+		TextureElement->SetFile(*Path);
+		return TextureElement;
+	}
+	BitmapTex* Tex;
+};
+
+
+
 IDatasmithMaterialExpression* FDatasmithMaxVRayHDRITexmapToUEPbr::Convert( FDatasmithMaxMaterialsToUEPbr* MaxMaterialToUEPbr, Texmap* InTexmap )
 {
 	FString TexturePath;
@@ -191,6 +240,11 @@ IDatasmithMaterialExpression* FDatasmithMaxVRayHDRITexmapToUEPbr::Convert( FData
 	IDatasmithMaterialExpressionTexture* TextureExpression = static_cast< IDatasmithMaterialExpressionTexture* >( MaxMaterialToUEPbr->ConvertState.MaterialElement->AddMaterialExpression( EDatasmithMaterialExpressionType::Texture ) );
 	FString ActualBitmapName = FDatasmithMaxMatWriter::GetActualVRayBitmapName( (BitmapTex*)InTexmap );
 
+	TSharedRef<FVRayHDRIToTextureElementConverter> Converter = MakeShared<FVRayHDRIToTextureElementConverter>();
+	Converter->Tex = (BitmapTex*)InTexmap;
+
+	MaxMaterialToUEPbr->AddTexmap(InTexmap, ActualBitmapName, Converter);
+
 	TextureExpression->SetTexturePathName( *ActualBitmapName );
 	FDatasmithMaxTexmapToUEPbrUtils::SetupTextureCoordinates( MaxMaterialToUEPbr, TextureExpression->GetInputCoordinate(), InTexmap );
 
diff --git a/Engine/Source/Programs/Enterprise/Datasmith/DatasmithMaxExporter/Public/DatasmithMaxDirectLink.h b/Engine/Source/Programs/Enterprise/Datasmith/DatasmithMaxExporter/Public/DatasmithMaxDirectLink.h
--- a/Engine/Source/Programs/Enterprise/Datasmith/DatasmithMaxExporter/Public/DatasmithMaxDirectLink.h
+++ b/Engine/Source/Programs/Enterprise/Datasmith/DatasmithMaxExporter/Public/DatasmithMaxDirectLink.h
@@ -43,6 +43,7 @@ namespace DatasmithMaxDirectLink
 
 class ISceneTracker;
 class FLayerTracker;
+class FMaterialsCollectionTracker;
 
 
 // Global export options, stored in preferences
@@ -353,6 +354,8 @@ struct FSceneUpdateStats
 	int32 CheckTimeSliderInvalidated = 0;
 	int32 ConvertNodesConverted = 0;
 	int32 UpdateInstancesGeometryUpdated = 0;
+	int32 UpdateTexturesBaked = 0;
+	int64 UpdateTexturesBakedPixels = 0;
 };
 
 #define SCENE_UPDATE_STAT_INC(Category, Name) {Stats.##Category##Name++;}
@@ -410,6 +413,7 @@ public:
 	virtual void ReleaseIesTexture(const FString& IesFilePath) = 0;
 
 	virtual TSharedRef<IDatasmithScene> GetDatasmithSceneRef() = 0;
+	virtual const TCHAR* GetAssetsOutputPath() = 0;
 
 	// Utility
 	virtual FNodeTracker* GetNodeTrackerByNodeName(const TCHAR* Name) = 0;
@@ -574,6 +578,17 @@ private:
 	}
 };
 
+// Converts texmap to texture element
+// This is different from IDatasmithMaxTexmapToUEPbr which performs texmap conversion to full MaterialExpression which is used in material graph
+// IDatasmithMaxTexmapToUEPbr uses ToTextureElementConverter when it needs a texture element
+class ITexmapToTextureElementConverter
+{
+public:
+	virtual ~ITexmapToTextureElementConverter(){}
+	virtual TSharedPtr<IDatasmithTextureElement> Convert(FMaterialsCollectionTracker& MaterialsTracker, const FString& ActualBitmapName) = 0;
+
+	FString TextureElementName;
+};
 
 class FMaterialsCollectionTracker
 {
@@ -615,6 +630,14 @@ public:
 		return InvalidatedMaterialTrackers;
 	}
 
+	void UpdateTexmap(Texmap* Texmap);
+
+	// Adds a texmap that need to be converted to texture element
+	void AddTexmapForConversion(Texmap* Texmap, const FString& DesiredTextureElementName, const TSharedPtr<ITexmapToTextureElementConverter>& Converter);
+
+	void AddTextureElement(const TSharedPtr<IDatasmithTextureElement>& );
+	void RemoveTextureElement(const TSharedPtr<IDatasmithTextureElement>& TextureElement);
+
 	ISceneTracker& SceneTracker;
 
 	TMap<FMaterialKey, FMaterialTrackerHandle> MaterialTrackers; // Tracks all assigned materials
@@ -632,7 +655,14 @@ public:
 
 	TMap<Texmap*, TSet<FMaterialTracker*>> UsedTextureToMaterialTracker; // Materials uses by nodes keep set of assigned textures they are used for
 	TMap<Texmap*, TSet<TSharedPtr<IDatasmithTextureElement>>> UsedTextureToDatasmithElement; // Keep track of Datasmith Element created for texmap to simplify update/removal(no need to search DatasmithScene)
-	// note: each texmap can create multiple texture elements
+	TMap<TSharedPtr<IDatasmithTextureElement>, TSet<Texmap*>> TextureElementToTexmap; // Back reference - which texture elements correspond to which texmaps(different texmaps can make same element)
+
+	TMap<Texmap*, TSharedPtr<ITexmapToTextureElementConverter>> TexmapConverters;
+
+	// Bitmap textures using the same file become same texture element
+	TMap<FString, TSharedPtr<IDatasmithTextureElement>>  BitmapTextureElements;
+
+	TSet<TSharedPtr<IDatasmithTextureElement>> TextureElementsAddedToScene;
 
 	FSceneUpdateStats& Stats;
 
diff --git a/Engine/Source/Programs/Enterprise/Datasmith/DatasmithMaxExporter/Public/DatasmithMaxWriter.h b/Engine/Source/Programs/Enterprise/Datasmith/DatasmithMaxExporter/Public/DatasmithMaxWriter.h
--- a/Engine/Source/Programs/Enterprise/Datasmith/DatasmithMaxExporter/Public/DatasmithMaxWriter.h
+++ b/Engine/Source/Programs/Enterprise/Datasmith/DatasmithMaxExporter/Public/DatasmithMaxWriter.h
@@ -142,9 +142,6 @@ public:
 	static FString CropBitmap(BitmapTex* InBitmapTex);
 	static FString DumpBitmap(TSharedPtr<IDatasmithCompositeTexture>& CompTex, BitmapTex* InBitmapTex, const TCHAR* Prefix, bool bForceInvert, bool bIsGrayscale);
 	static FString DumpAutodeskBitmap(TSharedPtr<IDatasmithCompositeTexture>& CompTex, Texmap* InTexmap, const TCHAR* Prefix, bool bForceInvert, bool bIsGrayscale);
-	static void GetRegularTexmap(TSharedRef< IDatasmithScene > DatasmithScene, BitmapTex* InBitmapTex, TArray<TSharedPtr< IDatasmithTextureElement >>* OutTextureElements);
-	static void GetAutodeskTexmap(TSharedRef< IDatasmithScene > DatasmithScene, Texmap* InTexMap, TArray<TSharedPtr< IDatasmithTextureElement >>* OutTextureElements);
-	static TSharedPtr< IDatasmithTextureElement > AddBakeable(TSharedRef< IDatasmithScene > DatasmithScene, Texmap* InTexmap, const TCHAR* AssetsPath);
 	static void ExportStandardMaterial(TSharedRef< IDatasmithScene > DatasmithScene, TSharedPtr< IDatasmithMaterialElement >& MaterialElement, Mtl* Material);
 	static void ExportBlendMaterial(TSharedRef< IDatasmithScene > DatasmithScene, TSharedPtr< IDatasmithMaterialElement >& MaterialElement, Mtl* Material);
 	static FString DumpColorCorrect(TSharedRef< IDatasmithScene > DatasmithScene, TSharedPtr<IDatasmithCompositeTexture>& CompTex, Texmap* InTexmap, const TCHAR* Prefix, bool bForceInvert, bool bIsGrayscale);
@@ -165,7 +162,6 @@ public:
 	// VRAY DEPENDANT
 	//---------------------------------------------------------------
 	static FString DumpVrayColor(TSharedPtr<IDatasmithCompositeTexture>& CompTex, Texmap* InTexmap, const TCHAR* Prefix, bool bForceInvert);
-	static void GetVrayHdri(TSharedRef< IDatasmithScene > DatasmithScene, BitmapTex* InBitmapTex, TArray<TSharedPtr< IDatasmithTextureElement >>* OutTextureElements);
 	static FString DumpVrayHdri(TSharedPtr<IDatasmithCompositeTexture>& CompTex, BitmapTex* InBitmapTex, const TCHAR* Prefix, bool bForceInvert);
 	static void ExportVRayMaterial(TSharedRef< IDatasmithScene > DatasmithScene, TSharedPtr< IDatasmithMaterialElement >& MaterialElement, Mtl* Material);
 	static void ExportVRayLightMaterial(TSharedRef< IDatasmithScene > DatasmithScene, TSharedPtr< IDatasmithMaterialElement >& MaterialElement, Mtl* Material);
