DatasmithSDK: ship the unreal.natvis file with the SDK, use it by default in the SDKSample project

diff --git a/Engine/Source/Programs/Enterprise/Datasmith/DatasmithSDK/DatasmithSDK.Target.cs b/Engine/Source/Programs/Enterprise/Datasmith/DatasmithSDK/DatasmithSDK.Target.cs
--- a/Engine/Source/Programs/Enterprise/Datasmith/DatasmithSDK/DatasmithSDK.Target.cs
+++ b/Engine/Source/Programs/Enterprise/Datasmith/DatasmithSDK/DatasmithSDK.Target.cs
@@ -77,6 +77,11 @@ public class DatasmithSDKTarget : TargetRules
 			@"$(EngineDir)\Binaries\$(TargetPlatform)\DatasmithSDK\Public\"
 		);
 
+		PostBuildCopy(
+			@"$(EngineDir)\Extras\VisualStudioDebugging\Unreal.natvis",
+			@"$(EngineDir)\Binaries\$(TargetPlatform)\DatasmithSDK\"
+		);
+
 		// Other headers we depend on, but that are not part of our public API:
 		PostBuildCopy(
 			@"$(EngineDir)\Source\Runtime\TraceLog\Public\*.h",
