[3ds Max Exporter] Bimpap fixes

- explicit "srgb"  flag set for texture element appropriately
- fixed usecase when same image file used twice(or more) in different bitmap texmaps but with different gamma

diff --git a/Engine/Source/Programs/Enterprise/Datasmith/DatasmithMaxExporter/Private/MaxMaterialsToUEPbr/DatasmithMaxTexmapToUEPbr.cpp b/Engine/Source/Programs/Enterprise/Datasmith/DatasmithMaxExporter/Private/MaxMaterialsToUEPbr/DatasmithMaxTexmapToUEPbr.cpp
--- a/Engine/Source/Programs/Enterprise/Datasmith/DatasmithMaxExporter/Private/MaxMaterialsToUEPbr/DatasmithMaxTexmapToUEPbr.cpp
+++ b/Engine/Source/Programs/Enterprise/Datasmith/DatasmithMaxExporter/Private/MaxMaterialsToUEPbr/DatasmithMaxTexmapToUEPbr.cpp
@@ -850,17 +850,19 @@ bool FDatasmithMaxBitmapToUEPbr::IsSupported( const FDatasmithMaxMaterialsToUEPb
 class FBitmapToTextureElementConverter: public DatasmithMaxDirectLink::ITexmapToTextureElementConverter
 {
 public:
-	virtual TSharedPtr<IDatasmithTextureElement> Convert(DatasmithMaxDirectLink::FMaterialsCollectionTracker& MaterialsTracker, const FString& ActualBitmapName) override
+	virtual TSharedPtr<IDatasmithTextureElement> Convert(DatasmithMaxDirectLink::FMaterialsCollectionTracker& MaterialsTracker, const FString& TextureName) override
 	{
-		FString Path = FDatasmithMaxMatWriter::GetActualBitmapPath(Tex);
-
-		if (TSharedPtr<IDatasmithTextureElement>* Found =  MaterialsTracker.BitmapTextureElements.Find(Path))
+		// Bitmap name is build using all bitmap settings used for the texmap - gamma and cropping
+		// Don't convert same bitmap with same settings twice
+		if (TSharedPtr<IDatasmithTextureElement>* Found =  MaterialsTracker.BitmapTextureElements.Find(ActualBitmapName))
 		{
 			return *Found;
 		}
 
-		TSharedPtr< IDatasmithTextureElement > TextureElement = FDatasmithSceneFactory::CreateTexture(*ActualBitmapName);
-		MaterialsTracker.BitmapTextureElements.Add(Path, TextureElement);
+		FString Path = FDatasmithMaxMatWriter::GetActualBitmapPath(Tex);
+
+		TSharedPtr< IDatasmithTextureElement > TextureElement = FDatasmithSceneFactory::CreateTexture(*TextureName);
+		MaterialsTracker.BitmapTextureElements.Add(ActualBitmapName, TextureElement);
 
 		if (gammaMgr.IsEnabled())
 		{
@@ -875,6 +877,24 @@ public:
 				TextureElement->SetRGBCurve(Gamma);
 			}
 		}
+
+		// todo: ideally TextureMode should be set for Texture element
+		// But Unreal Datasmith Importer currently has some legacy decisions that contradict with the plain logic of Texture Mode
+		// e.g. setting to Bump makes heightmap to convert to normals on import. Which is not desired when that texture is used as heightmap in material graph
+		//TextureElement->SetTextureMode(TextureMode);
+
+		EDatasmithColorSpace ColorSpace = EDatasmithColorSpace::sRGB;
+		switch (TextureMode)
+		{
+		case EDatasmithTextureMode::Normal: 
+		case EDatasmithTextureMode::NormalGreenInv: 
+		case EDatasmithTextureMode::Bump: 
+		case EDatasmithTextureMode::Ies:
+			ColorSpace = EDatasmithColorSpace::Linear;
+			break;
+		default: ;
+		}
+		TextureElement->SetSRGB(ColorSpace);
 		
 		if (!FPaths::FileExists(*Path))
 		{
@@ -887,6 +907,8 @@ public:
 	BitmapTex* Tex;
 
 	bool bIsSRGB;
+	FString ActualBitmapName;
+	EDatasmithTextureMode TextureMode;
 };
 
 IDatasmithMaterialExpression* FDatasmithMaxBitmapToUEPbr::Convert( FDatasmithMaxMaterialsToUEPbr* MaxMaterialToUEPbr, Texmap* InTexmap )
@@ -900,6 +922,8 @@ IDatasmithMaterialExpression* FDatasmithMaxBitmapToUEPbr::Convert( FDatasmithMax
 	TSharedRef<FBitmapToTextureElementConverter> Converter = MakeShared<FBitmapToTextureElementConverter>();
 	Converter->Tex = InBitmapTex;
 	Converter->bIsSRGB = bIsSRGB;
+	Converter->ActualBitmapName = ActualBitmapName;
+	Converter->TextureMode = MaxMaterialToUEPbr->ConvertState.DefaultTextureMode;
 
 	MaxMaterialToUEPbr->AddTexmap(InTexmap, ActualBitmapName, Converter);
 
