Datasmith Revit Exporter: Fixing Mesh Instance duplication (when customized mesh is present)

diff --git a/Engine/Source/Programs/Enterprise/Datasmith/DatasmithRevitExporter/Private/DatasmithRevitDocumentData.cs b/Engine/Source/Programs/Enterprise/Datasmith/DatasmithRevitExporter/Private/DatasmithRevitDocumentData.cs
--- a/Engine/Source/Programs/Enterprise/Datasmith/DatasmithRevitExporter/Private/DatasmithRevitDocumentData.cs
+++ b/Engine/Source/Programs/Enterprise/Datasmith/DatasmithRevitExporter/Private/DatasmithRevitDocumentData.cs
@@ -543,6 +543,15 @@ namespace DatasmithRevitExporter
 			)
 			{
 				FBaseElementData InstanceData = new FBaseElementData(InInstanceType, DocumentData);
+
+				FamilyInstance CurrentFamilyInstance = CurrentElement as FamilyInstance;
+				if (CurrentFamilyInstance != null && CurrentFamilyInstance.HasModifiedGeometry())
+				{
+					//In case the FamilyInstance has a modified geometry, then don't instantiate the original mesh,
+					//the onPolymesh will provide the customized geometry instead.
+					bInAllowMeshInstancing = false;
+				}
+
 				InstanceData.bAllowMeshInstancing = bInAllowMeshInstancing;
 				InstanceDataStack.Push(InstanceData);
 
