Fixed crash when clicking on AutoSync

diff --git a/Engine/Source/Programs/Enterprise/Datasmith/DatasmithARCHICADExporter/Private/Commander.cpp b/Engine/Source/Programs/Enterprise/Datasmith/DatasmithARCHICADExporter/Private/Commander.cpp
--- a/Engine/Source/Programs/Enterprise/Datasmith/DatasmithARCHICADExporter/Private/Commander.cpp
+++ b/Engine/Source/Programs/Enterprise/Datasmith/DatasmithARCHICADExporter/Private/Commander.cpp
@@ -196,7 +196,8 @@ void FCommander::DoSnapshotOrExport(const IO::Location* InExportedFile)
 		void* CurrentSight = nullptr;
 		if (ACAPI_3D_GetCurrentWindowSight(&CurrentSight) != NoError)
 		{
-			throw UE_AC_Error("Current view isn't 3D", UE_AC_Error::kNotIn3DView);
+			UE_AC_DebugF("FCommander::DoSnapshotOrExport - Error : Current view isn't 3D\n");
+			return;
 		}
 		Modeler::SightPtr		 SightPtr((Modeler::Sight*)CurrentSight);
 		Modeler::ConstModel3DPtr Model3D(SightPtr->GetMainModelPtr());
diff --git a/Engine/Source/Programs/Enterprise/Datasmith/DatasmithARCHICADExporter/Private/Palette.cpp b/Engine/Source/Programs/Enterprise/Datasmith/DatasmithARCHICADExporter/Private/Palette.cpp
--- a/Engine/Source/Programs/Enterprise/Datasmith/DatasmithARCHICADExporter/Private/Palette.cpp
+++ b/Engine/Source/Programs/Enterprise/Datasmith/DatasmithARCHICADExporter/Private/Palette.cpp
@@ -264,7 +264,7 @@ void FPalette::WindowChanged()
 	if (CurrentPalette != nullptr)
 	{
 		CurrentPalette->Window->IconSnapshot.SetStatus(bIs3DView);
-		CurrentPalette->Window->IconAutoSync.SetStatus(bIs3DView && FCommander::IsAutoSyncEnabled());
+		CurrentPalette->Window->IconAutoSync.SetStatus(bIs3DView);
 	}
 #endif
 }
diff --git a/Engine/Source/Programs/Enterprise/Datasmith/DatasmithARCHICADExporter/Private/Utils/AddonTools.cpp b/Engine/Source/Programs/Enterprise/Datasmith/DatasmithARCHICADExporter/Private/Utils/AddonTools.cpp
--- a/Engine/Source/Programs/Enterprise/Datasmith/DatasmithARCHICADExporter/Private/Utils/AddonTools.cpp
+++ b/Engine/Source/Programs/Enterprise/Datasmith/DatasmithARCHICADExporter/Private/Utils/AddonTools.cpp
@@ -278,7 +278,11 @@ bool Is3DCurrenWindow()
 			return false;
 		}
 		UE_AC_DebugF("Is3DCurrenWindow - APIDb_GetCurrentWindowID error=%s\n", GetErrorName(GSErr));
+
+		// Something wrong happens. Might not be a 3D window. Not taking any chance
+		return false;
 	}
+
 	return WindowInfo.typeID == APIWind_3DModelID;
 }
 
