Fix broken filetypes

diff --git a/Engine/Source/Programs/Enterprise/Datasmith/DatasmithSDKSample/pch.cpp b/Engine/Source/Programs/Enterprise/Datasmith/DatasmithSDKSample/pch.cpp
--- a/Engine/Source/Programs/Enterprise/Datasmith/DatasmithSDKSample/pch.cpp
+++ b/Engine/Source/Programs/Enterprise/Datasmith/DatasmithSDKSample/pch.cpp
@@ -1,3 +1,3 @@
-ÿþ// Copyright Epic Games, Inc. All Rights Reserved.
+// Copyright Epic Games, Inc. All Rights Reserved.
 
 #include "pch.h"
