Datasmith Revit Exporter Decal related fixes
Issues fixed: missing decals on same surface upon export, missing decals on complex surfaces such as HermiteSplines, Decal clipping on big curvature surfaces in UE.

diff --git a/Engine/Source/Programs/Enterprise/Datasmith/DatasmithRevitExporter/Private/DatasmithRevitDecal.cs b/Engine/Source/Programs/Enterprise/Datasmith/DatasmithRevitExporter/Private/DatasmithRevitDecal.cs
--- a/Engine/Source/Programs/Enterprise/Datasmith/DatasmithRevitExporter/Private/DatasmithRevitDecal.cs
+++ b/Engine/Source/Programs/Enterprise/Datasmith/DatasmithRevitExporter/Private/DatasmithRevitDecal.cs
@@ -89,22 +89,15 @@ namespace DatasmithRevitExporter
 			}
 		}
 
-		public static ElementId GetDecalElementId(MaterialNode InMaterialNode)
+		private static void FillDecalElementIdAndAppearancePairList(Asset Decal, ref IList<KeyValuePair<ElementId,Asset>> DecalElementIds)
 		{
+			const string PropertyName_Decal = "Decal";
 			const string PropertyName_DecalElemId = "decalelementId";
+			const string PropertyName_DecalAppearance = "DecalAppearance";
+			const string PropertyName_SurfaceMaterial = "surface_material";
 
-			if (!InMaterialNode.HasOverriddenAppearance)
-			{
-				return ElementId.InvalidElementId;
-			}
-
-			Asset Decal = InMaterialNode.GetAppearanceOverride();
-
-			if (Decal == null || Decal.Name != "Decal")
-			{
-				return ElementId.InvalidElementId;
-			}
-
+			int DecalElementId = -1;
+			Asset DecalAppearance = null;
 			for (int PropIndex = 0; PropIndex < Decal.Size; ++PropIndex)
 			{
 #if REVIT_API_2018
@@ -112,72 +105,82 @@ namespace DatasmithRevitExporter
 #else
 				AssetProperty Prop = Decal.Get(PropIndex);
 #endif
+
+				if (Prop.Type == AssetPropertyType.Reference)
+				{
+					AssetPropertyReference RefProp = (Prop as AssetPropertyReference);
+					foreach (AssetProperty InnerProp in RefProp.GetAllConnectedProperties())
+					{
+						if (Prop.Name == PropertyName_SurfaceMaterial && InnerProp.Name == PropertyName_Decal && InnerProp.Type == AssetPropertyType.Asset)
+						{
+							Asset InnerDecal = InnerProp as Asset;
+							FillDecalElementIdAndAppearancePairList(InnerDecal, ref DecalElementIds);
+						}
+
+						if (InnerProp.Name == PropertyName_DecalAppearance && InnerProp.Type == AssetPropertyType.Asset)
+						{
+							DecalAppearance = InnerProp as Asset;
+						}
+					}
+				}
+
 				if (Prop.Name == PropertyName_DecalElemId && Prop.Type == AssetPropertyType.Integer)
 				{
-					int Id = (Prop as AssetPropertyInteger).Value;
-					return new ElementId(Id);
+					DecalElementId = (Prop as AssetPropertyInteger).Value;
 				}
 			}
-			return ElementId.InvalidElementId;
+
+			if (DecalAppearance != null && DecalElementId != -1)
+			{
+				DecalElementIds.Add(new KeyValuePair<ElementId, Asset>(new ElementId(DecalElementId), DecalAppearance));
+			}
 		}
 
-		public static FDecalMaterial Create(MaterialNode InMaterialNode, Material InMaterial)
+		public static IList<KeyValuePair<ElementId, Asset>> GetDecalElementIdAndAppearancePairList(MaterialNode InMaterialNode)
 		{
+			IList<KeyValuePair<ElementId, Asset>> DecalElementIdAndAppearancePairList = new List<KeyValuePair<ElementId, Asset>>();
+
 			if (!InMaterialNode.HasOverriddenAppearance)
 			{
-				return null;
+				return DecalElementIdAndAppearancePairList;
 			}
 
 			Asset Decal = InMaterialNode.GetAppearanceOverride();
 
 			if (Decal == null || Decal.Name != "Decal")
 			{
-				return null;
+				return DecalElementIdAndAppearancePairList;
 			}
 
-			const string AssetName_Material = "DecalAppearance";
+			FillDecalElementIdAndAppearancePairList(Decal, ref DecalElementIdAndAppearancePairList);
 
+			return DecalElementIdAndAppearancePairList;
+		}
+
+		public static FDecalMaterial Create(MaterialNode InMaterialNode, Material InMaterial, int DecalAppearanceIndex, Asset DecalAppearanceAsset)
+		{
 			FDecalMaterial DecalMaterial = new FDecalMaterial();
 
-			DecalMaterial.MaterialName = FDatasmithFacadeElement.GetStringHash($"Decal_{FMaterialData.GetMaterialName(InMaterialNode, InMaterial)}");
+			DecalMaterial.MaterialName = FDatasmithFacadeElement.GetStringHash($"Decal_{FMaterialData.GetMaterialName(InMaterialNode, InMaterial)}_{DecalAppearanceIndex}");
 
-			// Look for DecalAppearance asset, which defines the material parameters
+			DecalMaterial.SetMaterialProperties(DecalAppearanceAsset);
 
-			bool bMaterialInitialized = false;
+			return DecalMaterial;
+		}
 
-			for (int PropIndex = 0; PropIndex < Decal.Size; ++PropIndex)
+		//Ignores MaterialName, only checks Render value equality
+		public bool CheckRenderValueEquiality(FDecalMaterial Input)
+		{
+			if (DiffuseTexturePath == Input.DiffuseTexturePath
+				&& BumpTexturePath == Input.BumpTexturePath
+				&& CutoutTexturePath == Input.CutoutTexturePath
+				&& BumpAmount == Input.BumpAmount
+				&& Transparency == Input.Transparency
+				&& Luminance == Input.Luminance)
 			{
-#if REVIT_API_2018
-				AssetProperty Prop = Decal[PropIndex];
-#else
-				AssetProperty Prop = Decal.Get(PropIndex);
-#endif
-				if (Prop.Type == AssetPropertyType.Reference)
-				{
-					AssetPropertyReference RefProp = (Prop as AssetPropertyReference);
-					foreach (AssetProperty InnerProp in RefProp.GetAllConnectedProperties())
-					{
-						if (InnerProp.Type == AssetPropertyType.Asset)
-						{
-							Asset InnerPropAsset = InnerProp as Asset;
-
-							if (InnerPropAsset.Name == AssetName_Material)
-							{
-								DecalMaterial.SetMaterialProperties(InnerPropAsset);
-								bMaterialInitialized = true;
-								break;
-							}
-						}
-					}
-				}
-
-				if (bMaterialInitialized)
-				{
-					break;
-				}
+				return true;
 			}
-
-			return DecalMaterial;
+			return false;
 		}
 	}
 }
diff --git a/Engine/Source/Programs/Enterprise/Datasmith/DatasmithRevitExporter/Private/DatasmithRevitDocumentData.cs b/Engine/Source/Programs/Enterprise/Datasmith/DatasmithRevitExporter/Private/DatasmithRevitDocumentData.cs
--- a/Engine/Source/Programs/Enterprise/Datasmith/DatasmithRevitExporter/Private/DatasmithRevitDocumentData.cs
+++ b/Engine/Source/Programs/Enterprise/Datasmith/DatasmithRevitExporter/Private/DatasmithRevitDocumentData.cs
@@ -1479,15 +1479,38 @@ namespace DatasmithRevitExporter
 
 			if (InMaterialNode.HasOverriddenAppearance)
 			{
-				ElementId ReferencingDecalId = FDecalMaterial.GetDecalElementId(InMaterialNode);
-
-				if (ReferencingDecalId != ElementId.InvalidElementId)
+				IList<KeyValuePair<ElementId, Asset>> ReferencingDecalIdAndAssetPairs = FDecalMaterial.GetDecalElementIdAndAppearancePairList(InMaterialNode);
+				//Assets in on the .Value seem to be always unique, regardless if they are "instances" of the same DecalType:
+				if (ReferencingDecalIdAndAssetPairs.Count > 0)
 				{
-					FDecalMaterial DecalMaterial = FDecalMaterial.Create(InMaterialNode, CurrentMaterial);
-
-					if (DecalMaterial != null)
+					IList<FDecalMaterial> DecalMaterials = new List<FDecalMaterial>();
+					foreach (KeyValuePair<ElementId, Asset> DecalIdAndAssetPair in ReferencingDecalIdAndAssetPairs)
 					{
-						DecalMaterialsMap.Add(ReferencingDecalId, DecalMaterial);
+						FDecalMaterial DecalMaterial = FDecalMaterial.Create(InMaterialNode, CurrentMaterial, DecalMaterials.Count, DecalIdAndAssetPair.Value);
+
+						if (DecalMaterial != null)
+						{
+							//unique DecalMaterial:
+							FDecalMaterial ExistingDecalMaterial = null;
+							foreach (FDecalMaterial ExistingDecalMaterialCandidate in DecalMaterials)
+							{
+								if (ExistingDecalMaterialCandidate.CheckRenderValueEquiality(DecalMaterial))
+								{
+									ExistingDecalMaterial = ExistingDecalMaterialCandidate;
+									break;
+								}
+							}
+
+							if (ExistingDecalMaterial != null)
+							{
+								DecalMaterialsMap.Add(DecalIdAndAssetPair.Key, ExistingDecalMaterial);
+							}
+							else
+							{
+								DecalMaterialsMap.Add(DecalIdAndAssetPair.Key, DecalMaterial);
+								DecalMaterials.Add(DecalMaterial);
+							}
+						}
 					}
 				}
 			}
diff --git a/Engine/Source/Programs/Enterprise/Datasmith/DatasmithRevitExporter/Private/DatasmithRevitUtils.cs b/Engine/Source/Programs/Enterprise/Datasmith/DatasmithRevitExporter/Private/DatasmithRevitUtils.cs
--- a/Engine/Source/Programs/Enterprise/Datasmith/DatasmithRevitExporter/Private/DatasmithRevitUtils.cs
+++ b/Engine/Source/Programs/Enterprise/Datasmith/DatasmithRevitExporter/Private/DatasmithRevitUtils.cs
@@ -119,13 +119,14 @@ namespace DatasmithRevitExporter
 				{
 					DecalQuad.Add(QuadLine);
 				}
-				else if (GeomObj is Arc QuadArc)
+				else if (GeomObj is Curve QuadCurve)
 				{
 					try
 					{
-						XYZ StartPoint = QuadArc.GetEndPoint(0);
-						XYZ EndPoint = QuadArc.GetEndPoint(1);
-						DecalQuad.Add(Line.CreateBound(StartPoint, EndPoint));
+						XYZ StartPoint = QuadCurve.GetEndPoint(0);
+						XYZ EndPoint = QuadCurve.GetEndPoint(1);
+						Line BoundLine = Line.CreateBound(StartPoint, EndPoint);
+						DecalQuad.Add(BoundLine);
 					}
 					catch {}
 				}
@@ -155,7 +156,20 @@ namespace DatasmithRevitExporter
 
 			const float CENTIMETERS_PER_FOOT = 30.48F;
 
-			OutDecalDimensions = new XYZ(DecalQuad[0].Length * CENTIMETERS_PER_FOOT * 0.5, DecalQuad[1].Length * CENTIMETERS_PER_FOOT * 0.5, 2.0);
+			Plane TestPlane = Plane.CreateByThreePoints(TopLeft, TopRight, BottomRight);
+			UV TestUv;
+			double Distance;
+			TestPlane.Project(BottomLeft, out TestUv, out Distance);
+
+			Distance *= CENTIMETERS_PER_FOOT * 1.05; // leniency multiplier of 1.05, for imprecisions
+
+			double DimensionZ = 2.0;
+			if (DimensionZ < Distance)
+			{
+				DimensionZ = Distance;
+			}
+
+			OutDecalDimensions = new XYZ(DecalQuad[0].Length * CENTIMETERS_PER_FOOT * 0.5, DecalQuad[1].Length * CENTIMETERS_PER_FOOT * 0.5, DimensionZ);
 		}
 	}
 }
