[SketchUp] Image fixes

- Datasmith: WorldPositionOffset input added to Datasmith UePbr material
- Datasmith: DatasmithCSharpFacade - added ClearCoat/ClearCoatRoughness/WorldPositionOffset
- Fixed z-fighting - Image Geometry is offset along local Z axis. PositionOffset param of image's material can be used to control amount of offset
- Enabled transparency on image material when it contains transparent pixels

diff --git a/Engine/Plugins/Enterprise/DatasmithImporter/Source/DatasmithImporter/Private/DatasmithMaterialExpressions.cpp b/Engine/Plugins/Enterprise/DatasmithImporter/Source/DatasmithImporter/Private/DatasmithMaterialExpressions.cpp
--- a/Engine/Plugins/Enterprise/DatasmithImporter/Source/DatasmithImporter/Private/DatasmithMaterialExpressions.cpp
+++ b/Engine/Plugins/Enterprise/DatasmithImporter/Source/DatasmithImporter/Private/DatasmithMaterialExpressions.cpp
@@ -3129,6 +3129,11 @@ void FDatasmithMaterialExpressions::CreateUEPbrMaterialGraph(const TSharedPtr< I
 		}
 	}
 
+	if (MaterialElement->GetWorldPositionOffset().GetExpression())
+	{
+		ConnectExpression(MaterialElement.ToSharedRef(), MaterialExpressions, MaterialElement->GetWorldPositionOffset().GetExpression(), GetMaterialOrFunctionSlot(UnrealMaterialOrFunction, EDatasmithTextureSlot::WORLDPOSITIONOFFSET), MaterialElement->GetWorldPositionOffset().GetOutputIndex());
+	}
+
 	// Connect expressions to any UMaterialExpressionCustomOutput since these aren't part of the predefined material outputs
 	for ( int32 ExpressionIndex = 0; ExpressionIndex < MaterialElement->GetExpressionsCount(); ++ExpressionIndex )
 	{
diff --git a/Engine/Plugins/Enterprise/DatasmithImporter/Source/DatasmithImporter/Private/DatasmithMaterialImporter.cpp b/Engine/Plugins/Enterprise/DatasmithImporter/Source/DatasmithImporter/Private/DatasmithMaterialImporter.cpp
--- a/Engine/Plugins/Enterprise/DatasmithImporter/Source/DatasmithImporter/Private/DatasmithMaterialImporter.cpp
+++ b/Engine/Plugins/Enterprise/DatasmithImporter/Source/DatasmithImporter/Private/DatasmithMaterialImporter.cpp
@@ -151,6 +151,7 @@ namespace DatasmithMaterialImporterUtils
 		Hash = HashCombine( Hash, ComputeExpressionInputHash( &MaterialElement->GetAmbientOcclusion(), VisitedExpressions ) );
 		Hash = HashCombine( Hash, ComputeExpressionInputHash( &MaterialElement->GetClearCoat(), VisitedExpressions ) );
 		Hash = HashCombine( Hash, ComputeExpressionInputHash( &MaterialElement->GetClearCoatRoughness(), VisitedExpressions ) );
+		Hash = HashCombine( Hash, ComputeExpressionInputHash( &MaterialElement->GetWorldPositionOffset(), VisitedExpressions ) );
 		Hash = HashCombine( Hash, ComputeExpressionInputHash( &MaterialElement->GetMaterialAttributes(), VisitedExpressions ) );
 
 		TSet<IDatasmithMaterialExpression*> AllExpressions;
diff --git a/Engine/Source/Developer/Datasmith/DatasmithFacade/Private/DatasmithFacadeUEPbrMaterial.cpp b/Engine/Source/Developer/Datasmith/DatasmithFacade/Private/DatasmithFacadeUEPbrMaterial.cpp
--- a/Engine/Source/Developer/Datasmith/DatasmithFacade/Private/DatasmithFacadeUEPbrMaterial.cpp
+++ b/Engine/Source/Developer/Datasmith/DatasmithFacade/Private/DatasmithFacadeUEPbrMaterial.cpp
@@ -368,6 +368,24 @@ FDatasmithFacadeExpressionInput FDatasmithFacadeUEPbrMaterial::GetAmbientOcclusi
 	return FDatasmithFacadeExpressionInput( &UEPbrMaterial->GetAmbientOcclusion(), UEPbrMaterial );
 }
 
+FDatasmithFacadeExpressionInput FDatasmithFacadeUEPbrMaterial::GetClearCoat() const
+{
+	TSharedPtr<IDatasmithUEPbrMaterialElement> UEPbrMaterial = GetDatasmithUEPbrMaterialElement();
+	return FDatasmithFacadeExpressionInput( &UEPbrMaterial->GetClearCoat(), UEPbrMaterial );
+}
+
+FDatasmithFacadeExpressionInput FDatasmithFacadeUEPbrMaterial::GetClearCoatRoughness() const
+{
+	TSharedPtr<IDatasmithUEPbrMaterialElement> UEPbrMaterial = GetDatasmithUEPbrMaterialElement();
+	return FDatasmithFacadeExpressionInput( &UEPbrMaterial->GetClearCoatRoughness(), UEPbrMaterial );
+}
+
+FDatasmithFacadeExpressionInput FDatasmithFacadeUEPbrMaterial::GetWorldPositionOffset() const
+{
+	TSharedPtr<IDatasmithUEPbrMaterialElement> UEPbrMaterial = GetDatasmithUEPbrMaterialElement();
+	return FDatasmithFacadeExpressionInput( &UEPbrMaterial->GetWorldPositionOffset(), UEPbrMaterial );
+}
+
 FDatasmithFacadeExpressionInput FDatasmithFacadeUEPbrMaterial::GetMaterialAttributes() const
 {
 	TSharedPtr<IDatasmithUEPbrMaterialElement> UEPbrMaterial = GetDatasmithUEPbrMaterialElement();
diff --git a/Engine/Source/Developer/Datasmith/DatasmithFacade/Public/DatasmithFacadeUEPbrMaterial.h b/Engine/Source/Developer/Datasmith/DatasmithFacade/Public/DatasmithFacadeUEPbrMaterial.h
--- a/Engine/Source/Developer/Datasmith/DatasmithFacade/Public/DatasmithFacadeUEPbrMaterial.h
+++ b/Engine/Source/Developer/Datasmith/DatasmithFacade/Public/DatasmithFacadeUEPbrMaterial.h
@@ -330,6 +330,9 @@ public:
 	FDatasmithFacadeExpressionInput GetNormal() const;
 	FDatasmithFacadeExpressionInput GetRefraction() const;
 	FDatasmithFacadeExpressionInput GetAmbientOcclusion() const;
+	FDatasmithFacadeExpressionInput GetClearCoat() const;
+	FDatasmithFacadeExpressionInput GetClearCoatRoughness() const;
+	FDatasmithFacadeExpressionInput GetWorldPositionOffset() const;
 	FDatasmithFacadeExpressionInput GetMaterialAttributes() const;
 
 	int GetBlendMode() const;
diff --git a/Engine/Source/Programs/Enterprise/Datasmith/DatasmithFacadeCSharp/Private/DatasmithFacadeCSharp.cpp b/Engine/Source/Programs/Enterprise/Datasmith/DatasmithFacadeCSharp/Private/DatasmithFacadeCSharp.cpp
--- a/Engine/Source/Programs/Enterprise/Datasmith/DatasmithFacadeCSharp/Private/DatasmithFacadeCSharp.cpp
+++ b/Engine/Source/Programs/Enterprise/Datasmith/DatasmithFacadeCSharp/Private/DatasmithFacadeCSharp.cpp
@@ -1863,6 +1863,42 @@ SWIGEXPORT void * SWIGSTDCALL CSharp_FDatasmithFacadeUEPbrMaterial_GetAmbientOcc
 }
 
 
+SWIGEXPORT void * SWIGSTDCALL CSharp_FDatasmithFacadeUEPbrMaterial_GetClearCoat(void * jarg1) {
+  void * jresult ;
+  FDatasmithFacadeUEPbrMaterial *arg1 = (FDatasmithFacadeUEPbrMaterial *) 0 ;
+  SwigValueWrapper< FDatasmithFacadeExpressionInput > result;
+  
+  arg1 = (FDatasmithFacadeUEPbrMaterial *)jarg1; 
+  result = ((FDatasmithFacadeUEPbrMaterial const *)arg1)->GetClearCoat();
+  jresult = new FDatasmithFacadeExpressionInput((const FDatasmithFacadeExpressionInput &)result); 
+  return jresult;
+}
+
+
+SWIGEXPORT void * SWIGSTDCALL CSharp_FDatasmithFacadeUEPbrMaterial_GetClearCoatRoughness(void * jarg1) {
+  void * jresult ;
+  FDatasmithFacadeUEPbrMaterial *arg1 = (FDatasmithFacadeUEPbrMaterial *) 0 ;
+  SwigValueWrapper< FDatasmithFacadeExpressionInput > result;
+  
+  arg1 = (FDatasmithFacadeUEPbrMaterial *)jarg1; 
+  result = ((FDatasmithFacadeUEPbrMaterial const *)arg1)->GetClearCoatRoughness();
+  jresult = new FDatasmithFacadeExpressionInput((const FDatasmithFacadeExpressionInput &)result); 
+  return jresult;
+}
+
+
+SWIGEXPORT void * SWIGSTDCALL CSharp_FDatasmithFacadeUEPbrMaterial_GetWorldPositionOffset(void * jarg1) {
+  void * jresult ;
+  FDatasmithFacadeUEPbrMaterial *arg1 = (FDatasmithFacadeUEPbrMaterial *) 0 ;
+  SwigValueWrapper< FDatasmithFacadeExpressionInput > result;
+  
+  arg1 = (FDatasmithFacadeUEPbrMaterial *)jarg1; 
+  result = ((FDatasmithFacadeUEPbrMaterial const *)arg1)->GetWorldPositionOffset();
+  jresult = new FDatasmithFacadeExpressionInput((const FDatasmithFacadeExpressionInput &)result); 
+  return jresult;
+}
+
+
 SWIGEXPORT void * SWIGSTDCALL CSharp_FDatasmithFacadeUEPbrMaterial_GetMaterialAttributes(void * jarg1) {
   void * jresult ;
   FDatasmithFacadeUEPbrMaterial *arg1 = (FDatasmithFacadeUEPbrMaterial *) 0 ;
diff --git a/Engine/Source/Programs/Enterprise/Datasmith/DatasmithFacadeCSharp/Public/DatasmithFacadeCSharpPINVOKE.cs b/Engine/Source/Programs/Enterprise/Datasmith/DatasmithFacadeCSharp/Public/DatasmithFacadeCSharpPINVOKE.cs
--- a/Engine/Source/Programs/Enterprise/Datasmith/DatasmithFacadeCSharp/Public/DatasmithFacadeCSharpPINVOKE.cs
+++ b/Engine/Source/Programs/Enterprise/Datasmith/DatasmithFacadeCSharp/Public/DatasmithFacadeCSharpPINVOKE.cs
@@ -620,6 +620,15 @@ class DatasmithFacadeCSharpPINVOKE {
   [global::System.Runtime.InteropServices.DllImport("DatasmithFacadeCSharp", EntryPoint="CSharp_FDatasmithFacadeUEPbrMaterial_GetAmbientOcclusion")]
   public static extern global::System.IntPtr FDatasmithFacadeUEPbrMaterial_GetAmbientOcclusion(global::System.Runtime.InteropServices.HandleRef jarg1);
 
+  [global::System.Runtime.InteropServices.DllImport("DatasmithFacadeCSharp", EntryPoint="CSharp_FDatasmithFacadeUEPbrMaterial_GetClearCoat")]
+  public static extern global::System.IntPtr FDatasmithFacadeUEPbrMaterial_GetClearCoat(global::System.Runtime.InteropServices.HandleRef jarg1);
+
+  [global::System.Runtime.InteropServices.DllImport("DatasmithFacadeCSharp", EntryPoint="CSharp_FDatasmithFacadeUEPbrMaterial_GetClearCoatRoughness")]
+  public static extern global::System.IntPtr FDatasmithFacadeUEPbrMaterial_GetClearCoatRoughness(global::System.Runtime.InteropServices.HandleRef jarg1);
+
+  [global::System.Runtime.InteropServices.DllImport("DatasmithFacadeCSharp", EntryPoint="CSharp_FDatasmithFacadeUEPbrMaterial_GetWorldPositionOffset")]
+  public static extern global::System.IntPtr FDatasmithFacadeUEPbrMaterial_GetWorldPositionOffset(global::System.Runtime.InteropServices.HandleRef jarg1);
+
   [global::System.Runtime.InteropServices.DllImport("DatasmithFacadeCSharp", EntryPoint="CSharp_FDatasmithFacadeUEPbrMaterial_GetMaterialAttributes")]
   public static extern global::System.IntPtr FDatasmithFacadeUEPbrMaterial_GetMaterialAttributes(global::System.Runtime.InteropServices.HandleRef jarg1);
 
diff --git a/Engine/Source/Programs/Enterprise/Datasmith/DatasmithFacadeCSharp/Public/FDatasmithFacadeUEPbrMaterial.cs b/Engine/Source/Programs/Enterprise/Datasmith/DatasmithFacadeCSharp/Public/FDatasmithFacadeUEPbrMaterial.cs
--- a/Engine/Source/Programs/Enterprise/Datasmith/DatasmithFacadeCSharp/Public/FDatasmithFacadeUEPbrMaterial.cs
+++ b/Engine/Source/Programs/Enterprise/Datasmith/DatasmithFacadeCSharp/Public/FDatasmithFacadeUEPbrMaterial.cs
@@ -83,6 +83,21 @@ public class FDatasmithFacadeUEPbrMaterial : FDatasmithFacadeBaseMaterial {
     return ret;
   }
 
+  public FDatasmithFacadeExpressionInput GetClearCoat() {
+    FDatasmithFacadeExpressionInput ret = new FDatasmithFacadeExpressionInput(DatasmithFacadeCSharpPINVOKE.FDatasmithFacadeUEPbrMaterial_GetClearCoat(swigCPtr), true);
+    return ret;
+  }
+
+  public FDatasmithFacadeExpressionInput GetClearCoatRoughness() {
+    FDatasmithFacadeExpressionInput ret = new FDatasmithFacadeExpressionInput(DatasmithFacadeCSharpPINVOKE.FDatasmithFacadeUEPbrMaterial_GetClearCoatRoughness(swigCPtr), true);
+    return ret;
+  }
+
+  public FDatasmithFacadeExpressionInput GetWorldPositionOffset() {
+    FDatasmithFacadeExpressionInput ret = new FDatasmithFacadeExpressionInput(DatasmithFacadeCSharpPINVOKE.FDatasmithFacadeUEPbrMaterial_GetWorldPositionOffset(swigCPtr), true);
+    return ret;
+  }
+
   public FDatasmithFacadeExpressionInput GetMaterialAttributes() {
     FDatasmithFacadeExpressionInput ret = new FDatasmithFacadeExpressionInput(DatasmithFacadeCSharpPINVOKE.FDatasmithFacadeUEPbrMaterial_GetMaterialAttributes(swigCPtr), true);
     return ret;
diff --git a/Engine/Source/Programs/Enterprise/Datasmith/DatasmithSketchUpRubyExporter/Private/DatasmithSketchUpExportContext.h b/Engine/Source/Programs/Enterprise/Datasmith/DatasmithSketchUpRubyExporter/Private/DatasmithSketchUpExportContext.h
--- a/Engine/Source/Programs/Enterprise/Datasmith/DatasmithSketchUpRubyExporter/Private/DatasmithSketchUpExportContext.h
+++ b/Engine/Source/Programs/Enterprise/Datasmith/DatasmithSketchUpRubyExporter/Private/DatasmithSketchUpExportContext.h
@@ -166,6 +166,7 @@ namespace DatasmithSketchUp
 
 		const TCHAR* GetImageFilePath(FImageFile&);
 		FMD5Hash GetImageFileHash(FImageFile&);
+		bool GetImageHasAlpha(FImageFile& ImageFile);
 
 	private:
 		FExportContext& Context;
diff --git a/Engine/Source/Programs/Enterprise/Datasmith/DatasmithSketchUpRubyExporter/Private/DatasmithSketchUpImage.cpp b/Engine/Source/Programs/Enterprise/Datasmith/DatasmithSketchUpRubyExporter/Private/DatasmithSketchUpImage.cpp
--- a/Engine/Source/Programs/Enterprise/Datasmith/DatasmithSketchUpRubyExporter/Private/DatasmithSketchUpImage.cpp
+++ b/Engine/Source/Programs/Enterprise/Datasmith/DatasmithSketchUpRubyExporter/Private/DatasmithSketchUpImage.cpp
@@ -25,6 +25,7 @@
 #include "DatasmithSceneFactory.h"
 #include "DatasmithUtils.h"
 
+#include "Algo/AnyOf.h"
 #include "Misc/Paths.h"
 
 
@@ -366,6 +367,8 @@ namespace DatasmithSketchUp
 
 		FString FilePath;
 
+		bool bHasAlpha = false;
+
 		FImageFile(SUImageRepRef InImageRep)
 			: ImageRep(InImageRep)
 		{
@@ -387,6 +390,12 @@ namespace DatasmithSketchUp
 			MD5.Update(Data.GetData(), Data.Num());
 
 			ImageHash.Set(MD5);
+
+			TArray<SUColor> Colors;
+			Colors.SetNum(Width*Height);
+			SUImageRepGetDataAsColors(ImageRep, Colors.GetData());
+
+			bHasAlpha = Algo::AnyOf(Colors, [](const SUColor& Color) { return Color.alpha != 255; });
 		}
 
 		FMD5Hash ImageHash;
@@ -419,3 +428,8 @@ FMD5Hash FImageFileCollection::GetImageFileHash(FImageFile& ImageFile)
 {
 	return ImageFile.ImageHash;
 }
+
+bool FImageFileCollection::GetImageHasAlpha(FImageFile& ImageFile)
+{
+	return ImageFile.bHasAlpha;	
+}
diff --git a/Engine/Source/Programs/Enterprise/Datasmith/DatasmithSketchUpRubyExporter/Private/DatasmithSketchUpMaterial.cpp b/Engine/Source/Programs/Enterprise/Datasmith/DatasmithSketchUpRubyExporter/Private/DatasmithSketchUpMaterial.cpp
--- a/Engine/Source/Programs/Enterprise/Datasmith/DatasmithSketchUpRubyExporter/Private/DatasmithSketchUpMaterial.cpp
+++ b/Engine/Source/Programs/Enterprise/Datasmith/DatasmithSketchUpRubyExporter/Private/DatasmithSketchUpMaterial.cpp
@@ -783,6 +783,7 @@ public:
 		Context.DatasmithScene->AddTexture(TextureElement);
 
 		FImageMaterial::FOptions MaterialOptions;
+		MaterialOptions.bUseAlphaChannel = Context.ImageFiles.GetImageHasAlpha(ImageFile);
 		CreateMaterialElement(Context, *MaterialName, *MaterialLabel, TextureElement->GetName(), MaterialOptions);
 
 		Hash = Context.ImageFiles.GetImageFileHash(ImageFile);
@@ -825,7 +826,6 @@ public:
 
 		DatasmithMaterialElement->SetTwoSided(true); // imported Image in SketchUp is displayed as two-sided with no way to controls it, so 
 
-		// todo: inspect texture to see if there's alpha. Also consider masked opacity(might be better fit for SketchUp images?)
 		bool bTranslucent = false; 
 
 		{
@@ -860,14 +860,7 @@ public:
 			// Set the Datasmith material element opacity.
 			if (Options.bUseAlphaChannel)
 			{
-				// Invert texture translarency to get Unreal opacity
-				IDatasmithMaterialExpressionGeneric* ExpressionOpacity = DatasmithMaterialElement->AddMaterialExpression<IDatasmithMaterialExpressionGeneric>();
-				ExpressionOpacity->SetExpressionName(TEXT("OneMinus"));
-
-
-				ExpressionTexture->ConnectExpression(*ExpressionOpacity->GetInput(0), 3);
-
-				ExpressionOpacity->ConnectExpression(DatasmithMaterialElement->GetOpacity());
+				ExpressionTexture->ConnectExpression(DatasmithMaterialElement->GetOpacity(), 4);  // Use Alpha output from Texture sampler
 			}
 		}
 
@@ -876,6 +869,24 @@ public:
 			DatasmithMaterialElement->SetBlendMode(/*EBlendMode::BLEND_Translucent*/2);
 		}
 
+		// World Position Offset
+		// Shift image along local Z
+		{
+			IDatasmithMaterialExpressionGeneric* ExpressionObjectOrientation = DatasmithMaterialElement->AddMaterialExpression<IDatasmithMaterialExpressionGeneric>();
+			ExpressionObjectOrientation->SetExpressionName(TEXT("ObjectOrientation"));
+			IDatasmithMaterialExpressionScalar* ExpressionPositionOffset = DatasmithMaterialElement->AddMaterialExpression<IDatasmithMaterialExpressionScalar>();
+			ExpressionPositionOffset->SetName(TEXT("PositionOffset"));
+			ExpressionPositionOffset->GetScalar() = 1.f;
+
+			IDatasmithMaterialExpressionGeneric* MultiplyOrientationWithOffset = DatasmithMaterialElement->AddMaterialExpression<IDatasmithMaterialExpressionGeneric>();
+			MultiplyOrientationWithOffset->SetExpressionName(TEXT("Multiply"));
+
+			ExpressionObjectOrientation->ConnectExpression(*MultiplyOrientationWithOffset->GetInput(0));
+			ExpressionPositionOffset->ConnectExpression(*MultiplyOrientationWithOffset->GetInput(1));
+
+			MultiplyOrientationWithOffset->ConnectExpression(DatasmithMaterialElement->GetWorldPositionOffset());
+		}
+
 		Context.DatasmithScene->AddMaterial(DatasmithMaterialElement);
 
 		// Return the Datasmith material element.
diff --git a/Engine/Source/Runtime/Datasmith/DatasmithCore/Private/DatasmithMaterialElementsImpl.cpp b/Engine/Source/Runtime/Datasmith/DatasmithCore/Private/DatasmithMaterialElementsImpl.cpp
--- a/Engine/Source/Runtime/Datasmith/DatasmithCore/Private/DatasmithMaterialElementsImpl.cpp
+++ b/Engine/Source/Runtime/Datasmith/DatasmithCore/Private/DatasmithMaterialElementsImpl.cpp
@@ -203,6 +203,7 @@ FDatasmithUEPbrMaterialElementImpl::FDatasmithUEPbrMaterialElementImpl( const TC
 	, AmbientOcclusion(   MakeShared< FDatasmithExpressionInputImpl >( TEXT("AmbientOcclusion") ) )
 	, ClearCoat(		  MakeShared< FDatasmithExpressionInputImpl >( TEXT("ClearCoat") ) )
 	, ClearCoatRoughness( MakeShared< FDatasmithExpressionInputImpl >( TEXT("ClearCoatRoughness") ) )
+	, WorldPositionOffset( MakeShared< FDatasmithExpressionInputImpl >( TEXT("WorldPositionOffset") ) )
 	, MaterialAttributes( MakeShared< FDatasmithExpressionInputImpl >( TEXT("MaterialAttributes") ) )
 	, BlendMode(0)
 	, bTwoSided( false )
@@ -224,6 +225,7 @@ FDatasmithUEPbrMaterialElementImpl::FDatasmithUEPbrMaterialElementImpl( const TC
 	RegisterReferenceProxy( AmbientOcclusion, "AmbientOcclusion" );
 	RegisterReferenceProxy( ClearCoat, "ClearCoat" );
 	RegisterReferenceProxy( ClearCoatRoughness, "ClearCoatRoughness" );
+	RegisterReferenceProxy( WorldPositionOffset, "WorldPositionOffset" );
 	RegisterReferenceProxy( MaterialAttributes, "MaterialAttributes" );
 
 	RegisterReferenceProxy( Expressions, "Expressions" );
@@ -277,6 +279,7 @@ FMD5Hash FDatasmithUEPbrMaterialElementImpl::CalculateElementHash(bool bForce)
 	UpdateMD5(*AmbientOcclusion.Edit());
 	UpdateMD5(*ClearCoat.Edit());
 	UpdateMD5(*ClearCoatRoughness.Edit());
+	UpdateMD5(*WorldPositionOffset.Edit());
 	UpdateMD5(*MaterialAttributes.Edit());
 
 	ElementHash.Set(MD5);
@@ -375,6 +378,7 @@ void FDatasmithUEPbrMaterialElementImpl::ResetExpressionGraph( bool bRemoveAllEx
 	AmbientOcclusion.Edit()->SetExpression(nullptr);
 	ClearCoat.Edit()->SetExpression(nullptr);
 	ClearCoatRoughness.Edit()->SetExpression(nullptr);
+	WorldPositionOffset.Edit()->SetExpression(nullptr);
 	MaterialAttributes.Edit()->SetExpression(nullptr);
 }
 
diff --git a/Engine/Source/Runtime/Datasmith/DatasmithCore/Private/DatasmithMaterialElementsImpl.h b/Engine/Source/Runtime/Datasmith/DatasmithCore/Private/DatasmithMaterialElementsImpl.h
--- a/Engine/Source/Runtime/Datasmith/DatasmithCore/Private/DatasmithMaterialElementsImpl.h
+++ b/Engine/Source/Runtime/Datasmith/DatasmithCore/Private/DatasmithMaterialElementsImpl.h
@@ -676,6 +676,7 @@ public:
 	virtual IDatasmithExpressionInput& GetAmbientOcclusion() override { return *AmbientOcclusion.Edit(); }
 	virtual IDatasmithExpressionInput& GetClearCoat() override { return *ClearCoat.Edit(); }
 	virtual IDatasmithExpressionInput& GetClearCoatRoughness() override { return *ClearCoatRoughness.Edit(); }
+	virtual IDatasmithExpressionInput& GetWorldPositionOffset() override { return *WorldPositionOffset.Edit(); }
 	virtual IDatasmithExpressionInput& GetMaterialAttributes() override { return *MaterialAttributes.Edit(); }
 
 	virtual int GetBlendMode() const override {return BlendMode; }
@@ -725,6 +726,7 @@ protected:
 	TDatasmithReferenceProxy< FDatasmithExpressionInputImpl > AmbientOcclusion;
 	TDatasmithReferenceProxy< FDatasmithExpressionInputImpl > ClearCoat;
 	TDatasmithReferenceProxy< FDatasmithExpressionInputImpl > ClearCoatRoughness;
+	TDatasmithReferenceProxy< FDatasmithExpressionInputImpl > WorldPositionOffset;
 	TDatasmithReferenceProxy< FDatasmithExpressionInputImpl > MaterialAttributes;
 
 	TDatasmithReferenceArrayProxy< IDatasmithMaterialExpression > Expressions;
diff --git a/Engine/Source/Runtime/Datasmith/DatasmithCore/Private/DatasmithSceneXmlReader.cpp b/Engine/Source/Runtime/Datasmith/DatasmithCore/Private/DatasmithSceneXmlReader.cpp
--- a/Engine/Source/Runtime/Datasmith/DatasmithCore/Private/DatasmithSceneXmlReader.cpp
+++ b/Engine/Source/Runtime/Datasmith/DatasmithCore/Private/DatasmithSceneXmlReader.cpp
@@ -1918,6 +1918,7 @@ void FDatasmithSceneXmlReader::ParseUEPbrMaterial(FXmlNode* InNode, TSharedPtr<
 	TryConnectMaterialInput(OutElement->GetAmbientOcclusion());
 	TryConnectMaterialInput(OutElement->GetClearCoat());
 	TryConnectMaterialInput(OutElement->GetClearCoatRoughness());
+	TryConnectMaterialInput(OutElement->GetWorldPositionOffset());
 	TryConnectMaterialInput(OutElement->GetMaterialAttributes());
 
 	for ( const FXmlNode* ChildNode : InNode->GetChildrenNodes() )
diff --git a/Engine/Source/Runtime/Datasmith/DatasmithCore/Private/DatasmithSceneXmlWriter.cpp b/Engine/Source/Runtime/Datasmith/DatasmithCore/Private/DatasmithSceneXmlWriter.cpp
--- a/Engine/Source/Runtime/Datasmith/DatasmithCore/Private/DatasmithSceneXmlWriter.cpp
+++ b/Engine/Source/Runtime/Datasmith/DatasmithCore/Private/DatasmithSceneXmlWriter.cpp
@@ -1460,6 +1460,7 @@ void FDatasmithSceneXmlWriterImpl::WriteUEPbrMaterialElement( const TSharedRef<
 	WriteUEPbrMaterialExpressionInput( MaterialElement, MaterialElement->GetAmbientOcclusion(), Archive, Indent + 1 );
 	WriteUEPbrMaterialExpressionInput( MaterialElement, MaterialElement->GetClearCoat(), Archive, Indent + 1 );
 	WriteUEPbrMaterialExpressionInput( MaterialElement, MaterialElement->GetClearCoatRoughness(), Archive, Indent + 1 );
+	WriteUEPbrMaterialExpressionInput( MaterialElement, MaterialElement->GetWorldPositionOffset(), Archive, Indent + 1 );
 	WriteUEPbrMaterialExpressionInput( MaterialElement, MaterialElement->GetMaterialAttributes(), Archive, Indent + 1 );
 
 	if ( MaterialElement->GetUseMaterialAttributes() )
diff --git a/Engine/Source/Runtime/Datasmith/DatasmithCore/Public/DatasmithMaterialElements.h b/Engine/Source/Runtime/Datasmith/DatasmithCore/Public/DatasmithMaterialElements.h
--- a/Engine/Source/Runtime/Datasmith/DatasmithCore/Public/DatasmithMaterialElements.h
+++ b/Engine/Source/Runtime/Datasmith/DatasmithCore/Public/DatasmithMaterialElements.h
@@ -223,6 +223,7 @@ public:
 	virtual IDatasmithExpressionInput& GetAmbientOcclusion() = 0;
 	virtual IDatasmithExpressionInput& GetClearCoat() = 0;
 	virtual IDatasmithExpressionInput& GetClearCoatRoughness() = 0;
+	virtual IDatasmithExpressionInput& GetWorldPositionOffset() = 0;
 	virtual IDatasmithExpressionInput& GetMaterialAttributes() = 0;
 
 	/** InBlendMode must match the values of EBlendMode from EngineTypes.h */
