Datasmith Revit Exporter Instancing fixes (Also fixing SectionBox clipping issues)

diff --git a/Engine/Source/Programs/Enterprise/Datasmith/DatasmithRevitExporter/Private/DatasmithRevitDocumentData.cs b/Engine/Source/Programs/Enterprise/Datasmith/DatasmithRevitExporter/Private/DatasmithRevitDocumentData.cs
--- a/Engine/Source/Programs/Enterprise/Datasmith/DatasmithRevitExporter/Private/DatasmithRevitDocumentData.cs
+++ b/Engine/Source/Programs/Enterprise/Datasmith/DatasmithRevitExporter/Private/DatasmithRevitDocumentData.cs
@@ -15,6 +15,181 @@ using Autodesk.Revit.DB.Visual;
 
 namespace DatasmithRevitExporter
 {
+	public class OrientatedBoundingBox
+	{
+		int tlf = 0;
+		int trf = 1;
+		int blf = 2;
+		int brf = 3;
+		int tlb = 4;
+		int trb = 5;
+		int blb = 6;
+		int brb = 7;
+		
+		double Epsilon = 0.000001;
+
+		//  Z  
+		//  | Y
+		//  |/
+		//  o---X
+		//
+		//     tlb-------MAX
+		//     /|        /|
+		//    / |       / |
+		//  tlf-------trf |
+		//   |  |      |  |     
+		//   | blb-----|-brb
+		//   | /       | /
+		//   |/        |/
+		//  MIN-------brf
+		//
+		// Vertices[0] := tlf: Top Left Front
+		// Vertices[1] := trf: Top Right Front
+		// Vertices[2] := blf: Bottom Left Front	=> min
+		// Vertices[3] := brf: Bottom Right Front
+		// Vertices[4] := tlb: Top Left Back
+		// Vertices[5] := trb: Top Right Back		=> max
+		// Vertices[6] := blb: Bottom Left Back
+		// Vertices[7] := brb: Bottom Right Back
+
+		public XYZ[] Vertices = new XYZ[8]; //8 vertices // corners
+		
+		double SideXDistance;
+		double SideYDistance;
+		double SideZDistance;
+
+		public XYZ AxisAllignedMin;
+		public XYZ AxisAllignedMax;
+
+		Plane SideX0, SideX1;
+		Plane SideY0, SideY1;
+		Plane SideZ0, SideZ1;
+
+		public bool bIsValidData;
+
+		public OrientatedBoundingBox(Transform InTransform, XYZ MIN, XYZ MAX, bool InCalculatePlanes = false)
+		{
+			SideXDistance = (MAX.X - MIN.X) + Epsilon;
+			SideYDistance = (MAX.Y - MIN.Y) + Epsilon;
+			SideZDistance = (MAX.Z - MIN.Z) + Epsilon;
+
+			bIsValidData = SideXDistance > Epsilon && SideYDistance > Epsilon && SideZDistance > Epsilon;
+
+			if (!bIsValidData)
+			{
+				return;
+			}
+
+			Vertices[tlf] = new XYZ(MIN.X, MIN.Y, MAX.Z);
+			Vertices[trf] = new XYZ(MAX.X, MIN.Y, MAX.Z);
+			Vertices[brf] = new XYZ(MAX.X, MIN.Y, MIN.Z);
+			Vertices[blf] = MIN;
+			Vertices[tlb] = new XYZ(MIN.X, MAX.Y, MAX.Z);
+			Vertices[trb] = MAX;
+			Vertices[brb] = new XYZ(MAX.X, MAX.Y, MIN.Z);
+			Vertices[blb] = new XYZ(MIN.X, MAX.Y, MIN.Z);
+
+			double MinX = double.MaxValue;
+			double MinY = double.MaxValue;
+			double MinZ = double.MaxValue;
+			double MaxX = -double.MaxValue;
+			double MaxY = -double.MaxValue;
+			double MaxZ = -double.MaxValue;
+
+			for (int Index = 0; Index <= brb; Index++)
+			{
+				Vertices[Index] = InTransform.OfPoint(Vertices[Index]);
+
+				if (Vertices[Index].X < MinX) MinX = Vertices[Index].X;
+				if (Vertices[Index].Y < MinY) MinY = Vertices[Index].Y;
+				if (Vertices[Index].Z < MinZ) MinZ = Vertices[Index].Z;
+				if (Vertices[Index].X > MaxX) MaxX = Vertices[Index].X;
+				if (Vertices[Index].Y > MaxY) MaxY = Vertices[Index].Y;
+				if (Vertices[Index].Z > MaxZ) MaxZ = Vertices[Index].Z;
+			}
+
+			AxisAllignedMax = new XYZ(MaxX, MaxY, MaxZ);
+			AxisAllignedMin = new XYZ(MinX, MinY, MinZ);
+
+			if (InCalculatePlanes)
+			{
+				//We only need the planes for the SectionBox:
+				CalculatePlanes();
+			}
+		}
+
+		private void CalculatePlanes()
+		{
+			SideX0 = Plane.CreateByThreePoints(Vertices[tlf], Vertices[blb], Vertices[tlb]);
+			SideX1 = Plane.CreateByThreePoints(Vertices[trf], Vertices[trb], Vertices[brb]);
+
+			SideY0 = Plane.CreateByThreePoints(Vertices[tlf], Vertices[trf], Vertices[brf]);
+			SideY1 = Plane.CreateByThreePoints(Vertices[tlb], Vertices[brb], Vertices[trb]);
+
+			SideZ0 = Plane.CreateByThreePoints(Vertices[blf], Vertices[brb], Vertices[blb]);
+			SideZ1 = Plane.CreateByThreePoints(Vertices[tlf], Vertices[tlb], Vertices[trb]);
+		}
+
+		private bool IsPointWithin(Plane Plane1, Plane Plane2, double Distance, XYZ Point)
+		{
+			UV Uv;
+			double Distance1;
+			double Distance2;
+
+			Plane1.Project(Point, out Uv, out Distance1);
+
+			if (Distance1 > Distance)
+			{
+				return false;
+			}
+
+			Plane2.Project(Point, out Uv, out Distance2);
+
+			if (Distance2 > Distance)
+			{
+				return false;
+			}
+
+			return true;
+		}
+
+		private bool IsPointInside(XYZ Point)
+		{
+			return IsPointWithin(SideX0, SideX1, SideXDistance, Point)
+				&& IsPointWithin(SideY0, SideY1, SideYDistance, Point)
+				&& IsPointWithin(SideZ0, SideZ1, SideZDistance, Point);
+		}
+
+		// If self contains Candidate it returns false:
+		public bool DoesIntersect(OrientatedBoundingBox Candidate)
+		{
+			int VertexWithinCounter = 0;
+			int VertexOutsideCounter = 0;
+			for (int Index = 0; Index < Candidate.Vertices.Length; Index++)
+			{
+				if (IsPointInside(Candidate.Vertices[Index]))
+				{
+					VertexWithinCounter++;
+				}
+				else
+				{
+					VertexOutsideCounter++;
+				}
+
+				if (VertexOutsideCounter > 0 && VertexWithinCounter > 0)
+				{
+					return true;
+				}
+			}
+
+			//It returns true only if we have at least 1 vertex within and at least 1 vertex outside
+			//every other scenario should return false
+			//as in:
+			//	- all vertex outside := false
+			//	- all vertex inside := false => geometries with completely contained bounding boxes won't get clipped.
+			return false;
+		}
+	}
 	public class FDocumentData
 	{
 		// This class reflects the child -> super component relationship in Revit into the exported hierarchy (children under super components actors).
@@ -1100,7 +1275,8 @@ namespace DatasmithRevitExporter
 		public Document									CurrentDocument { get; private set; } = null;
 		public FDirectLink								DirectLink { get; private set; } = null;
 
-		public List<Outline>							SectionBoxOutlines = new List<Outline>();
+		public Outline									SectionBoxOutline = null;
+		private OrientatedBoundingBox					SectionBox = null;
 
 		public FDocumentData(
 			Document InDocument,
@@ -1131,13 +1307,25 @@ namespace DatasmithRevitExporter
 			// Cache document section boxes
 			if (CurrentDocument.ActiveView != null)
 			{
-				FilteredElementCollector Collector = new FilteredElementCollector(CurrentDocument, CurrentDocument.ActiveView.Id);
-				IList<Element> SectionBoxes = Collector.OfCategory(BuiltInCategory.OST_SectionBox).ToElements();
+				BoundingBoxXYZ BBox = (GetElement(CurrentDocument.ActiveView.Id) as View3D).GetSectionBox();
 
-				foreach (var SectionBox in SectionBoxes)
+				if (BBox.IsSet && BBox.Enabled)
 				{
-					BoundingBoxXYZ BBox = SectionBox.get_BoundingBox(CurrentDocument.ActiveView);
-					SectionBoxOutlines.Add(GetOutline(BBox.Transform, BBox));
+					SectionBox = new OrientatedBoundingBox(BBox.Transform, BBox.Min, BBox.Max, true);
+					if (!SectionBox.bIsValidData)
+					{
+						SectionBox = null;
+						SectionBoxOutline = null;
+					}
+					else
+					{
+						SectionBoxOutline = new Outline(SectionBox.AxisAllignedMin, SectionBox.AxisAllignedMax);
+					}
+				}
+				else
+				{
+					SectionBox = null;
+					SectionBoxOutline = null;
 				}
 			}
 
@@ -1147,24 +1335,6 @@ namespace DatasmithRevitExporter
 			}
 		}
 
-		private Outline GetOutline(Transform InTransform, BoundingBoxXYZ InBoundingBox)
-		{
-			XYZ A = InTransform.OfPoint(InBoundingBox.Min);
-			XYZ B = InTransform.OfPoint(InBoundingBox.Max);
-
-			XYZ PMin = new XYZ(
-					Math.Min(A.X, B.X),
-					Math.Min(A.Y, B.Y),
-					Math.Min(A.Z, B.Z));
-
-			XYZ PMax = new XYZ(
-					Math.Max(A.X, B.X),
-					Math.Max(A.Y, B.Y),
-					Math.Max(A.Z, B.Z));
-
-			return new Outline(PMin, PMax);
-		}
-
 		public Element GetElement(
 			ElementId InElementId
 		)
@@ -1354,21 +1524,18 @@ namespace DatasmithRevitExporter
 
 			bool bIntersectedBySectionBox = false;
 
-			if (SectionBoxOutlines.Count > 0)
+			if (SectionBox != null)
 			{
 				BoundingBoxXYZ InstanceBoundingBox = InInstanceType.get_BoundingBox(CurrentDocument.ActiveView);
-
+				
 				if (InstanceBoundingBox != null)
 				{
-					Outline InstanceOutline = GetOutline(InWorldTransform, InstanceBoundingBox);
-
-					foreach (Outline SectionBoxOutline in SectionBoxOutlines)
+					OrientatedBoundingBox InstanceOrientatedBoundingBox = new OrientatedBoundingBox(InWorldTransform, InstanceBoundingBox.Min, InstanceBoundingBox.Max);
+					if (InstanceOrientatedBoundingBox.bIsValidData)
 					{
-						bIntersectedBySectionBox = (SectionBoxOutline.Intersects(InstanceOutline, 0) != SectionBoxOutline.ContainsOtherOutline(InstanceOutline, 0));
-						if (bIntersectedBySectionBox)
-						{
-							break;
-						}
+						Outline InstanceOutline = new Outline(InstanceOrientatedBoundingBox.AxisAllignedMin, InstanceOrientatedBoundingBox.AxisAllignedMax);
+
+						bIntersectedBySectionBox = SectionBox.DoesIntersect(InstanceOrientatedBoundingBox);
 					}
 				}
 			}
