[3ds Max Export] Remove unused submaterial

- submaterials(of a composite material) that are not used on geometry are note added to datasmith scene

diff --git a/Engine/Source/Programs/Enterprise/Datasmith/DatasmithMaxExporter/Private/DatasmithMaxDirectLinkMaterial.cpp b/Engine/Source/Programs/Enterprise/Datasmith/DatasmithMaxExporter/Private/DatasmithMaxDirectLinkMaterial.cpp
--- a/Engine/Source/Programs/Enterprise/Datasmith/DatasmithMaxExporter/Private/DatasmithMaxDirectLinkMaterial.cpp
+++ b/Engine/Source/Programs/Enterprise/Datasmith/DatasmithMaxExporter/Private/DatasmithMaxDirectLinkMaterial.cpp
@@ -32,8 +32,33 @@ namespace DatasmithMaxDirectLink
 void FMaterialsCollectionTracker::AddActualMaterial(FMaterialTracker& MaterialTracker, Mtl* Material)
 {
 	// Record relationships between tracked material and actual materials used on geometry(e.g. submaterials of a tracked multisubobj material)
-	UsedMaterialToMaterialTracker.FindOrAdd(Material).Add(&MaterialTracker);
 	MaterialTracker.AddActualMaterial(Material);
+	TSet<FMaterialTracker*>& MaterialTrackersForMaterial = ActualMaterialToMaterialTracker.FindOrAdd(Material);
+	MaterialTrackersForMaterial.Add(&MaterialTracker);
+}
+
+void FMaterialsCollectionTracker::ReleaseActualMaterial(FMaterialTracker& MaterialTracker, Mtl* Material)
+{
+	TSet<FMaterialTracker*>* Found = ActualMaterialToMaterialTracker.Find(Material);
+	if (!Found)
+	{
+		return;
+	}
+
+	TSet<FMaterialTracker*>& MaterialTrackersForMaterial = *Found;
+	MaterialTrackersForMaterial.Remove(&MaterialTracker);
+	if (!MaterialTrackersForMaterial.Num())
+	{
+		//Clean up material when it's not used by other tracked materials
+		ActualMaterialToMaterialTracker.Remove(Material);
+
+		if (FString Name; UsedMaterialToDatasmithMaterialName.RemoveAndCopyValue(Material, Name))
+		{
+			MaterialNameProvider.RemoveExistingName(Name);
+		}
+
+		RemoveDatasmithMaterialForUsedMaterial(Material);
+	}
 }
 
 const TCHAR* FMaterialsCollectionTracker::GetMaterialName(Mtl* Material)
@@ -46,11 +71,98 @@ const TCHAR* FMaterialsCollectionTracker::GetMaterialName(Mtl* Material)
 		MaterialNameProvider.GenerateUniqueName(FDatasmithUtils::SanitizeObjectName(Material->GetName().data())));
 }
 
-void FMaterialsCollectionTracker::AssignMeshMaterials(const TSharedPtr<IDatasmithMeshElement>& MeshElement, Mtl* Material, const TSet<uint16>& SupportedChannels)
+bool FMaterialsCollectionTracker::IsMaterialUsed(Mtl* Material)
+{
+	return (UsedMaterialToMeshElementSlot.Contains(Material) || UsedMaterialToMeshActorSlot.Contains(Material));
+}
+
+void FMaterialsCollectionTracker::SetMaterialForMeshElement(const TSharedPtr<IDatasmithMeshElement>& MeshElement, Mtl* MaterialAssignedToNode, Mtl* ActualMaterial, uint16 SlotId)
+{
+	MeshElement->SetMaterial(GetMaterialName(ActualMaterial), SlotId);
+
+	if (!UsedMaterialToDatasmithMaterial.Contains(ActualMaterial))
+	{
+		// Invalidate material assigned directly to node to make it build its submaterials
+		// Rebuilding individual submaterials looks much more complex especially taking into account that Update can be cancelled/resumed
+		InvalidateMaterial(MaterialAssignedToNode);
+	}
+
+	UsedMaterialToMeshElementSlot.FindOrAdd(ActualMaterial).FindOrAdd(MeshElement.Get()).Add(SlotId);
+	MeshElementToUsedMaterialSlot.FindOrAdd(MeshElement.Get()).FindOrAdd(ActualMaterial).Add(SlotId);
+}
+
+void FMaterialsCollectionTracker::SetMaterialForMeshActorElement(
+	const TSharedPtr<IDatasmithMeshActorElement>& MeshActor, Mtl* MaterialAssignedToNode, Mtl* Material, int32 SlotId)
+{
+	MeshActor->AddMaterialOverride(GetMaterialName(Material), SlotId);
+
+	UsedMaterialToMeshActorSlot.FindOrAdd(Material).FindOrAdd(MeshActor.Get()).Add(SlotId);
+	MeshActorToUsedMaterialSlot.FindOrAdd(MeshActor.Get()).FindOrAdd(Material).Add(SlotId);
+}
+
+void FMaterialsCollectionTracker::UnSetMaterialsForMeshElement(const TSharedPtr<IDatasmithMeshElement>& MeshElement)
+{
+	if (!MeshElement)
+	{
+		return;
+	}
+
+	if (!MeshElementToUsedMaterialSlot.Contains(MeshElement.Get()))
+	{
+		return;
+	}
+
+	TArray<Mtl*> Materials;
+	MeshElementToUsedMaterialSlot[MeshElement.Get()].GetKeys(Materials);
+	for (Mtl* Material : Materials)
+	{
+		UsedMaterialToMeshElementSlot[Material].Remove(MeshElement.Get());
+		if (UsedMaterialToMeshElementSlot[Material].IsEmpty())
+		{
+			if (!IsMaterialUsed(Material))
+			{
+				RemoveDatasmithMaterialForUsedMaterial(Material);
+			}
+			UsedMaterialToMeshElementSlot.Remove(Material);
+		}
+	}
+	MeshElementToUsedMaterialSlot.Remove(MeshElement.Get());
+}
+
+void FMaterialsCollectionTracker::UnSetMaterialsForMeshActor(const TSharedPtr<IDatasmithMeshActorElement>& MeshActor)
+{
+	if (!MeshActor)
+	{
+		return;
+	}
+
+	if (!MeshActorToUsedMaterialSlot.Contains(MeshActor.Get()))
+	{
+		return;
+	}
+
+	TArray<Mtl*> Materials;
+	MeshActorToUsedMaterialSlot[MeshActor.Get()].GetKeys(Materials);
+	for (Mtl* Material : Materials)
+	{
+		UsedMaterialToMeshActorSlot[Material].Remove(MeshActor.Get());
+		if (UsedMaterialToMeshActorSlot[Material].IsEmpty())
+		{
+			if (!IsMaterialUsed(Material))
+			{
+				RemoveDatasmithMaterialForUsedMaterial(Material);
+			}
+			UsedMaterialToMeshActorSlot.Remove(Material);
+		}
+	}
+	MeshActorToUsedMaterialSlot.Remove(MeshActor.Get());
+}
+
+void FMaterialsCollectionTracker::SetMaterialsForMeshElement(const TSharedPtr<IDatasmithMeshElement>& MeshElement, Mtl* Material, const TSet<uint16>& SupportedChannels)
 {
 	if (FDatasmithMaxMatHelper::GetMaterialClass(Material) == EDSMaterialType::XRefMat)
 	{
-		AssignMeshMaterials(MeshElement, FDatasmithMaxMatHelper::GetRenderedXRefMaterial(Material), SupportedChannels);
+		SetMaterialsForMeshElement(MeshElement, FDatasmithMaxMatHelper::GetRenderedXRefMaterial(Material), SupportedChannels);
 		return;
 	}
 	if (Material != nullptr)
@@ -70,32 +182,32 @@ void FMaterialsCollectionTracker::AssignMeshMaterials(const TSharedPtr<IDatasmit
 				{
 					if (FDatasmithMaxMatHelper::GetMaterialClass(SubMaterial) == EDSMaterialType::TheaRandom)
 					{
-						MeshElement->SetMaterial(*FDatasmithMaxSceneExporter::GetRandomSubMaterial(SubMaterial, MeshElement->GetDimensions()), DisplayedChannelID);
+						SetMaterialForMeshElement(MeshElement, Material, FDatasmithMaxSceneExporter::GetRandomSubMaterial(SubMaterial, MeshElement->GetDimensions()), DisplayedChannelID);
 					}
 					else
 					{
-						MeshElement->SetMaterial(GetMaterialName(SubMaterial), DisplayedChannelID);
+						SetMaterialForMeshElement(MeshElement, Material, SubMaterial, DisplayedChannelID);
 					}
 				}
 			}
 			else if (FDatasmithMaxMatHelper::GetMaterialClass(Material) == EDSMaterialType::TheaRandom)
 			{
-				MeshElement->SetMaterial(*FDatasmithMaxSceneExporter::GetRandomSubMaterial(Material, MeshElement->GetDimensions()), DisplayedChannelID);
+				SetMaterialForMeshElement(MeshElement, Material, FDatasmithMaxSceneExporter::GetRandomSubMaterial(Material, MeshElement->GetDimensions()), DisplayedChannelID);
 			}
 			// Single material
 			else
 			{
-				MeshElement->SetMaterial(GetMaterialName(Material), DisplayedChannelID);
+				SetMaterialForMeshElement(MeshElement, Material, Material, DisplayedChannelID);
 			}
 		}
 	}
 }
 
-void FMaterialsCollectionTracker::AssignMeshActorMaterials(const TSharedPtr<IDatasmithMeshActorElement>& MeshActor, Mtl* Material, TSet<uint16>& SupportedChannels, const FVector3f& RandomSeed)
+void FMaterialsCollectionTracker::SetMaterialsForMeshActor(const TSharedPtr<IDatasmithMeshActorElement>& MeshActor, Mtl* Material, TSet<uint16>& SupportedChannels, const FVector3f& RandomSeed)
 {
 	if (FDatasmithMaxMatHelper::GetMaterialClass(Material) == EDSMaterialType::XRefMat)
 	{
-		AssignMeshActorMaterials(MeshActor, FDatasmithMaxMatHelper::GetRenderedXRefMaterial(Material), SupportedChannels, RandomSeed);
+		SetMaterialsForMeshActor(MeshActor, FDatasmithMaxMatHelper::GetRenderedXRefMaterial(Material), SupportedChannels, RandomSeed);
 		return;
 	}
 
@@ -107,11 +219,11 @@ void FMaterialsCollectionTracker::AssignMeshActorMaterials(const TSharedPtr<IDat
 			{
 				if (FDatasmithMaxMatHelper::GetMaterialClass(Material) != EDSMaterialType::TheaRandom)
 				{
-					MeshActor->AddMaterialOverride(GetMaterialName(Material), -1);
+					SetMaterialForMeshActorElement(MeshActor, Material, Material, -1);
 				}
 				else
 				{
-					MeshActor->AddMaterialOverride( *FDatasmithMaxSceneExporter::GetRandomSubMaterial(Material, RandomSeed), -1 );
+					SetMaterialForMeshActorElement(MeshActor, Material, FDatasmithMaxSceneExporter::GetRandomSubMaterial(Material, RandomSeed), -1 );
 				}
 			}
 			else
@@ -130,11 +242,11 @@ void FMaterialsCollectionTracker::AssignMeshActorMaterials(const TSharedPtr<IDat
 				{
 					if (FDatasmithMaxMatHelper::GetMaterialClass(SubMaterial) != EDSMaterialType::TheaRandom)
 					{
-						MeshActor->AddMaterialOverride(GetMaterialName(SubMaterial), -1);
+						SetMaterialForMeshActorElement(MeshActor, Material, SubMaterial, -1);
 					}
 					else
 					{
-						MeshActor->AddMaterialOverride( *FDatasmithMaxSceneExporter::GetRandomSubMaterial(SubMaterial, RandomSeed), -1 );
+						SetMaterialForMeshActorElement(MeshActor, Material, FDatasmithMaxSceneExporter::GetRandomSubMaterial(Material, RandomSeed), -1 );
 					}
 				}
 			}
@@ -149,11 +261,11 @@ void FMaterialsCollectionTracker::AssignMeshActorMaterials(const TSharedPtr<IDat
 				{
 					if (FDatasmithMaxMatHelper::GetMaterialClass(Material) != EDSMaterialType::TheaRandom)
 					{
-						MeshActor->AddMaterialOverride(GetMaterialName(Material), ActualSubObj);
+						SetMaterialForMeshActorElement(MeshActor, Material, Material, ActualSubObj);
 					}
 					else
 					{
-						MeshActor->AddMaterialOverride( *FDatasmithMaxSceneExporter::GetRandomSubMaterial(Material, RandomSeed), ActualSubObj );
+						SetMaterialForMeshActorElement(MeshActor, Material, FDatasmithMaxSceneExporter::GetRandomSubMaterial(Material, RandomSeed), -1);
 					}
 				}
 				else
@@ -167,11 +279,11 @@ void FMaterialsCollectionTracker::AssignMeshActorMaterials(const TSharedPtr<IDat
 						if (FDatasmithMaxMatHelper::GetMaterialClass(SubMaterial) != EDSMaterialType::TheaRandom)
 						{
 							//Material slots in Max are not zero-based, so we serialize our SlotID starting from 1 for better visual consistency.
-							MeshActor->AddMaterialOverride(GetMaterialName(SubMaterial), Mid + 1);
+							SetMaterialForMeshActorElement(MeshActor, Material, SubMaterial, Mid + 1);
 						}
 						else
 						{
-							MeshActor->AddMaterialOverride( *FDatasmithMaxSceneExporter::GetRandomSubMaterial(SubMaterial, RandomSeed), MaterialIndex + 1 );
+							SetMaterialForMeshActorElement(MeshActor, Material, FDatasmithMaxSceneExporter::GetRandomSubMaterial(SubMaterial, RandomSeed), MaterialIndex + 1);
 						}
 					}
 				}
@@ -233,13 +345,16 @@ void FMaterialsCollectionTracker::Reset()
 	MaterialTrackers.Reset();
 	InvalidatedMaterialTrackers.Reset();
 
-	EncounteredMaterials.Reset();
-	EncounteredTextures.Reset();
-	MaterialNames.Reset();
+	ActualMaterialToMaterialTracker.Reset();
 
-	UsedMaterialToMaterialTracker.Reset();
 	UsedMaterialToDatasmithMaterial.Reset();
 	UsedMaterialToDatasmithMaterialName.Reset();
+	UsedMaterialToMeshElementSlot.Reset();
+	MeshElementToUsedMaterialSlot.Reset();
+
+	UsedMaterialToMeshActorSlot.Reset();
+	MeshActorToUsedMaterialSlot.Reset();
+
 	MaterialNameProvider.Clear();
 
 	UsedTextureToMaterialTracker.Reset();
@@ -268,6 +383,8 @@ void FMaterialsCollectionTracker::AddDatasmithMaterialForUsedMaterial(TSharedRef
 
 void FMaterialsCollectionTracker::ConvertMaterial(Mtl* Material, TSharedRef<IDatasmithScene> DatasmithScene, const TCHAR* AssetsPath, TSet<Texmap*>& TexmapsConverted)
 {
+	SCENE_UPDATE_STAT_INC(UpdateMaterials, Total);
+
 	if (UsedMaterialToDatasmithMaterial.Contains(Material))
 	{
 		// Material might have been already converted - if it's present in UsedMaterialToDatasmithMaterial this means that it(or multisubobj it's part of) wasn't changed
@@ -275,7 +392,12 @@ void FMaterialsCollectionTracker::ConvertMaterial(Mtl* Material, TSharedRef<IDat
 		return;
 	}
 
-	SCENE_UPDATE_STAT_INC(UpdateMaterials, Total);
+	if (!IsMaterialUsed(Material))
+	{
+		// Don't convert materials that don't have visible geometry exported
+		// This happens when a submaterial of MultiSubobj material is unused on a mesh
+		return;
+	}
 
 	TSet<Texmap*> TexmapsUsedByMaterial;
 	FMaterialConversionContext MaterialConversionContext = {TexmapsUsedByMaterial, *this};
@@ -293,7 +415,7 @@ void FMaterialsCollectionTracker::ConvertMaterial(Mtl* Material, TSharedRef<IDat
 	// Tie texture used by an actual material to tracked material
 	for (Texmap* Tex : TexmapsUsedByMaterial)
 	{
-		for (FMaterialTracker* MaterialTracker : UsedMaterialToMaterialTracker[Material])
+		for (FMaterialTracker* MaterialTracker : ActualMaterialToMaterialTracker[Material])
 		{
 			MaterialTracker->AddActualTexture(Tex);
 			UsedTextureToMaterialTracker.FindOrAdd(Tex).Add(MaterialTracker);
@@ -311,28 +433,20 @@ void FMaterialsCollectionTracker::ReleaseMaterial(FMaterialTracker& MaterialTrac
 	InvalidatedMaterialTrackers.Remove(&MaterialTracker);
 }
 
+void FMaterialsCollectionTracker::RemoveDatasmithMaterialForUsedMaterial(Mtl* Material)
+{
+	TSharedPtr<IDatasmithBaseMaterialElement> DatasmithMaterial;
+	if(UsedMaterialToDatasmithMaterial.RemoveAndCopyValue(Material, DatasmithMaterial))
+	{
+		SceneTracker.RemoveMaterial(DatasmithMaterial);
+	}
+}
 
 void FMaterialsCollectionTracker::RemoveConvertedMaterial(FMaterialTracker& MaterialTracker)
 {
 	for (Mtl* Material: MaterialTracker.GetActualMaterials())
 	{
-		TSet<FMaterialTracker*>& MaterialTrackersForMaterial = UsedMaterialToMaterialTracker[Material];
-		MaterialTrackersForMaterial.Remove(&MaterialTracker);
-		if (!MaterialTrackersForMaterial.Num())
-		{
-			UsedMaterialToMaterialTracker.Remove(Material);
-
-			if (FString Name; UsedMaterialToDatasmithMaterialName.RemoveAndCopyValue(Material, Name))
-			{
-				MaterialNameProvider.RemoveExistingName(Name);
-			}
-
-			TSharedPtr<IDatasmithBaseMaterialElement> DatasmithMaterial;
-			if(UsedMaterialToDatasmithMaterial.RemoveAndCopyValue(Material, DatasmithMaterial))
-			{
-				SceneTracker.RemoveMaterial(DatasmithMaterial);
-			}
-		}
+		ReleaseActualMaterial(MaterialTracker, Material);
 	}
 
 	for (Texmap* Tex: MaterialTracker.GetActualTexmaps())
diff --git a/Engine/Source/Programs/Enterprise/Datasmith/DatasmithMaxExporter/Private/DatasmithMaxExporterWithDirectLink.cpp b/Engine/Source/Programs/Enterprise/Datasmith/DatasmithMaxExporter/Private/DatasmithMaxExporterWithDirectLink.cpp
--- a/Engine/Source/Programs/Enterprise/Datasmith/DatasmithMaxExporter/Private/DatasmithMaxExporterWithDirectLink.cpp
+++ b/Engine/Source/Programs/Enterprise/Datasmith/DatasmithMaxExporter/Private/DatasmithMaxExporterWithDirectLink.cpp
@@ -181,7 +181,7 @@ struct FInstances: FNoncopyable
 	void AssignMaterialToStaticMesh(FMaterialsCollectionTracker& MaterialsCollectionTracker,  Mtl* InMaterial)
 	{
 		Material = InMaterial;
-		MaterialsCollectionTracker.AssignMeshMaterials(Converted.DatasmithMeshElement, Material, Converted.SupportedChannels);
+		MaterialsCollectionTracker.SetMaterialsForMeshElement(Converted.DatasmithMeshElement, Material, Converted.SupportedChannels);
 	}
 };
 
@@ -1365,7 +1365,6 @@ public:
 	 // returns if any change in scene was encountered and scene update completed(i.e. DirectLink Sync can be run)
 	bool UpdateInternal(FUpdateProgress::FStage& MainStage)
 	{
-
 		CurrentSyncPoint.Time = GetCOREInterface()->GetTime();
 
 		// Changes present only when there are modified layers(changes checked manually), nodes(notified by Max) or materials(notified by Max with all changes in dependencies)
@@ -2108,10 +2107,14 @@ public:
 			// Remove mesh actor before removing its parent Actor in case there a separate MeshActor
 			if (Converted.DatasmithMeshActor)
 			{
-				// if (NodeTracker.DatasmithMeshActor != NodeTracker.DatasmithActorElement)
+				if (Converted.DatasmithActorElement != Converted.DatasmithMeshActor)
 				{
 					Converted.DatasmithActorElement->RemoveChild(Converted.DatasmithMeshActor);
 				}
+
+				MaterialsCollectionTracker.UnSetMaterialsForMeshActor(Converted.DatasmithMeshActor);
+
+
 				Converted.DatasmithMeshActor.Reset();
 				// todo: consider pool of MeshActors
 			}
@@ -2142,6 +2145,7 @@ public:
 	// Called when mesh element is not needed anymore and should be removed from the scene
 	virtual void ReleaseMeshElement(FMeshConverted& Converted) override
 	{
+		MaterialsCollectionTracker.UnSetMaterialsForMeshElement(Converted.DatasmithMeshElement);
 		GetDatasmithScene().RemoveMesh(Converted.DatasmithMeshElement);
 		Converted.ReleaseMeshConverted();
 	}
@@ -2621,11 +2625,11 @@ public:
 		NodeTracker.Validity.NarrowValidityToInterval(ValidityInterval);
 	}
 
-	void RegisterNodeForMaterial(FNodeTracker& NodeTracker, Mtl* Material)
+	FMaterialTracker* RegisterNodeForMaterial(FNodeTracker& NodeTracker, Mtl* Material)
 	{
 		if (!Material)
 		{
-			return;
+			return nullptr;
 		}
 
 		FMaterialTracker* MaterialTracker = MaterialsCollectionTracker.AddMaterial(Material);
@@ -2636,6 +2640,7 @@ public:
 		{
 			NotificationsHandler->AddMaterial(Material);
 		}
+		return MaterialTracker;
 	}
 
 	virtual void UnregisterNodeForMaterial(FNodeTracker& NodeTracker) override 
@@ -2702,7 +2707,7 @@ public:
 			return;
 		}
 
-		TSet<FMaterialTracker*>* MaterialTrackers = MaterialsCollectionTracker.UsedMaterialToMaterialTracker.Find(ActualMaterial);
+		TSet<FMaterialTracker*>* MaterialTrackers = MaterialsCollectionTracker.ActualMaterialToMaterialTracker.Find(ActualMaterial);
 
 		if (!MaterialTrackers)
 		{
@@ -2744,7 +2749,6 @@ public:
 	virtual void AddGeometryNodeInstance(FNodeTracker& NodeTracker, FMeshNodeConverter& MeshConverter, Object* Obj) override
 	{
 		InvalidateInstances(InstancesManager.AddNodeTracker(NodeTracker, MeshConverter, Obj));
-		
 	}
 
 	virtual void RemoveGeometryNodeInstance(FNodeTracker& NodeTracker) override
@@ -2860,7 +2864,7 @@ public:
 			{
 				if (Instances.Material != Material)
 				{
-					MaterialsCollectionTracker.AssignMeshActorMaterials(NodeTracker.GetConverted().DatasmithMeshActor, Material, Instances.Converted.SupportedChannels, FVector3f(NodeTracker.GetConverted().DatasmithMeshActor->GetTranslation()));
+					MaterialsCollectionTracker.SetMaterialsForMeshActor(NodeTracker.GetConverted().DatasmithMeshActor, Material, Instances.Converted.SupportedChannels, FVector3f(NodeTracker.GetConverted().DatasmithMeshActor->GetTranslation()));
 				}
 			}
 		}
@@ -2992,7 +2996,7 @@ public:
 				NodeConverter.Meshes.Add(MeshConverted);
 
 				RegisterNodeForMaterial(NodeTracker, Material);
-				MaterialsCollectionTracker.AssignMeshMaterials(MeshConverted.DatasmithMeshElement, Material, MeshConverted.SupportedChannels);
+				MaterialsCollectionTracker.SetMaterialsForMeshElement(MeshConverted.DatasmithMeshElement, Material, MeshConverted.SupportedChannels);
 
 				FString MeshLabel = NodeTrackersNames.GetNodeName(NodeTracker) + (TEXT("_") + FString::FromInt(MeshIndex));
 				MeshConverted.DatasmithMeshElement->SetLabel(*MeshLabel);
@@ -3002,7 +3006,7 @@ public:
 				// todo: override material
 				TSharedPtr< IDatasmithActorElement > InversedHISMActor;
 				// todo: ExportHierarchicalInstanceStaticMeshActor CustomMeshNode only used for Material - can be simplified, Material anyway is dealt with outside too
-				TSharedRef<IDatasmithActorElement> HismActorElement = FDatasmithMaxSceneExporter::ExportHierarchicalInstanceStaticMeshActor(NodeTracker.Node, MeshSource.Node, *MeshLabel, MeshConverted.SupportedChannels,
+				TSharedRef<IDatasmithActorElement> HismActorElement = ExportHierarchicalInstanceStaticMeshActor(NodeTracker.Node, MeshSource.Node, *MeshLabel, MeshConverted.SupportedChannels,
 					Material, &Transforms, *MeshName, Converter.UnitToCentimeter, EStaticMeshExportMode::Default, InversedHISMActor);
 				NodeTracker.GetConverted().DatasmithActorElement->AddChild(HismActorElement, EDatasmithActorAttachmentRule::KeepWorldTransform);
 				if (InversedHISMActor)
@@ -3014,6 +3018,136 @@ public:
 		});
 	}
 
+	TSharedRef< IDatasmithActorElement > ExportHierarchicalInstanceStaticMeshActor(INode* Node, INode* CustomMeshNode, const TCHAR* Label, TSet<uint16>& SupportedChannels, Mtl* StaticMeshMtl, const TArray<Matrix3>* Instances,
+		const TCHAR* MeshName, float UnitMultiplier, const EStaticMeshExportMode& ExportMode, TSharedPtr< IDatasmithActorElement >& OutInversedHISMActor)
+	{
+		check(Node && Instances && MeshName);
+
+		TSharedPtr< IDatasmithMeshActorElement > MeshActor;
+		FVector Pos, Scale;
+		FQuat Rotation;
+
+		const FVector3f RandomSeed(FMath::Rand(), FMath::Rand(), FMath::Rand());
+		auto FinalizeHISMActor = [&](TSharedPtr< IDatasmithMeshActorElement >& FinalizingActor, FVector3f Seed, const TCHAR* ActorLabel)
+		{
+			FinalizingActor->SetStaticMeshPathName(MeshName);
+
+			INode* MeshNode = CustomMeshNode ? CustomMeshNode : Node;
+
+			if (ExportMode == EStaticMeshExportMode::Default && StaticMeshMtl != MeshNode->GetMtl())
+			{
+				TSharedRef< IDatasmithMeshActorElement > MeshActorRef = FinalizingActor.ToSharedRef();
+				MaterialsCollectionTracker.SetMaterialsForMeshActor(MeshActorRef, StaticMeshMtl, SupportedChannels, Seed);
+			}
+
+			if (ActorLabel)
+			{
+				FinalizingActor->SetLabel(ActorLabel);
+			}
+
+			FinalizingActor->SetIsAComponent(true);
+
+			switch (ExportMode)
+			{
+			case EStaticMeshExportMode::BoundingBox:
+				FinalizingActor->AddTag(TEXT("Datasmith.Attributes.Geometry: BoundingBox"));
+				break;
+			default:
+				break;
+			}
+		};
+
+		if (Node->GetWSMDerivedObject() != nullptr)
+		{
+			FDatasmithMaxSceneExporter::MaxToUnrealCoordinates(Node->GetObjTMAfterWSM(GetCOREInterface()->GetTime()), Pos, Rotation, Scale, UnitMultiplier);
+		}
+		else
+		{
+			FDatasmithMaxSceneExporter::MaxToUnrealCoordinates(Node->GetObjectTM(GetCOREInterface()->GetTime()), Pos, Rotation, Scale, UnitMultiplier);
+		}
+
+		const int32 InstancesCount = Instances->Num();
+		if ( InstancesCount == 1 )
+		{
+			// Export the the hism as a normal static mesh
+			MeshActor = FDatasmithSceneFactory::CreateMeshActor(MeshName);
+
+			// Apply the relative transfrom of the instance to the forest world transform
+			FTransform WithoutInstance(Rotation, Pos, Scale);
+			FVector InstancePos, InstanceScale;
+			FQuat InstanceRotation;
+			FDatasmithMaxSceneExporter::MaxToUnrealCoordinates((*Instances)[0], InstancePos, InstanceRotation, InstanceScale, UnitMultiplier);
+			FTransform WithInstance(FTransform(InstanceRotation, InstancePos, InstanceScale) * WithoutInstance);
+
+			MeshActor->SetTranslation(WithInstance.GetLocation());
+			MeshActor->SetScale(WithInstance.GetScale3D());
+			MeshActor->SetRotation(WithInstance.GetRotation());
+		}
+		else
+		{
+			TSharedRef< IDatasmithHierarchicalInstancedStaticMeshActorElement > HierarchicalInstanceStaticMeshActor = FDatasmithSceneFactory::CreateHierarchicalInstanceStaticMeshActor(MeshName);
+			MeshActor = HierarchicalInstanceStaticMeshActor;
+
+			FTransform HISMActorTransform(Rotation, Pos, Scale);
+			MeshActor->SetTranslation(Pos);
+			MeshActor->SetScale(Scale);
+			MeshActor->SetRotation(Rotation);
+
+			TArray< FTransform > InstancesTransform;
+			InstancesTransform.Reserve(InstancesCount);
+			for (int32 i = 0; i < InstancesCount; i++)
+			{
+				FDatasmithMaxSceneExporter::MaxToUnrealCoordinates((*Instances)[i], Pos, Rotation, Scale, UnitMultiplier);
+				InstancesTransform.Emplace(Rotation, Pos, Scale);
+			}
+
+			TArray< FTransform > NonInvertedTransforms, InvertedTransforms;
+			DatasmithMaxHelper::FilterInvertedScaleTransforms(InstancesTransform, NonInvertedTransforms, InvertedTransforms);
+
+			HierarchicalInstanceStaticMeshActor->ReserveSpaceForInstances(NonInvertedTransforms.Num());
+			for (const FTransform& InstanceTransform : NonInvertedTransforms)
+			{
+				HierarchicalInstanceStaticMeshActor->AddInstance(InstanceTransform);
+			}
+
+			//Adding an inverted HierarchicalInstancedStaticMeshActorElement for the instances with an inverted mesh due to negative scaling.
+			if (InvertedTransforms.Num() > 0)
+			{
+				FString InvertedHISMName = FString(MeshName).Append("_inv");
+				TSharedPtr< IDatasmithHierarchicalInstancedStaticMeshActorElement > InvertedHierarchicalInstanceStaticMeshActor = FDatasmithSceneFactory::CreateHierarchicalInstanceStaticMeshActor(*InvertedHISMName);
+				OutInversedHISMActor = InvertedHierarchicalInstanceStaticMeshActor;
+
+				OutInversedHISMActor->SetTranslation( HISMActorTransform.GetTranslation() );
+				OutInversedHISMActor->SetScale( -1 * HISMActorTransform.GetScale3D() );
+				OutInversedHISMActor->SetRotation( HISMActorTransform.GetRotation() );
+				
+				InvertedHierarchicalInstanceStaticMeshActor->ReserveSpaceForInstances(InvertedTransforms.Num());
+				for (const FTransform& InstanceTransform : InvertedTransforms)
+				{
+					//Correcting the instance transform to reverse the negative scaling of the parent.
+					InvertedHierarchicalInstanceStaticMeshActor->AddInstance( FTransform( InstanceTransform.GetRotation(), -1 * InstanceTransform.GetTranslation(), -1 * InstanceTransform.GetScale3D() ) );
+				}
+
+				TSharedPtr< IDatasmithMeshActorElement > InvertedMeshActor = InvertedHierarchicalInstanceStaticMeshActor;
+				if (Label)
+				{
+					FString InversedLabelString = FString(Label).Append("_inv");
+					FinalizeHISMActor(InvertedMeshActor, RandomSeed, *InversedLabelString);
+				}
+				else
+				{
+					FinalizeHISMActor(InvertedMeshActor, RandomSeed, nullptr);
+				}
+			}
+		}
+
+		FinalizeHISMActor(MeshActor, RandomSeed, Label);
+
+		return MeshActor.ToSharedRef();
+	}
+
+
+
 	void ConvertNamedCollisionNode(FNodeTracker& NodeTracker)
 	{
 		// Split collision prefix and find node that might use this node as collision mesh
diff --git a/Engine/Source/Programs/Enterprise/Datasmith/DatasmithMaxExporter/Private/DatasmithMaxMeshExporter.cpp b/Engine/Source/Programs/Enterprise/Datasmith/DatasmithMaxExporter/Private/DatasmithMaxMeshExporter.cpp
--- a/Engine/Source/Programs/Enterprise/Datasmith/DatasmithMaxExporter/Private/DatasmithMaxMeshExporter.cpp
+++ b/Engine/Source/Programs/Enterprise/Datasmith/DatasmithMaxExporter/Private/DatasmithMaxMeshExporter.cpp
@@ -323,177 +323,6 @@ FVector FDatasmithMaxMeshExporter::Point3ToFVector(Point3 Point) const
 					 Point.z );
 }
 
-void AssignMeshMaterials(TSharedPtr<IDatasmithMeshElement>& MeshElement, Mtl* Material, const TSet<uint16>& SupportedChannels)
-{
-	if (FDatasmithMaxMatHelper::GetMaterialClass(Material) == EDSMaterialType::XRefMat)
-	{
-		AssignMeshMaterials(MeshElement, FDatasmithMaxMatHelper::GetRenderedXRefMaterial(Material), SupportedChannels);
-		return;
-	}
-	if (Material != nullptr)
-	{
-		TArray<uint16> ChannelsSorted = SupportedChannels.Array();
-		ChannelsSorted.Sort();
-		for (uint16 Channel : ChannelsSorted)
-		{
-			//Max's channel UI is not zero-based, so we register an incremented ChannelID for better visual consistency after importing in Unreal.
-			uint16 DisplayedChannelID = Channel + 1;
-
-			// Multi mat
-			if (FDatasmithMaxMatHelper::GetMaterialClass(Material) == EDSMaterialType::MultiMat)
-			{
-				// Replicate the 3ds max behavior where material ids greater than the number of sub-materials wrap around and are mapped to the existing sub-materials
-				if ( Mtl* SubMaterial = Material->GetSubMtl(Channel % Material->NumSubMtls()) )
-				{
-					if (FDatasmithMaxMatHelper::GetMaterialClass(SubMaterial) == EDSMaterialType::TheaRandom)
-					{
-						MeshElement->SetMaterial(*FDatasmithMaxSceneExporter::GetRandomSubMaterial(SubMaterial, MeshElement->GetDimensions()), DisplayedChannelID);
-					}
-					else
-					{
-						MeshElement->SetMaterial(SubMaterial->GetName().data(), DisplayedChannelID);
-					}
-				}
-			}
-			else if (FDatasmithMaxMatHelper::GetMaterialClass(Material) == EDSMaterialType::TheaRandom)
-			{
-				MeshElement->SetMaterial(*FDatasmithMaxSceneExporter::GetRandomSubMaterial(Material, MeshElement->GetDimensions()), DisplayedChannelID);
-			}
-			// Single material
-			else
-			{
-				MeshElement->SetMaterial(Material->GetName().data(), DisplayedChannelID);
-			}
-		}
-	}
-}
-
-TSharedPtr< IDatasmithMeshElement > FDatasmithMaxMeshExporter::ExportMesh(FMaxExportMeshArgs& ExportMeshArgs, TSet<uint16>& SupportedChannels, FString& Error)
-{
-	TSharedPtr< IDatasmithMeshElement > MeshElement;
-	UnitMultiplier = ExportMeshArgs.UnitMultiplier;
-
-	if (ExportMeshArgs.Node == NULL)
-	{
-		return MeshElement;
-	}
-
-	FDatasmithMesh DSMesh;
-
-	if (ExportMeshArgs.DatasmithAttributes)
-	{
-		if (ExportMeshArgs.DatasmithAttributes->GetExportMode() == EStaticMeshExportMode::BoundingBox)
-		{
-			if (FillDatasmithMeshFromBoundingBox(DSMesh, ExportMeshArgs.Node, ExportMeshArgs.Pivot))
-			{
-				FDatasmithMeshExporter DatasmithMeshExporter;
-				return DatasmithMeshExporter.ExportToUObject(ExportMeshArgs.ExportPath, ExportMeshArgs.ExportName, DSMesh, nullptr, FDatasmithExportOptions::LightmapUV);
-			}
-			else
-			{
-				return MeshElement;
-			}
-		}
-	}
-
-	if (!CreateDatasmithMesh(DSMesh, ExportMeshArgs, SupportedChannels))
-	{
-		return MeshElement;
-	}
-
-	TUniquePtr< FDatasmithMesh > DSCollisionMesh;
-	bool bIsCollisionFromDatasmithAttributes;
-	INode* CollisionNode = GetCollisionNode(ExportMeshArgs.Node, ExportMeshArgs.DatasmithAttributes, bIsCollisionFromDatasmithAttributes);
-	if (CollisionNode)
-	{
-
-		FMaxExportMeshArgs CollisionExportArgs;
-		CollisionExportArgs.Node = CollisionNode;
-		CollisionExportArgs.bForceSingleMat = ExportMeshArgs.bForceSingleMat;
-		CollisionExportArgs.UnitMultiplier = ExportMeshArgs.UnitMultiplier;
-		TSet<uint16> CollisionSupportedChannels;
-
-		DSCollisionMesh = MakeUnique< FDatasmithMesh >();
-		FTransform ColliderPivot = FDatasmithMaxSceneExporter::GetPivotTransform(CollisionNode, ExportMeshArgs.UnitMultiplier);
-
-		if (bIsCollisionFromDatasmithAttributes)
-		{
-			if (!ExportMeshArgs.bBakePivot)
-			{
-				FTransform RealPivot = FDatasmithMaxSceneExporter::GetPivotTransform(ExportMeshArgs.Node, ExportMeshArgs.UnitMultiplier);
-				ColliderPivot = ColliderPivot * RealPivot.Inverse();
-			}
-			CollisionExportArgs.Pivot = ColliderPivot;
-		}
-		else
-		{
-
-			TimeValue Now = GetCOREInterface()->GetTime();
-			FTransform NodeWTM = FTransformFromMatrix3(ExportMeshArgs.Node->GetNodeTM(Now), ExportMeshArgs.UnitMultiplier);
-			FTransform ColliderNodeWTM = FTransformFromMatrix3(CollisionNode->GetNodeTM(Now), ExportMeshArgs.UnitMultiplier);
-
-			// if object-offset has been baked into the mesh data, we want collision mesh data in the mesh's node space
-			//   MeshVert_Node = RealPivot * max_vert_data
-			//   MeshVert_world = NodeWTM * MeshVert_Node
-			// Collision mesh vertices in world space:
-			//   CollVert_node = ColliderPivot * CollVert_obj
-			//   CollVert_world = ColliderNodeWTM * CollVert_node
-			//   CollVert_mesh = NodeWTM-1 * CollVert_world
-			FTransform BakedTransform = ColliderPivot * ColliderNodeWTM * NodeWTM.Inverse();
-
-			if (!ExportMeshArgs.bBakePivot)
-			{
-				// if object-offset has not been baked, we want collision mesh data in the mesh's object space
-				FTransform RealPivot = FDatasmithMaxSceneExporter::GetPivotTransform(ExportMeshArgs.Node, ExportMeshArgs.UnitMultiplier);
-				BakedTransform = BakedTransform * RealPivot.Inverse();
-			}
-
-			CollisionExportArgs.Pivot = BakedTransform;
-		}
-
-		CreateDatasmithMesh(*DSCollisionMesh, CollisionExportArgs, CollisionSupportedChannels);
-	}
-
-	FDatasmithMeshExporter DatasmithMeshExporter;
-	MeshElement = DatasmithMeshExporter.ExportToUObject(ExportMeshArgs.ExportPath, ExportMeshArgs.ExportName, DSMesh, DSCollisionMesh.Get(), FDatasmithExportOptions::LightmapUV );
-
-	if ( !MeshElement.IsValid() )
-	{
-		Error = DatasmithMeshExporter.GetLastError();
-	}
-	else
-	{
-		AssignMeshMaterials(MeshElement, ExportMeshArgs.Node->GetMtl(), SupportedChannels);
-
-		if (ExportMeshArgs.ExportName)
-		{
-			// Mapping between the 3ds max channel and the exported mesh channel
-			const FUVChannelsMap* UVChannelMap = MeshNamesToUVChannels.Find(ExportMeshArgs.ExportName);
-			if (ExportMeshArgs.DatasmithAttributes && UVChannelMap)
-			{
-				constexpr int32 MaxToUnrealUVOffset = -1;
-				constexpr int32 DefaultValue = -1;
-				const int32 SelectedLightmapUVChannel = ExportMeshArgs.DatasmithAttributes->GetLightmapUVChannel();
-				const int32* ExportedSelectedChannel = UVChannelMap->Find(SelectedLightmapUVChannel + MaxToUnrealUVOffset);
-
-				if (ExportedSelectedChannel)
-				{
-					MeshElement->SetLightmapCoordinateIndex(*ExportedSelectedChannel);
-				}
-				else if (SelectedLightmapUVChannel != DefaultValue)
-				{
-					DatasmithMaxLogger::Get().AddGeneralError(*FString::Printf(TEXT("%s won't use the channel %i for its lightmap because it's not supported by the mesh. A new channel will be generated.")
-						, static_cast<const TCHAR*>(ExportMeshArgs.Node->GetName())
-						, SelectedLightmapUVChannel));
-				}
-			}
-		}
-	}
-
-
-	return MeshElement;
-}
-
 bool FDatasmithMaxMeshExporter::CreateDatasmithMesh(FDatasmithMesh& DSMesh, FMaxExportMeshArgs& ExportMeshArgs, TSet<uint16>& SupportedChannels)
 {
 	BOOL bNeedsDelete = false;
diff --git a/Engine/Source/Programs/Enterprise/Datasmith/DatasmithMaxExporter/Private/DatasmithMaxSceneExporter.cpp b/Engine/Source/Programs/Enterprise/Datasmith/DatasmithMaxExporter/Private/DatasmithMaxSceneExporter.cpp
--- a/Engine/Source/Programs/Enterprise/Datasmith/DatasmithMaxExporter/Private/DatasmithMaxSceneExporter.cpp
+++ b/Engine/Source/Programs/Enterprise/Datasmith/DatasmithMaxExporter/Private/DatasmithMaxSceneExporter.cpp
@@ -207,7 +207,7 @@ int FDatasmithMaxSceneExporter::GetSeedFromMaterial(Mtl* Material)
 	return Seed;
 }
 
-FString FDatasmithMaxSceneExporter::GetRandomSubMaterial(Mtl* Material, FVector3f RandomSeed)
+Mtl* FDatasmithMaxSceneExporter::GetRandomSubMaterial(Mtl* Material, FVector3f RandomSeed)
 {
 	uint32 RandomIndex = int(FMath::Abs(RandomSeed.X - RandomSeed.Y + RandomSeed.Z)) + GetSeedFromMaterial(Material);
 
@@ -223,10 +223,10 @@ FString FDatasmithMaxSceneExporter::GetRandomSubMaterial(Mtl* Material, FVector3
 
 	if (SubMats.Num() < 1)
 	{
-		return FString();
+		return nullptr;
 	}
 
-	return FString( SubMats[RandomIndex % SubMats.Num()]->GetName().data() );
+	return SubMats[RandomIndex % SubMats.Num()];
 }
 
 FTransform FDatasmithMaxSceneExporter::GetPivotTransform( INode* Node, float UnitMultiplier )
@@ -277,287 +277,6 @@ void FDatasmithMaxSceneExporter::ExportAnimation( TSharedRef< IDatasmithLevelSeq
 	}
 }
 
-void FDatasmithMaxSceneExporter::ExportMeshActor(TSharedRef< IDatasmithScene > DatasmithScene, TSet<uint16>& SupportedChannels, INode* Node, const TCHAR* MeshName,
-	float UnitMultiplier, bool bPivotIsBakedInGeometry, Mtl* StaticMeshMtl, const EStaticMeshExportMode& ExportMode)
-{
-	FTransform Pivot = GetPivotTransform( Node, UnitMultiplier );
-	bool bNeedPivotComponent = !bPivotIsBakedInGeometry && ( !Pivot.Equals( FTransform::Identity ) );
-
-	TSharedRef< IDatasmithMeshActorElement > MeshActor = FDatasmithSceneFactory::CreateMeshActor( *FString::FromInt( Node->GetHandle() ) );
-	MeshActor->SetStaticMeshPathName( MeshName );
-
-	switch (ExportMode)
-	{
-	case EStaticMeshExportMode::BoundingBox:
-		MeshActor->AddTag(TEXT("Datasmith.Attributes.Geometry: BoundingBox"));
-		break;
-	default:
-		break;
-	} 
-
-	TSharedRef< IDatasmithActorElement > NodeActor = MeshActor;
-
-	if ( bNeedPivotComponent )
-	{
-		NodeActor = FDatasmithSceneFactory::CreateActor( *FString::FromInt( Node->GetHandle() ) );
-
-		MeshActor->SetTranslation( Pivot.GetTranslation() );
-		MeshActor->SetRotation( Pivot.GetRotation() );
-		MeshActor->SetScale( Pivot.GetScale3D() );
-	}
-
-	ParseActor( Node, NodeActor, UnitMultiplier, DatasmithScene );
-
-	if (ExportMode == EStaticMeshExportMode::Default && StaticMeshMtl != Node->GetMtl())
-	{
-		ParseMaterialForMeshActor( Node->GetMtl(), MeshActor, SupportedChannels, FVector3f(MeshActor->GetTranslation()) );
-	}
-
-	if ( bPivotIsBakedInGeometry || bNeedPivotComponent )
-	{
-		FTransform NodeActorTransform( NodeActor->GetRotation(), NodeActor->GetTranslation(), NodeActor->GetScale() );
-		NodeActorTransform = Pivot.Inverse() * NodeActorTransform; // Remove pivot from the node actor transform
-
-		NodeActor->SetRotation( NodeActorTransform.GetRotation() );
-		NodeActor->SetTranslation( NodeActorTransform.GetTranslation() );
-		NodeActor->SetScale( NodeActorTransform.GetScale3D() );
-	}
-
-	if ( bNeedPivotComponent )
-	{
-		MeshActor->SetName( *( FString( NodeActor->GetName() ) + TEXT("_Pivot") ) );
-		MeshActor->SetLabel( *( FString( NodeActor->GetLabel() ) + TEXT("_Pivot") ) );
-		MeshActor->SetIsAComponent( true );
-
-		NodeActor->AddChild( MeshActor, EDatasmithActorAttachmentRule::KeepRelativeTransform );
-
-		if ( !FMath::IsNearlyEqual( NodeActor->GetScale().X, NodeActor->GetScale().Y ) || !FMath::IsNearlyEqual( NodeActor->GetScale().X, NodeActor->GetScale().Z ) )
-		{
-			DatasmithMaxLogger::Get().AddInvalidTransform( Node );
-		}
-	}
-
-	DatasmithScene->AddActor( NodeActor );
-}
-
-TSharedRef< IDatasmithActorElement > FDatasmithMaxSceneExporter::ExportHierarchicalInstanceStaticMeshActor(INode* Node, INode* CustomMeshNode, const TCHAR* Label, TSet<uint16>& SupportedChannels, Mtl* StaticMeshMtl, const TArray<Matrix3>* Instances,
-	const TCHAR* MeshName, float UnitMultiplier, const EStaticMeshExportMode& ExportMode, TSharedPtr< IDatasmithActorElement >& OutInversedHISMActor)
-{
-	check(Node && Instances && MeshName);
-
-	TSharedPtr< IDatasmithMeshActorElement > MeshActor;
-	FVector Pos, Scale;
-	FQuat Rotation;
-
-	const FVector3f RandomSeed(FMath::Rand(), FMath::Rand(), FMath::Rand());
-	auto FinalizeHISMActor = [&](TSharedPtr< IDatasmithMeshActorElement >& FinalizingActor, FVector3f Seed, const TCHAR* ActorLabel)
-	{
-		FinalizingActor->SetStaticMeshPathName(MeshName);
-
-		INode* MeshNode = CustomMeshNode ? CustomMeshNode : Node;
-
-		if (ExportMode == EStaticMeshExportMode::Default && StaticMeshMtl != MeshNode->GetMtl())
-		{
-			TSharedRef< IDatasmithMeshActorElement > MeshActorRef = FinalizingActor.ToSharedRef();
-			ParseMaterialForMeshActor(StaticMeshMtl, MeshActorRef, SupportedChannels, Seed);
-		}
-
-		if (ActorLabel)
-		{
-			FinalizingActor->SetLabel(ActorLabel);
-		}
-
-		FinalizingActor->SetIsAComponent(true);
-
-		switch (ExportMode)
-		{
-		case EStaticMeshExportMode::BoundingBox:
-			FinalizingActor->AddTag(TEXT("Datasmith.Attributes.Geometry: BoundingBox"));
-			break;
-		default:
-			break;
-		}
-	};
-
-	if (Node->GetWSMDerivedObject() != nullptr)
-	{
-		MaxToUnrealCoordinates(Node->GetObjTMAfterWSM(GetCOREInterface()->GetTime()), Pos, Rotation, Scale, UnitMultiplier);
-	}
-	else
-	{
-		MaxToUnrealCoordinates(Node->GetObjectTM(GetCOREInterface()->GetTime()), Pos, Rotation, Scale, UnitMultiplier);
-	}
-
-	const int32 InstancesCount = Instances->Num();
-	if ( InstancesCount == 1 )
-	{
-		// Export the the hism as a normal static mesh
-		MeshActor = FDatasmithSceneFactory::CreateMeshActor(MeshName);
-
-		// Apply the relative transfrom of the instance to the forest world transform
-		FTransform WithoutInstance(Rotation, Pos, Scale);
-		FVector InstancePos, InstanceScale;
-		FQuat InstanceRotation;
-		MaxToUnrealCoordinates((*Instances)[0], InstancePos, InstanceRotation, InstanceScale, UnitMultiplier);
-		FTransform WithInstance(FTransform(InstanceRotation, InstancePos, InstanceScale) * WithoutInstance);
-
-		MeshActor->SetTranslation(WithInstance.GetLocation());
-		MeshActor->SetScale(WithInstance.GetScale3D());
-		MeshActor->SetRotation(WithInstance.GetRotation());
-	}
-	else
-	{
-		TSharedRef< IDatasmithHierarchicalInstancedStaticMeshActorElement > HierarchicalInstanceStaticMeshActor = FDatasmithSceneFactory::CreateHierarchicalInstanceStaticMeshActor(MeshName);
-		MeshActor = HierarchicalInstanceStaticMeshActor;
-
-		FTransform HISMActorTransform(Rotation, Pos, Scale);
-		MeshActor->SetTranslation(Pos);
-		MeshActor->SetScale(Scale);
-		MeshActor->SetRotation(Rotation);
-
-		TArray< FTransform > InstancesTransform;
-		InstancesTransform.Reserve(InstancesCount);
-		for (int32 i = 0; i < InstancesCount; i++)
-		{
-			MaxToUnrealCoordinates((*Instances)[i], Pos, Rotation, Scale, UnitMultiplier);
-			InstancesTransform.Emplace(Rotation, Pos, Scale);
-		}
-
-		TArray< FTransform > NonInvertedTransforms, InvertedTransforms;
-		DatasmithMaxHelper::FilterInvertedScaleTransforms(InstancesTransform, NonInvertedTransforms, InvertedTransforms);
-
-		HierarchicalInstanceStaticMeshActor->ReserveSpaceForInstances(NonInvertedTransforms.Num());
-		for (const FTransform& InstanceTransform : NonInvertedTransforms)
-		{
-			HierarchicalInstanceStaticMeshActor->AddInstance(InstanceTransform);
-		}
-
-		//Adding an inverted HierarchicalInstancedStaticMeshActorElement for the instances with an inverted mesh due to negative scaling.
-		if (InvertedTransforms.Num() > 0)
-		{
-			FString InvertedHISMName = FString(MeshName).Append("_inv");
-			TSharedPtr< IDatasmithHierarchicalInstancedStaticMeshActorElement > InvertedHierarchicalInstanceStaticMeshActor = FDatasmithSceneFactory::CreateHierarchicalInstanceStaticMeshActor(*InvertedHISMName);
-			OutInversedHISMActor = InvertedHierarchicalInstanceStaticMeshActor;
-
-			OutInversedHISMActor->SetTranslation( HISMActorTransform.GetTranslation() );
-			OutInversedHISMActor->SetScale( -1 * HISMActorTransform.GetScale3D() );
-			OutInversedHISMActor->SetRotation( HISMActorTransform.GetRotation() );
-			
-			InvertedHierarchicalInstanceStaticMeshActor->ReserveSpaceForInstances(InvertedTransforms.Num());
-			for (const FTransform& InstanceTransform : InvertedTransforms)
-			{
-				//Correcting the instance transform to reverse the negative scaling of the parent.
-				InvertedHierarchicalInstanceStaticMeshActor->AddInstance( FTransform( InstanceTransform.GetRotation(), -1 * InstanceTransform.GetTranslation(), -1 * InstanceTransform.GetScale3D() ) );
-			}
-
-			TSharedPtr< IDatasmithMeshActorElement > InvertedMeshActor = InvertedHierarchicalInstanceStaticMeshActor;
-			if (Label)
-			{
-				FString InversedLabelString = FString(Label).Append("_inv");
-				FinalizeHISMActor(InvertedMeshActor, RandomSeed, *InversedLabelString);
-			}
-			else
-			{
-				FinalizeHISMActor(InvertedMeshActor, RandomSeed, nullptr);
-			}
-		}
-	}
-
-	FinalizeHISMActor(MeshActor, RandomSeed, Label);
-
-	return MeshActor.ToSharedRef();
-}
-
-void FDatasmithMaxSceneExporter::ParseMaterialForMeshActor(Mtl* Material, TSharedRef< IDatasmithMeshActorElement >& MeshActor, TSet<uint16>& SupportedChannels, const FVector3f& RandomSeed)
-{
-	if (FDatasmithMaxMatHelper::GetMaterialClass(Material) == EDSMaterialType::XRefMat)
-	{
-		ParseMaterialForMeshActor(FDatasmithMaxMatHelper::GetRenderedXRefMaterial(Material), MeshActor, SupportedChannels, RandomSeed);
-		return;
-	}
-
-	if (Material != nullptr)
-	{
-		if (SupportedChannels.Num() <= 1)
-		{
-			if (FDatasmithMaxMatHelper::GetMaterialClass(Material) != EDSMaterialType::MultiMat)
-			{
-				if (FDatasmithMaxMatHelper::GetMaterialClass(Material) != EDSMaterialType::TheaRandom)
-				{
-					MeshActor->AddMaterialOverride(Material->GetName().data(), -1);
-				}
-				else
-				{
-					MeshActor->AddMaterialOverride( *GetRandomSubMaterial(Material, RandomSeed), -1 );
-				}
-			}
-			else
-			{
-				int Mid = 0;
-
-				//Find the lowest supported material id.
-				SupportedChannels.Sort([](const uint16& A, const uint16& B) { return (A < B); });
-				for (const uint16 SupportedMid : SupportedChannels)
-				{
-					Mid = SupportedMid;
-				}
-
-				// Replicate the 3ds max behavior where material ids greater than the number of sub-materials wrap around and are mapped to the existing sub-materials
-				if (Mtl* SubMaterial = Material->GetSubMtl(Mid % Material->NumSubMtls()))
-				{
-					if (FDatasmithMaxMatHelper::GetMaterialClass(SubMaterial) != EDSMaterialType::TheaRandom)
-					{
-						MeshActor->AddMaterialOverride(SubMaterial->GetName().data(), -1);
-					}
-					else
-					{
-						MeshActor->AddMaterialOverride( *GetRandomSubMaterial(SubMaterial, RandomSeed), -1 );
-					}
-				}
-			}
-		}
-		else
-		{
-			int ActualSubObj = 1;
-			SupportedChannels.Sort([](const uint16& A, const uint16& B) { return (A < B); });
-			for (const uint16 Mid : SupportedChannels)
-			{
-				if (FDatasmithMaxMatHelper::GetMaterialClass(Material) != EDSMaterialType::MultiMat)
-				{
-					if (FDatasmithMaxMatHelper::GetMaterialClass(Material) != EDSMaterialType::TheaRandom)
-					{
-						MeshActor->AddMaterialOverride(Material->GetName().data(), ActualSubObj);
-					}
-					else
-					{
-						MeshActor->AddMaterialOverride( *GetRandomSubMaterial(Material, RandomSeed), ActualSubObj );
-					}
-				}
-				else
-				{
-					// Replicate the 3ds max behavior where material ids greater than the number of sub-materials wrap around and are mapped to the existing sub-materials
-					int32 MaterialIndex = Mid % Material->NumSubMtls();
-
-					Mtl* SubMaterial = Material->GetSubMtl(MaterialIndex);
-					if (SubMaterial != nullptr)
-					{
-						if (FDatasmithMaxMatHelper::GetMaterialClass(SubMaterial) != EDSMaterialType::TheaRandom)
-						{
-							//Material slots in Max are not zero-based, so we serialize our SlotID starting from 1 for better visual consistency.
-							MeshActor->AddMaterialOverride(SubMaterial->GetName().data(), Mid + 1);
-						}
-						else
-						{
-							MeshActor->AddMaterialOverride( *GetRandomSubMaterial(SubMaterial, RandomSeed), MaterialIndex + 1 );
-						}
-					}
-				}
-				ActualSubObj++;
-			}
-		}
-	}
-}
-
 TSharedPtr< IDatasmithLightActorElement > FDatasmithMaxSceneExporter::CreateLightElementForNode(INode* Node, const TCHAR* Name)
 {
 	EMaxLightClass LightClass = FDatasmithMaxSceneParser::GetLightClass(Node);
diff --git a/Engine/Source/Programs/Enterprise/Datasmith/DatasmithMaxExporter/Public/DatasmithMaxDirectLink.h b/Engine/Source/Programs/Enterprise/Datasmith/DatasmithMaxExporter/Public/DatasmithMaxDirectLink.h
--- a/Engine/Source/Programs/Enterprise/Datasmith/DatasmithMaxExporter/Public/DatasmithMaxDirectLink.h
+++ b/Engine/Source/Programs/Enterprise/Datasmith/DatasmithMaxExporter/Public/DatasmithMaxDirectLink.h
@@ -46,16 +46,16 @@ class FLayerTracker;
 class FMaterialsCollectionTracker;
 
 
-// Global export options, stored in preferences
+/** Global export options, stored in preferences */
 class IPersistentExportOptions
 {
 public:
 
-	// Whether to export animation or not
+	/** Whether to export animation or not */
 	virtual void SetAnimatedTransforms(bool) = 0;
 	virtual bool GetAnimatedTransforms() = 0;
 
-	// Whether to output export statistics to listener/log
+	/** Whether to output export statistics to listener/log */
 	virtual void SetStatSync(bool) = 0;
 	virtual bool GetStatSync() = 0;
 
@@ -67,7 +67,7 @@ public:
 
 };
 
-//---- Main class for export/change tracking
+/** Main class for export/change tracking */
 class IExporter
 {
 	public:
@@ -82,15 +82,15 @@ class IExporter
 
 	virtual void InitializeScene() = 0;
 
-	// Scene update
+	/** Scene update */
 	virtual bool UpdateScene(bool bQuiet) = 0;
 
 	virtual void ResetSceneTracking() = 0;
 
-	// ChangeTracking
+	/** ChangeTracking */
 	virtual void StartSceneChangeTracking() = 0;
 
-	// DirectLink
+	/** DirectLink */
 	virtual void InitializeDirectLinkForScene() = 0;
 	virtual void UpdateDirectLinkScene() = 0;
 	virtual bool ToggleAutoSync() = 0;
@@ -101,7 +101,8 @@ class IExporter
 	virtual void PerformSync(bool bQuiet) = 0;
 };
 
-bool CreateExporter(bool bEnableUI, const TCHAR* EnginePath); // Create exporter with ability for DirectLink change tracking
+/** Create exporter with ability for DirectLink change tracking */
+bool CreateExporter(bool bEnableUI, const TCHAR* EnginePath); 
 IExporter* GetExporter();
 void ShutdownExporter();
 bool Export(const TCHAR* Name, const TCHAR* OutputPath, bool bQuiet, bool bSelected);
@@ -123,13 +124,13 @@ struct FXRefScene
 	operator bool(){ return XRefFileIndex != -1; }
 };
 
-// Incapsulating time slider value to distinguish usage of TimeValue specifically for point on timeslider where we syncing
+/** Incapsulating time slider value to distinguish usage of TimeValue specifically for point on timeslider where we syncing */
 struct FSyncPoint
 {
 	TimeValue Time;
 };
 
-// Incapsulating time interval value to distinguish usage of Interval specifically for validity of synced entity 
+/** Incapsulating time interval value to distinguish usage of Interval specifically for validity of synced entity */
 struct FValidity
 {
 	explicit FValidity()
@@ -158,14 +159,14 @@ struct FValidity
 		return ValidityInterval.InInterval(State.Time) != 0;
 	}
 
-	// Maximize validity interval before updating invalidated 
+	/** Maximize validity interval before updating invalidated */
 	void ResetValidityInterval()
 	{
 		check(bIsInvalidated);
 		ValidityInterval.SetInfinite();
 	}
 
-	// Intersect
+	/** Intersect current validity interval with the param*/
 	void NarrowValidityToInterval(const Interval& InValidityInterval)
 	{
 		ValidityInterval &= InValidityInterval;
@@ -176,9 +177,11 @@ struct FValidity
 		NarrowValidityToInterval(Validity.ValidityInterval);
 	}
 
-	// this Validity contains another validity interval fully
-	// Used to determine if this validity doesn't need to be updated
-	// @return True is this validity doesn't need to be updated
+	/**
+	This Validity contains another validity interval fully.
+	Used to determine if this validity doesn't need to be updated
+	@return True is this validity doesn't need to be updated
+	*/
 	bool Overlaps(const FValidity& Validity) const
 	{
 		return ValidityInterval.IsSubset(Validity.ValidityInterval);
@@ -193,11 +196,13 @@ private:
 	bool bIsInvalidated; 
 };
 
-// NodeTracker - associated with INode, plus validity
-// Converter - holds convertion data for SPECIFIC node object type, e.g. geometry, railclone, light, dummy
-// Converted - holds Datasmith converted data for node
+/**
+NodeTracker - associated with INode, plus validity
+Converter - holds convertion data for SPECIFIC node object type, e.g. geometry, railclone, light, dummy
+Converted - holds Datasmith converted data for node
+*/
 
-// Identifies Max node to track its changes, including tracked dependencies
+/** Identifies Max node to track its changes, including tracked dependencies */
 class FNodeTracker: FNoncopyable
 {
 public:
@@ -271,42 +276,54 @@ public:
 		Converter.Reset();
 	}
 
-	// Node entity identification data
+	/** Node entity identification data */
 	FNodeKey NodeKey;
 	INode* const Node;
-	// Keep root node and xref index when this node is a direct child of an XRef scene
-	// This is needed to retrieve parent node(e.g. when updated)
-	// Keeping parent node itself doesn't work - it can change and the only way to get it when it's changed is to call INode::GetXRefParent
+
+	/**
+	Keep root node and xref index when this node is a direct child of an XRef scene
+	This is needed to retrieve parent node(e.g. when updated)
+	Keeping parent node itself doesn't work - it can change and the only way to get it when it's changed is to call INode::GetXRefParent
+	*/
 	FXRefScene XRefScene;
 	FString Name;
 
-	// Node validity
-	bool bParsed = false; // 
+	/** Node validity */
+	bool bParsed = false;
 	bool bDeleted = false;
 	FValidity Validity;
 	FValidity SubtreeValidity;
 
-	// Other related tracked entities
+	/** Other related tracked entities */
 	FNodeTracker* Collision = nullptr;
 	FLayerTracker* Layer = nullptr;
 	TSet<class FMaterialTracker*> MaterialTrackers; 
 
-	// Node conversion state
-	TUniquePtr<FNodeConverted> Converted = nullptr; // Datasmith element that this Node is converted to
-	TUniquePtr<FNodeConverter> Converter = nullptr; // Converter for specific Max object type
+	/** Node conversion state */
+
+	/** Datasmith element that this Node is converted to */
+	TUniquePtr<FNodeConverted> Converted = nullptr;
+	/** Converter for specific Max object type */
+	TUniquePtr<FNodeConverter> Converter = nullptr; 
 };
 
 class FMeshConverterSource
 {
 public:
-	// Node this mesh instantiates. When this is a 'regular' node it's  just instantiates mesh for Params
-	// but it's possible that this Node wants mesh to be bounding-box(when DatasmithAttributes spicifies it)
+	/**
+	Node this mesh instantiates. When this is a 'regular' node it's  just instantiates mesh for Params
+	but it's possible that this Node wants mesh to be bounding-box(when DatasmithAttributes spicifies it)
+	*/
 	INode* Node;
 
-	FString MeshName; // Suggested Mesh name, resulting mesh name should be this???
+	/** Suggested Mesh name */
+	FString MeshName;
+
+	/*Extracted render mesh*/
+	GeomUtils::FRenderMeshForConversion RenderMesh;
 
-	GeomUtils::FRenderMeshForConversion RenderMesh; // Extracted render mesh
-	bool bConsolidateMaterialIds; // Whether to join all materials id into single material slot for render mesh (used when a geometry doesn't have a multimaterial assigned)
+	/** Whether to join all materials id into single material slot for render mesh (used when a geometry doesn't have a multimaterial assigned) */
+	bool bConsolidateMaterialIds; 
 
 	GeomUtils::FRenderMeshForConversion CollisionMesh;
 };
@@ -363,14 +380,16 @@ struct FSceneUpdateStats
 #define SCENE_UPDATE_STAT_GET(Category, Name) Stats.##Category##Name
 
 
-// Modifies Datasmith scene in responce to change notification calls
-// Subscription to various Max notification systems is done separately, see  FNotifications
+/**
+Modifies Datasmith scene in responce to change notification calls
+Subscription to various Max notification systems is done separately, see  FNotifications
+*/
 class ISceneTracker
 {
 public:
 	virtual ~ISceneTracker(){}
 
-	// Change notifications
+	/** Change notifications */
 	virtual void NodeAdded(INode* Node) = 0;
 	virtual void NodeDeleted(INode* Node) = 0;
 	virtual void NodeGeometryChanged(INode* Node) = 0;
@@ -390,7 +409,7 @@ public:
 
 	virtual bool IsUpdateInProgress() = 0;
 
-	// Scene modification
+	/** Scene modification*/
 	virtual void AddMeshElement(TSharedPtr<IDatasmithMeshElement>& Mesh, FDatasmithMesh& DatasmithMesh, FDatasmithMesh* CollisionMesh) = 0;
 	virtual void ReleaseMeshElement(FMeshConverted& Converted) = 0;
 	virtual void SetupActor(FNodeTracker& NodeTracker) = 0;
@@ -400,7 +419,7 @@ public:
 	virtual void NodeXRefMerged(INode* Node) = 0;
 	virtual void RemapConvertedMaterialUVChannels(Mtl* ActualMaterial, const TSharedPtr<IDatasmithBaseMaterialElement>& DatasmithMaterial) = 0;
 	
-	// Sync/Update
+	/** Sync/Update */
 	FSyncPoint CurrentSyncPoint;
 
 	virtual void AddGeometryNodeInstance(FNodeTracker& NodeTracker, FMeshNodeConverter& MeshConverter, Object* Obj) = 0;
@@ -415,27 +434,32 @@ public:
 	virtual TSharedRef<IDatasmithScene> GetDatasmithSceneRef() = 0;
 	virtual const TCHAR* GetAssetsOutputPath() = 0;
 
-	// Utility
+	/** Utility */
 	virtual FNodeTracker* GetNodeTrackerByNodeName(const TCHAR* Name) = 0;
 	virtual FSceneUpdateStats& GetStats() = 0;
 };
 
-// Input data for mesh conversion
+/** Input data for mesh conversion*/
 struct MeshConversionParams
 {
-	INode* Node; // Node, this geom object created from
-	const GeomUtils::FRenderMeshForConversion& RenderMesh; // Extracted render mesh
-	bool bConsolidateMaterialIds; // Whether to join all materials id into single material slot (used when a geometry doesn't have a multimaterial assigned)
+	/** Node, this geom object created from */
+	INode* Node;
+
+	/** Extracted render mesh */
+	const GeomUtils::FRenderMeshForConversion& RenderMesh;
+
+	/** Whether to join all materials id into single material slot (used when a geometry doesn't have a multimaterial assigned) */
+	bool bConsolidateMaterialIds;
 };
 
-// Creates Mesh element and converts max mesh into it
+/** Creates Mesh element and converts max mesh into it */
 bool ConvertMaxMeshToDatasmith(TimeValue CurrentTime, ISceneTracker& Scene, FMeshConverterSource& MeshSource, FMeshConverted& MeshConverted);
 
 bool OpenDirectLinkUI();
 const TCHAR* GetDirectlinkCacheDirectory();
 
 
-//---- Max Notifications/Events/Callbacks handling
+/** Max Notifications/Events/Callbacks handling */
 class FNodeEventCallback;
 class FNodeObserver;
 class FMaterialObserver;
@@ -517,7 +541,7 @@ public:
 	}
 };
 
-//---- Material change tracking
+/** Material change tracking */
 class FMaterialTracker
 {
 public:
@@ -526,7 +550,8 @@ public:
 
 	Mtl* Material;
 
-	TArray<Mtl*> Materials; // Actual materials used for this assigned material
+	/*Actual materials used for this assigned material*/
+	TArray<Mtl*> Materials;
 	TArray<Texmap*> Textures;
 
 	bool bInvalidated = true;
@@ -544,7 +569,7 @@ public:
 	void ResetActualMaterialAndTextures()
 	{
 		Materials.Reset();
-		Textures.Reset(); //todo: unregister textures
+		Textures.Reset();
 	}
 
 	void AddActualMaterial(Mtl* ActualMaterial)
@@ -578,9 +603,11 @@ private:
 	}
 };
 
-// Converts texmap to texture element
-// This is different from IDatasmithMaxTexmapToUEPbr which performs texmap conversion to full MaterialExpression which is used in material graph
-// IDatasmithMaxTexmapToUEPbr uses ToTextureElementConverter when it needs a texture element
+/**
+Converts texmap to texture element
+This is different from IDatasmithMaxTexmapToUEPbr which performs texmap conversion to full MaterialExpression which is used in material graph
+IDatasmithMaxTexmapToUEPbr uses ToTextureElementConverter when it needs a texture element
+*/
 class ITexmapToTextureElementConverter
 {
 public:
@@ -598,28 +625,43 @@ public:
 
 	void Reset();
 
-	FMaterialTracker* AddMaterial(Mtl* Material); // Add material to track its changes
-
-	void AddActualMaterial(FMaterialTracker& MaterialTracker, Mtl* Material); // Add material used in a tracked material's graph
-	const TCHAR* GetMaterialName(Mtl* SubMaterial); // Get name used for Datasmith material. Datasmith material names must be unique(used to identify elements)
-
-	void AssignMeshMaterials(const TSharedPtr<IDatasmithMeshElement>& MeshElement, Mtl* Material, const TSet<uint16>& SupportedChannels);
-	void AssignMeshActorMaterials(const TSharedPtr<IDatasmithMeshActorElement>& MeshActor, Mtl* Material, TSet<uint16>& SupportedChannels, const FVector3f& RandomSeed);
-
-	void AddDatasmithMaterialForUsedMaterial(TSharedRef<IDatasmithScene> DatasmithScene, Mtl* Material, TSharedPtr<IDatasmithBaseMaterialElement> DatasmithMaterial); // Record which Datasmith material was created for a Max material, not only for tracked(assigned) materials
-
-	void InvalidateMaterial(Mtl* Material); // Mark changed
-
-	void UpdateMaterial(FMaterialTracker* MaterialTracker); // Reparse source material
-
-	void ConvertMaterial(Mtl* Material, TSharedRef<IDatasmithScene> DatasmithScene, const TCHAR* AssetsPath, TSet<Texmap*>& TexmapsConverted); // Convert to Datasmith
-
-	// When Material is not used my scene anymore - stop tracking it
+	/** Add material to track its changes*/
+	FMaterialTracker* AddMaterial(Mtl* Material);
+	/** When Material is not used in scene anymore - stop tracking it*/
 	void ReleaseMaterial(FMaterialTracker& MaterialTracker);
-
-	// Clean all converted data, remove from datasmith scene(e.g. before rebuilding material)
+	/** Mark changed*/
+	void InvalidateMaterial(Mtl* Material);
+	/** Reparse source material*/
+	void UpdateMaterial(FMaterialTracker* MaterialTracker);
+	/** Convert to Datasmith */
+	void ConvertMaterial(Mtl* Material, TSharedRef<IDatasmithScene> DatasmithScene, const TCHAR* AssetsPath, TSet<Texmap*>& TexmapsConverted);
+
+	/** Clean all converted data, remove from datasmith scene(e.g. before rebuilding material) */
 	void RemoveConvertedMaterial(FMaterialTracker& MaterialTracker);
 
+	/** Add material used in a specific tracked material's graph */
+	void AddActualMaterial(FMaterialTracker& MaterialTracker, Mtl* Material);
+	/** Reverse of AddActualMaterial*/
+	void ReleaseActualMaterial(FMaterialTracker& MaterialTracker, Mtl* Material);
+
+	/** Get name used for Datasmith material. Datasmith material names must be unique(used to identify elements) */
+	const TCHAR* GetMaterialName(Mtl* SubMaterial);
+
+	/** If material is used on a mesh/actor to create a datamsith element */
+	bool IsMaterialUsed(Mtl* Material);
+	/** Set datasmith materials for static mesh for assigned max material */
+	void SetMaterialsForMeshElement(const TSharedPtr<IDatasmithMeshElement>& MeshElement, Mtl* Material, const TSet<uint16>& SupportedChannels);
+	void SetMaterialForMeshElement(const TSharedPtr<IDatasmithMeshElement>& MeshElement, Mtl* MaterialAssignedToNode, Mtl* ActualMaterial, uint16 SlotId);
+	void UnSetMaterialsForMeshElement(const TSharedPtr<IDatasmithMeshElement>& MeshElement);
+	/** Set datasmith materials for mesh actor for assigned max material */
+	void SetMaterialsForMeshActor(const TSharedPtr<IDatasmithMeshActorElement>& MeshActor, Mtl* Material, TSet<uint16>& SupportedChannels, const FVector3f& RandomSeed);
+	void SetMaterialForMeshActorElement(const TSharedPtr<IDatasmithMeshActorElement>& MeshActor, Mtl* MaterialAssignedToNode, Mtl* Material, int32 SlotId);
+	void UnSetMaterialsForMeshActor(const TSharedPtr<IDatasmithMeshActorElement>&);
+
+	/** Record which Datasmith material was created for a Max material, not only for tracked(assigned) materials */
+	void AddDatasmithMaterialForUsedMaterial(TSharedRef<IDatasmithScene> DatasmithScene, Mtl* Material, TSharedPtr<IDatasmithBaseMaterialElement> DatasmithMaterial);
+	void RemoveDatasmithMaterialForUsedMaterial(Mtl* Material);
+
 	void ResetInvalidatedMaterials()
 	{
 		InvalidatedMaterialTrackers.Reset();		
@@ -632,7 +674,7 @@ public:
 
 	void UpdateTexmap(Texmap* Texmap);
 
-	// Adds a texmap that need to be converted to texture element
+	/** Adds a texmap that need to be converted to texture element */
 	void AddTexmapForConversion(Texmap* Texmap, const FString& DesiredTextureElementName, const TSharedPtr<ITexmapToTextureElementConverter>& Converter);
 
 	void AddTextureElement(const TSharedPtr<IDatasmithTextureElement>& );
@@ -640,26 +682,47 @@ public:
 
 	ISceneTracker& SceneTracker;
 
-	TMap<FMaterialKey, FMaterialTrackerHandle> MaterialTrackers; // Tracks all assigned materials
-	TSet<FMaterialTracker*> InvalidatedMaterialTrackers; // Materials needing update
-
-	TSet<Mtl*> EncounteredMaterials; // All materials from the assigned material's graphs 
-	TSet<Texmap*> EncounteredTextures; // All materials from the assigned material's graphs 
-
-	TArray<FString> MaterialNames; // todo: UETOOL-4369 fix changing materials names(to make them unique for easy export)
-
-	TMap<Mtl*, TSet<FMaterialTracker*>> UsedMaterialToMaterialTracker; // Materials uses by nodes keep set of assigned materials they are used for
+	/** Tracks all assigned materials */
+	TMap<FMaterialKey, FMaterialTrackerHandle> MaterialTrackers;
+	/** Materials needing update */
+	TSet<FMaterialTracker*> InvalidatedMaterialTrackers;
+
+	/**
+	3ds Max Materials can be:
+	- Tracked. When a material is assigned to a 3ds Max node. Assigned materials are tracked for changes by the plugin.
+	- Used. When a material is used to create a DatasmithMaterialElement associated with it and this element is used by a DatasmithStaticMesh or MeshActor
+	- Actual. Tracked + Used + <all unused submaterials>  - 
+	 Tracked material won't necessarily be Used, for various reasons, e.g. it can be a composite material(in this case only its submaterials are be Used). Or material is assigned but all geometry is invisible.
+	 Used materials can be missing in Tracked - a submaterial of an MultiSubobj material assigned to node is not assigned directly.
+	 */
+
+	/** Tracked material for every Actual material (i.e. set of material which use an Mtl as a submaterial ) */
+	TMap<Mtl*, TSet<FMaterialTracker*>> ActualMaterialToMaterialTracker;
+	/** All materials that have a corresponding Datasmith material created */
 	TMap<Mtl*, TSharedPtr<IDatasmithBaseMaterialElement>> UsedMaterialToDatasmithMaterial;
 
+	/** Cached unique name for each used material for Datasmith Element */
 	TMap<Mtl*, FString> UsedMaterialToDatasmithMaterialName;
 
-	TMap<Texmap*, TSet<FMaterialTracker*>> UsedTextureToMaterialTracker; // Materials uses by nodes keep set of assigned textures they are used for
-	TMap<Texmap*, TSet<TSharedPtr<IDatasmithTextureElement>>> UsedTextureToDatasmithElement; // Keep track of Datasmith Element created for texmap to simplify update/removal(no need to search DatasmithScene)
-	TMap<TSharedPtr<IDatasmithTextureElement>, TSet<Texmap*>> TextureElementToTexmap; // Back reference - which texture elements correspond to which texmaps(different texmaps can make same element)
+	/** Track materials used for static meshes */
+	TMap<Mtl*, TMap<IDatasmithMeshElement*, TSet<int32>>> UsedMaterialToMeshElementSlot;
+	TMap<IDatasmithMeshElement*, TMap<Mtl*, TSet<int32>>> MeshElementToUsedMaterialSlot;
+
+	/** Track materials used for overrides on mesh actors */
+	TMap<Mtl*, TMap<IDatasmithMeshActorElement*, TSet<int32>>> UsedMaterialToMeshActorSlot;
+	TMap<IDatasmithMeshActorElement*, TMap<Mtl*, TSet<int32>>> MeshActorToUsedMaterialSlot;
+
+	/** Texmaps tracking */
+	/** Materials uses by nodes keep set of assigned textures they are used for */
+	TMap<Texmap*, TSet<FMaterialTracker*>> UsedTextureToMaterialTracker;
+	/** Keep track of Datasmith Element created for texmap to simplify update/removal(no need to search DatasmithScene) */
+	TMap<Texmap*, TSet<TSharedPtr<IDatasmithTextureElement>>> UsedTextureToDatasmithElement;
+	/** Back reference - which texture elements correspond to which texmaps(different texmaps can make same element) */
+	TMap<TSharedPtr<IDatasmithTextureElement>, TSet<Texmap*>> TextureElementToTexmap;
 
 	TMap<Texmap*, TSharedPtr<ITexmapToTextureElementConverter>> TexmapConverters;
 
-	// Bitmap textures using the same file become same texture element
+	/** Bitmap textures using the same file become same texture element */
 	TMap<FString, TSharedPtr<IDatasmithTextureElement>>  BitmapTextureElements;
 
 	TSet<TSharedPtr<IDatasmithTextureElement>> TextureElementsAddedToScene;
@@ -669,7 +732,7 @@ public:
 	FDatasmithUniqueNameProvider MaterialNameProvider;
 };
 
-//----
+/** Logging */
 
 void LogFlush();
 
@@ -688,7 +751,7 @@ void LogWarningDialog(const FString& Msg);
 void LogCompletionDialog(const FString& Msg);
 void LogInfoDialog(const FString& Msg);
 
-// Debug logging
+/** Debug logging */
 void LogDebugImpl(const FString& Msg);
 void LogDebugDialog(const FString& Msg);
 
diff --git a/Engine/Source/Programs/Enterprise/Datasmith/DatasmithMaxExporter/Public/DatasmithMaxSceneExporter.h b/Engine/Source/Programs/Enterprise/Datasmith/DatasmithMaxExporter/Public/DatasmithMaxSceneExporter.h
--- a/Engine/Source/Programs/Enterprise/Datasmith/DatasmithMaxExporter/Public/DatasmithMaxSceneExporter.h
+++ b/Engine/Source/Programs/Enterprise/Datasmith/DatasmithMaxExporter/Public/DatasmithMaxSceneExporter.h
@@ -61,8 +61,6 @@ public:
 	static bool ExportActor(TSharedRef< IDatasmithScene > DatasmithScene, INode* Node, const TCHAR* Name, float UnitMultiplier);
 	static void ExportMeshActor(TSharedRef< IDatasmithScene > DatasmithScene,  TSet<uint16>& SupportedChannels, INode* Node, const TCHAR* MeshName,
 		float UnitMultiplier, bool bPivotIsBakedInGeometry, Mtl* StaticMeshMtl, const EStaticMeshExportMode& ExportMode);
-	static TSharedRef< IDatasmithActorElement > ExportHierarchicalInstanceStaticMeshActor(INode* Node, INode* CustomMeshNode, const TCHAR* Label, TSet<uint16>& SupportedChannels, Mtl* StaticMeshMtl, const TArray<Matrix3>* Instances,
-		const TCHAR* MeshName, float UnitMultiplier, const EStaticMeshExportMode& ExportMode, TSharedPtr< IDatasmithActorElement >& OutInversedHISMActor );
 	static bool ExportCameraActor(TSharedRef< IDatasmithScene > DatasmithScene, INode* Parent, INodeTab Instances, int InstanceIndex, const TCHAR* Name, float UnitMultiplier);
 	static void WriteEnvironment(TSharedRef< IDatasmithScene > DatasmithScene, bool bOnlySelection);
 	static void ExportToneOperator(TSharedRef< IDatasmithScene > DatasmithScene);
@@ -70,7 +68,7 @@ public:
 
 	static FString GetActualPath(const TCHAR* OriginalPath);
 
-	static FString GetRandomSubMaterial(Mtl* Material, FVector3f RandomSeed);
+	static Mtl* GetRandomSubMaterial(Mtl* Material, FVector3f RandomSeed);
 
 	/**
 	 * Extract the Node to Object transform in UE coordinates
@@ -83,8 +81,6 @@ public:
 	static FTransform GetPivotTransform(INode* Node, float UnitMultiplier);
 	static void MaxToUnrealCoordinates(Matrix3 Matrix, FVector& Translation, FQuat& Rotation, FVector& Scale, float UnitMultiplier, const FMaxLightCoordinateConversionParams& LightParams = FMaxLightCoordinateConversionParams());
 
-	static void ParseMaterialForMeshActor(Mtl* Material, TSharedRef< IDatasmithMeshActorElement >& MeshActor, TSet<uint16>& SupportedChannels, const FVector3f& RandomSeed);
-	
 	static TSharedPtr< IDatasmithLightActorElement > CreateLightElementForNode(INode* Node, const TCHAR* Name);
 	static bool ParseLight(DatasmithMaxDirectLink::FLightNodeConverter&, INode* Node, TSharedRef< IDatasmithLightActorElement > LightElement, TSharedRef< IDatasmithScene > DatasmithScene);
 
