Datasmith Revit Exporter - SectionBox acquisition crash fix

diff --git a/Engine/Source/Programs/Enterprise/Datasmith/DatasmithRevitExporter/Private/DatasmithRevitDocumentData.cs b/Engine/Source/Programs/Enterprise/Datasmith/DatasmithRevitExporter/Private/DatasmithRevitDocumentData.cs
--- a/Engine/Source/Programs/Enterprise/Datasmith/DatasmithRevitExporter/Private/DatasmithRevitDocumentData.cs
+++ b/Engine/Source/Programs/Enterprise/Datasmith/DatasmithRevitExporter/Private/DatasmithRevitDocumentData.cs
@@ -1307,26 +1307,28 @@ namespace DatasmithRevitExporter
 			// Cache document section boxes
 			if (CurrentDocument.ActiveView != null)
 			{
-				BoundingBoxXYZ BBox = (GetElement(CurrentDocument.ActiveView.Id) as View3D).GetSectionBox();
+				View3D CurrentView3d = (GetElement(CurrentDocument.ActiveView.Id) as View3D);
 
-				if (BBox.IsSet && BBox.Enabled)
+				SectionBox = null;
+				SectionBoxOutline = null;
+
+				if (CurrentView3d != null && CurrentView3d.IsSectionBoxActive)
 				{
-					SectionBox = new OrientatedBoundingBox(BBox.Transform, BBox.Min, BBox.Max, true);
-					if (!SectionBox.bIsValidData)
-					{
-						SectionBox = null;
-						SectionBoxOutline = null;
-					}
-					else
+					BoundingBoxXYZ BBox = CurrentView3d.GetSectionBox();
+
+					if (BBox.IsSet && BBox.Enabled)
 					{
-						SectionBoxOutline = new Outline(SectionBox.AxisAllignedMin, SectionBox.AxisAllignedMax);
+						SectionBox = new OrientatedBoundingBox(BBox.Transform, BBox.Min, BBox.Max, true);
+						if (SectionBox.bIsValidData)
+						{
+							SectionBoxOutline = new Outline(SectionBox.AxisAllignedMin, SectionBox.AxisAllignedMax);
+						}
+						else
+						{
+							SectionBox = null;
+						}
 					}
 				}
-				else
-				{
-					SectionBox = null;
-					SectionBoxOutline = null;
-				}
 			}
 
 			if (InLinkedDocumentId != null)
