[Solidworks] Display state change correctly changes materials

- Old DatasmithMeshElement are removed from datasmith scene when mesh is updated

diff --git a/Engine/Source/Programs/Enterprise/Datasmith/DatasmithSolidworksExporter/Private/AssemblyDocument.cs b/Engine/Source/Programs/Enterprise/Datasmith/DatasmithSolidworksExporter/Private/AssemblyDocument.cs
--- a/Engine/Source/Programs/Enterprise/Datasmith/DatasmithSolidworksExporter/Private/AssemblyDocument.cs
+++ b/Engine/Source/Programs/Enterprise/Datasmith/DatasmithSolidworksExporter/Private/AssemblyDocument.cs
@@ -27,10 +27,17 @@ namespace DatasmithSolidworks
 			public ConcurrentDictionary<string, FObjectMaterials> ComponentsMaterialsMap = new ConcurrentDictionary<string, FObjectMaterials>();
 			public Dictionary<string, string> ComponentToPartMap = new Dictionary<string, string>();
 			public Dictionary<string, float[]> ComponentsTransformsMap = new Dictionary<string, float[]>();
+
+			/** All updated(not dirty) components. Updated component will be removed from dirty on export end*/
 			public HashSet<string> CleanComponents = new HashSet<string>();
+
+			/** Flags for each component that needs update. Cleared when export completes*/
 			public Dictionary<string, uint> DirtyComponents = new Dictionary<string, uint>();
+
 			public HashSet<string> ComponentsToDelete = new HashSet<string>();
 			public Dictionary<string, Component2> ExportedComponentsMap = new Dictionary<string, Component2>();
+			/** Stores which mesh was exported for each component*/
+			public Dictionary<string, string> ComponentNameToMeshNameMap = new Dictionary<string, string>();
 		}
 
 		public AssemblyDoc SwAsmDoc { get; private set; } = null;
@@ -57,6 +64,13 @@ namespace DatasmithSolidworks
 				string ActorName = FDatasmithExporter.SanitizeName(CompName);
 				SyncState.ComponentToPartMap.Remove(CompName);
 				SyncState.ExportedComponentsMap.Remove(CompName);
+
+				if (SyncState.ComponentNameToMeshNameMap.TryGetValue(CompName, out string MeshName))
+				{
+					SyncState.ComponentNameToMeshNameMap.Remove(CompName);
+					Exporter.RemoveMesh(MeshName);
+				}
+
 				SyncState.ComponentsMaterialsMap.TryRemove(CompName, out FObjectMaterials _);
 				SyncState.ComponentsTransformsMap.Remove(CompName);
 				Exporter.RemoveActor(ActorName);
@@ -99,7 +113,8 @@ namespace DatasmithSolidworks
 
 			// Export meshes
 			SetExportStatus($"Component Meshes");
-			ConcurrentBag<Tuple<FDatasmithFacadeMeshElement, FDatasmithFacadeMesh>> CreatedMeshes = new ConcurrentBag<Tuple<FDatasmithFacadeMeshElement, FDatasmithFacadeMesh>>();
+			ConcurrentDictionary<Component2, Tuple<FDatasmithFacadeMeshElement, FDatasmithFacadeMesh>> CreatedMeshes = new ConcurrentDictionary<Component2, Tuple<FDatasmithFacadeMeshElement, FDatasmithFacadeMesh>>();
+			
 			Parallel.ForEach(MeshesToExportMap, KVP =>
 			{
 				Component2 Comp = KVP.Key;
@@ -113,18 +128,17 @@ namespace DatasmithSolidworks
 				if (MeshData != null)
 				{
 					Tuple<FDatasmithFacadeMeshElement, FDatasmithFacadeMesh> NewMesh = null;
-					Exporter.ExportMesh($"{KVP.Value}_Mesh", MeshData, KVP.Value, out NewMesh);
-
-					if (NewMesh != null)
+					if (Exporter.ExportMesh($"{KVP.Value}_Mesh", MeshData, KVP.Value, out NewMesh))
 					{
-						CreatedMeshes.Add(NewMesh);
+						CreatedMeshes[Comp] = NewMesh;
 					}
 				}
 			});
-			// Adding stuff to a datasmith scene cannot be multithreaded!
-			foreach (Tuple<FDatasmithFacadeMeshElement, FDatasmithFacadeMesh> MeshPair in CreatedMeshes)
+
+			foreach (var KVP in CreatedMeshes)
 			{
-				DatasmithScene.AddMesh(MeshPair.Item1);
+				string MeshName = Exporter.AddMesh(KVP.Value);
+				SyncState.ComponentNameToMeshNameMap[MeshName] = KVP.Key.Name2;
 			}
 
 			SyncState.ComponentsToDelete.Clear();
diff --git a/Engine/Source/Programs/Enterprise/Datasmith/DatasmithSolidworksExporter/Private/DatasmithExporter.cs b/Engine/Source/Programs/Enterprise/Datasmith/DatasmithSolidworksExporter/Private/DatasmithExporter.cs
--- a/Engine/Source/Programs/Enterprise/Datasmith/DatasmithSolidworksExporter/Private/DatasmithExporter.cs
+++ b/Engine/Source/Programs/Enterprise/Datasmith/DatasmithSolidworksExporter/Private/DatasmithExporter.cs
@@ -116,6 +116,27 @@ namespace DatasmithSolidworks
 			return Actor;
 		}
 
+		public string AddMesh(Tuple<FDatasmithFacadeMeshElement, FDatasmithFacadeMesh> NewMesh)
+		{
+			string Name = NewMesh.Item1.GetName();
+			RemoveMesh(Name);
+			ExportedMeshesMap.TryAdd(Name, NewMesh);
+			DatasmithScene.AddMesh(NewMesh.Item1);
+			return Name;
+		}
+
+		public void RemoveMesh(string MeshName)
+		{
+            if (MeshName == null)
+            {
+                return;
+            }
+			if (ExportedMeshesMap.TryRemove(MeshName, out Tuple<FDatasmithFacadeMeshElement, FDatasmithFacadeMesh> OldMesh))
+			{
+				DatasmithScene.RemoveMesh(OldMesh.Item1);
+			}
+		}
+
 		public void RemoveActor(string InActorName)
 		{
 			if (ExportedActorsMap.ContainsKey(InActorName))
@@ -187,22 +208,22 @@ namespace DatasmithSolidworks
 			}
 		}
 
-		public string ExportMesh(string InMeshName, FMeshData InData, string InUpdateMeshActor, out Tuple<FDatasmithFacadeMeshElement, FDatasmithFacadeMesh> OutMeshPair)
+		public bool ExportMesh(string InMeshName, FMeshData InData, string InUpdateMeshActor, out Tuple<FDatasmithFacadeMeshElement, FDatasmithFacadeMesh> OutMeshPair)
 		{
+			InMeshName = SanitizeName(InMeshName);  //Compute mesh name early(it might be needed to remove old mesh)
+
 			OutMeshPair = null;
 
 			if (InData.Vertices == null || InData.Normals == null || InData.TexCoords == null || InData.Triangles == null)
 			{
-				return null;
+				return false;
 			}
 
 			if (InData.Vertices.Length == 0 || InData.Normals.Length == 0 || InData.TexCoords.Length == 0 || InData.Triangles.Length == 0)
 			{
-				return null;
+				return false;
 			}
 
-			InMeshName = SanitizeName(InMeshName);
-
 			FDatasmithFacadeMesh Mesh = new FDatasmithFacadeMesh();
 			Mesh.SetName(InMeshName);
 
@@ -266,8 +287,6 @@ namespace DatasmithSolidworks
 
 			DatasmithScene.ExportDatasmithMesh(MeshElement, Mesh);
 
-			ExportedMeshesMap.TryAdd(InMeshName, new Tuple<FDatasmithFacadeMeshElement, FDatasmithFacadeMesh>(MeshElement, Mesh));
-
 			if (!string.IsNullOrEmpty(InUpdateMeshActor))
 			{
 				Tuple<EActorType, FDatasmithFacadeActor> ExportedActorInfo = null;
@@ -278,7 +297,7 @@ namespace DatasmithSolidworks
 				}
 			}
 
-			return InMeshName;
+			return true;
 		}
 
 		public void ExportMetadata(FMetadata InMetadata)
diff --git a/Engine/Source/Programs/Enterprise/Datasmith/DatasmithSolidworksExporter/Private/PartDocument.cs b/Engine/Source/Programs/Enterprise/Datasmith/DatasmithSolidworksExporter/Private/PartDocument.cs
--- a/Engine/Source/Programs/Enterprise/Datasmith/DatasmithSolidworksExporter/Private/PartDocument.cs
+++ b/Engine/Source/Programs/Enterprise/Datasmith/DatasmithSolidworksExporter/Private/PartDocument.cs
@@ -13,6 +13,7 @@ namespace DatasmithSolidworks
 		public PartDoc SwPartDoc { get; private set; } = null;
 		private FAssemblyDocument AsmDoc = null;
 		private string ComponentName = null;
+		private string MeshName = null;
 
 		public FObjectMaterials ExportedPartMaterials { get; set; }
 
@@ -93,21 +94,13 @@ namespace DatasmithSolidworks
 			ConcurrentBag<FBody> Bodies = FBody.FetchBodies(SwPartDoc);
 			FMeshData MeshData = FStripGeometry.CreateMeshData(Bodies, ExportedPartMaterials);
 
-			string MeshName = null;
+			Exporter.RemoveMesh(MeshName);
+			MeshName = null;
+			bool bHasMesh = Exporter.ExportMesh($"{ActorName}_Mesh", MeshData, ActorName, out Tuple<FDatasmithFacadeMeshElement, FDatasmithFacadeMesh> NewMesh);
 
-			if (MeshData != null)
+			if (bHasMesh)
 			{
-				Tuple<FDatasmithFacadeMeshElement, FDatasmithFacadeMesh> NewMesh = null;
-				MeshName = Exporter.ExportMesh($"{ActorName}_Mesh", MeshData, ActorName, out NewMesh);
-
-				if (NewMesh != null)
-				{
-					DatasmithScene.AddMesh(NewMesh.Item1);
-				}
-				else
-				{
-					MeshName = null;
-				}
+				MeshName = Exporter.AddMesh(NewMesh);
 			}
 
 			FDatasmithActorExportInfo ExportInfo = new FDatasmithActorExportInfo();
@@ -115,7 +108,7 @@ namespace DatasmithSolidworks
 			ExportInfo.bVisible = true;
 			ExportInfo.Label = ActorName;
 			ExportInfo.MeshName = MeshName;
-			ExportInfo.Type = MeshName != null ? EActorType.MeshActor : EActorType.SimpleActor;
+			ExportInfo.Type = bHasMesh ? EActorType.MeshActor : EActorType.SimpleActor;
 
 			Exporter.ExportOrUpdateActor(ExportInfo);
 		}
