Fixed a miscased include path

diff --git a/Engine/Source/Programs/Enterprise/Datasmith/DatasmithMaxExporter/Private/DatasmithMaxConverters.cpp b/Engine/Source/Programs/Enterprise/Datasmith/DatasmithMaxExporter/Private/DatasmithMaxConverters.cpp
--- a/Engine/Source/Programs/Enterprise/Datasmith/DatasmithMaxExporter/Private/DatasmithMaxConverters.cpp
+++ b/Engine/Source/Programs/Enterprise/Datasmith/DatasmithMaxExporter/Private/DatasmithMaxConverters.cpp
@@ -14,7 +14,7 @@
 #include "DatasmithMaxLogger.h"
 #include "DatasmithMaxHelper.h"
 
-#include "HAL/PlatformFilemanager.h"
+#include "HAL/PlatformFileManager.h"
 #include "GenericPlatform/GenericPlatformFile.h"
 
 namespace  DatasmithMaxDirectLink
