[3ds Max Exporter] Added warning to clearly identify problematic meshes with degenerate triangles

diff --git a/Engine/Source/Programs/Enterprise/Datasmith/DatasmithMaxExporter/Private/DatasmithMaxDirectLinkMesh.cpp b/Engine/Source/Programs/Enterprise/Datasmith/DatasmithMaxExporter/Private/DatasmithMaxDirectLinkMesh.cpp
--- a/Engine/Source/Programs/Enterprise/Datasmith/DatasmithMaxExporter/Private/DatasmithMaxDirectLinkMesh.cpp
+++ b/Engine/Source/Programs/Enterprise/Datasmith/DatasmithMaxExporter/Private/DatasmithMaxDirectLinkMesh.cpp
@@ -251,8 +251,8 @@ void FillDatasmithMeshFromMaxMesh(FDatasmithMesh& DatasmithMesh, Mesh& MaxMesh,
 
 	if (!bIsFirstUVChannelValid)
 	{
-		//DatasmithMaxLogger::Get().AddGeneralError(*FString::Printf(TEXT("%s's UV channel #1 contains degenerated triangles, this can cause issues in Unreal. It is recommended to properly unfold and flatten exported UV data.")
-		//	, static_cast<const TCHAR*>(ExportedNode->GetName())));
+		LogWarning(FString::Printf(TEXT("%s's UV channel #1 contains degenerated triangles, this can cause issues in Unreal. It is recommended to properly unfold and flatten exported UV data.")
+			, static_cast<const TCHAR*>(ExportedNode->GetName())));
 	}
 
 	// Faces
