DataSmith Revit Exporter compilation fixes for pre 2023 versions

diff --git a/Engine/Source/Programs/Enterprise/Datasmith/DatasmithRevitExporter/Private/DatasmithRevitDocumentData.cs b/Engine/Source/Programs/Enterprise/Datasmith/DatasmithRevitExporter/Private/DatasmithRevitDocumentData.cs
--- a/Engine/Source/Programs/Enterprise/Datasmith/DatasmithRevitExporter/Private/DatasmithRevitDocumentData.cs
+++ b/Engine/Source/Programs/Enterprise/Datasmith/DatasmithRevitExporter/Private/DatasmithRevitDocumentData.cs
@@ -1316,7 +1316,12 @@ namespace DatasmithRevitExporter
 				{
 					BoundingBoxXYZ BBox = CurrentView3d.GetSectionBox();
 
+
+#if REVIT_API_2023
 					if (BBox.IsSet && BBox.Enabled)
+#else
+					if (BBox.Enabled)
+#endif
 					{
 						SectionBox = new OrientatedBoundingBox(BBox.Transform, BBox.Min, BBox.Max, true);
 						if (SectionBox.bIsValidData)
