Datasmith Revit Exporter Autosync trigger modification

diff --git a/Engine/Source/Programs/Enterprise/Datasmith/DatasmithRevitExporter/Private/DatasmithRevitDirectLink.cs b/Engine/Source/Programs/Enterprise/Datasmith/DatasmithRevitExporter/Private/DatasmithRevitDirectLink.cs
--- a/Engine/Source/Programs/Enterprise/Datasmith/DatasmithRevitExporter/Private/DatasmithRevitDirectLink.cs
+++ b/Engine/Source/Programs/Enterprise/Datasmith/DatasmithRevitExporter/Private/DatasmithRevitDirectLink.cs
@@ -198,7 +198,20 @@ namespace DatasmithRevitExporter
 		{
 			FDirectLink ActiveInstance = FDocument.ActiveDocument?.ActiveDirectLinkInstance ?? null;
 
-			if (ActiveInstance == null || ActiveInstance.SyncCount == 0)
+			if (ActiveInstance == null)
+			{
+				return;
+			}
+
+			//OnDocumentChanged event can happen in intervals where PostCommand is blocked/ignored by Revit
+			//so instead of executing RunAutoSync in OnDocumentChanged, we execute it here (OnApplicationIdle),
+			//OnApplicationIdle is only executed in intervals where PostCommand is NOT blocked/ignored.
+			if (ActiveInstance.bHasChanges && bAutoSync)
+			{
+				ActiveInstance.RunAutoSync();
+			}
+
+			if (ActiveInstance.SyncCount == 0)
 			{
 				return;
 			}
@@ -274,11 +287,6 @@ namespace DatasmithRevitExporter
 
 			ActiveInstance.bHasChanges = ActiveInstance.bHasChanges || InArgs.GetDeletedElementIds().Any();
 			ActiveInstance.bHasChanges = ActiveInstance.bHasChanges || InArgs.GetAddedElementIds().Any();
-
-			if (ActiveInstance.bHasChanges && bAutoSync)
-			{
-				ActiveInstance.RunAutoSync();
-			}
 		}
 
 		public FDirectLink(View3D InView, FSettings InSettings)
