Added support for ArchiCAD 26 on Windows

diff --git a/Engine/Source/Programs/Enterprise/Datasmith/DatasmithARCHICADExporter/Installer/CustomActionManaged/CustomAction.cs b/Engine/Source/Programs/Enterprise/Datasmith/DatasmithARCHICADExporter/Installer/CustomActionManaged/CustomAction.cs
--- a/Engine/Source/Programs/Enterprise/Datasmith/DatasmithARCHICADExporter/Installer/CustomActionManaged/CustomAction.cs
+++ b/Engine/Source/Programs/Enterprise/Datasmith/DatasmithARCHICADExporter/Installer/CustomActionManaged/CustomAction.cs
@@ -122,7 +122,7 @@ namespace CustomActionManaged
 		{
 			try
 			{
-				string[] MajorVersions = new string[] { "25", "24", "23", };
+				string[] MajorVersions = new string[] { "26", "25", "24", "23", };
 
 				string[] MinorVersions = new string[] { "0", };
 
diff --git a/Engine/Source/Programs/Enterprise/Datasmith/DatasmithARCHICADExporter/Private/AddOn.cpp b/Engine/Source/Programs/Enterprise/Datasmith/DatasmithARCHICADExporter/Private/AddOn.cpp
--- a/Engine/Source/Programs/Enterprise/Datasmith/DatasmithARCHICADExporter/Private/AddOn.cpp
+++ b/Engine/Source/Programs/Enterprise/Datasmith/DatasmithARCHICADExporter/Private/AddOn.cpp
@@ -16,7 +16,9 @@ DISABLE_SDK_WARNINGS_START
 #undef PI
 
 #include "Sight.hpp"
+#if AC_VERSION < 26
 #include "AttributeReader.hpp"
+#endif
 #include "Model.hpp"
 
 DISABLE_SDK_WARNINGS_END
diff --git a/Engine/Source/Programs/Enterprise/Datasmith/DatasmithARCHICADExporter/Private/Commander.cpp b/Engine/Source/Programs/Enterprise/Datasmith/DatasmithARCHICADExporter/Private/Commander.cpp
--- a/Engine/Source/Programs/Enterprise/Datasmith/DatasmithARCHICADExporter/Private/Commander.cpp
+++ b/Engine/Source/Programs/Enterprise/Datasmith/DatasmithARCHICADExporter/Private/Commander.cpp
@@ -27,7 +27,11 @@ DISABLE_SDK_WARNINGS_START
 #include "Sight.hpp"
 DISABLE_SDK_WARNINGS_END
 
+#if AC_VERSION < 26
 #include "AttributeReader.hpp"
+#else
+#include "IAttributeReader.hpp"
+#endif
 #include "Model.hpp"
 
 BEGIN_NAMESPACE_UE_AC
@@ -203,8 +207,13 @@ void FCommander::DoSnapshotOrExport(const IO::Location* InExportedFile)
 		Modeler::ConstModel3DPtr Model3D(SightPtr->GetMainModelPtr());
 
 		ModelerAPI::Model Model;
+#if AC_VERSION < 26
 		AttributeReader	  Reader; // deprecated constructor, temporary!
 		UE_AC_TestGSError(EXPGetModel(Model3D, &Model, &Reader));
+#else
+		GS::Owner<Modeler::IAttributeReader> Reader(ACAPI_Attribute_GetCurrentAttributeSetReader());
+		UE_AC_TestGSError(EXPGetModel(Model3D, &Model, Reader.Get()));
+#endif
 		if (InExportedFile)
 		{
 			FExporter().DoExport(Model, *InExportedFile);
diff --git a/Engine/Source/Programs/Enterprise/Datasmith/DatasmithARCHICADExporter/Private/ElementEvent.cpp b/Engine/Source/Programs/Enterprise/Datasmith/DatasmithARCHICADExporter/Private/ElementEvent.cpp
--- a/Engine/Source/Programs/Enterprise/Datasmith/DatasmithARCHICADExporter/Private/ElementEvent.cpp
+++ b/Engine/Source/Programs/Enterprise/Datasmith/DatasmithARCHICADExporter/Private/ElementEvent.cpp
@@ -1,6 +1,7 @@
 // Copyright Epic Games, Inc. All Rights Reserved.
 
 #include "ElementEvent.h"
+#include "Utils/ElementTools.h"
 #include "Utils/TAssValueName.h"
 #include "Synchronizer.h"
 #include "Commander.h"
@@ -67,10 +68,22 @@ GSErrCode FElementEvent::Event(const API_NotifyElementType& ElemType)
 {
 	UE_AC_TraceF("-> FElementEvent::Event(%s)\n", TAssEnumName< API_ElementDBEventID >::GetName(ElemType.notifID));
 	GS::UniString ElemTypeName;
+#if AC_VERSION < 26
 	if (ACAPI_Goodies(APIAny_GetElemTypeNameID, (void*)(size_t)ElemType.elemHead.typeID, &ElemTypeName) != NoError)
 	{
 		ElemTypeName = GS::UniString::Printf("ElementType=%d", ElemType.elemHead.typeID);
 	}
+#else
+	const API_ElemTypeID ElemTypeID = ElemType.elemHead.type.typeID;
+	if (API_FirstElemType <= ElemTypeID && ElemTypeID <= API_LastElemType)
+	{
+		ElemTypeName = GS::UniString::Printf("ElementType=%s", FElementTools::TypeName(ElemTypeID));
+	}
+	else
+	{
+		ElemTypeName = GS::UniString::Printf("ElementType=%d", ElemTypeID);
+	}
+#endif
 
 	GS::UniString ElemUUID(APIGuidToString(ElemType.elemHead.guid));
 	switch (ElemType.notifID)
diff --git a/Engine/Source/Programs/Enterprise/Datasmith/DatasmithARCHICADExporter/Private/ElementID.cpp b/Engine/Source/Programs/Enterprise/Datasmith/DatasmithARCHICADExporter/Private/ElementID.cpp
--- a/Engine/Source/Programs/Enterprise/Datasmith/DatasmithARCHICADExporter/Private/ElementID.cpp
+++ b/Engine/Source/Programs/Enterprise/Datasmith/DatasmithARCHICADExporter/Private/ElementID.cpp
@@ -253,17 +253,17 @@ void FElementID::CollectDependantElementsType(API_ElemTypeID TypeID) const
 // Connect to parent or childs
 void FElementID::HandleDepedencies() const
 {
-	if (APIElementHeader.typeID == API_WallID)
+	if (GET_HEADER_TYPEID(APIElementHeader) == API_WallID)
 	{
 		CollectDependantElementsType(API_WindowID);
 		CollectDependantElementsType(API_DoorID);
 	}
-	else if (APIElementHeader.typeID == API_RoofID || APIElementHeader.typeID == API_ShellID)
+	else if (GET_HEADER_TYPEID(APIElementHeader) == API_RoofID || GET_HEADER_TYPEID(APIElementHeader) == API_ShellID)
 	{
 		CollectDependantElementsType(API_SkylightID);
 	}
-	else if (APIElementHeader.typeID == API_WindowID || APIElementHeader.typeID == API_DoorID ||
-			 APIElementHeader.typeID == API_SkylightID)
+	else if (GET_HEADER_TYPEID(APIElementHeader) == API_WindowID || GET_HEADER_TYPEID(APIElementHeader) == API_DoorID ||
+		GET_HEADER_TYPEID(APIElementHeader) == API_SkylightID)
 	{
 		// Do nothing
 	}
diff --git a/Engine/Source/Programs/Enterprise/Datasmith/DatasmithARCHICADExporter/Private/ElementID.h b/Engine/Source/Programs/Enterprise/Datasmith/DatasmithARCHICADExporter/Private/ElementID.h
--- a/Engine/Source/Programs/Enterprise/Datasmith/DatasmithARCHICADExporter/Private/ElementID.h
+++ b/Engine/Source/Programs/Enterprise/Datasmith/DatasmithARCHICADExporter/Private/ElementID.h
@@ -71,6 +71,16 @@ class FElementID
 	// Return the element's header
 	const API_Elem_Head& GetHeader() const { return APIElementHeader; }
 
+	// Return the element's typeID
+	const API_ElemTypeID& GetTypeID() const
+	{
+#if AC_VERSION < 26
+		return APIElementHeader.typeID;
+#else
+		return APIElementHeader.type.typeID;
+#endif
+	}
+
 	// Return the mesh class based on hash of 3D ModelerAPI::BaseElemId
 	FMeshClass* GetMeshClass();
 
diff --git a/Engine/Source/Programs/Enterprise/Datasmith/DatasmithARCHICADExporter/Private/Export.cpp b/Engine/Source/Programs/Enterprise/Datasmith/DatasmithARCHICADExporter/Private/Export.cpp
--- a/Engine/Source/Programs/Enterprise/Datasmith/DatasmithARCHICADExporter/Private/Export.cpp
+++ b/Engine/Source/Programs/Enterprise/Datasmith/DatasmithARCHICADExporter/Private/Export.cpp
@@ -9,7 +9,11 @@
 #include "exp.h"
 #include "Sight.hpp"
 DISABLE_SDK_WARNINGS_START
+#if AC_VERSION < 26
 #include "AttributeReader.hpp"
+#else
+#include "IAttributeReader.hpp"
+#endif
 DISABLE_SDK_WARNINGS_END
 #include "FileSystem.hpp"
 
@@ -57,8 +61,13 @@ GSErrCode FExport::SaveDatasmithFile(const API_IOParams& IOParams, const Modeler
 
 		ModelerAPI::Model		 model;
 		Modeler::ConstModel3DPtr model3D(InSight.GetMainModelPtr());
+#if AC_VERSION < 26
 		AttributeReader			 AttrReader; // deprecated constructor, temporary!
 		UE_AC_TestGSError(EXPGetModel(model3D, &model, &AttrReader));
+#else
+		GS::Owner<Modeler::IAttributeReader> AttrReader(ACAPI_Attribute_GetCurrentAttributeSetReader());
+		UE_AC_TestGSError(EXPGetModel(model3D, &model, AttrReader.Get()));
+#endif
 
 		FExporter exporter;
 		exporter.DoExport(model, IOParams);
diff --git a/Engine/Source/Programs/Enterprise/Datasmith/DatasmithARCHICADExporter/Private/ProjectEvent.cpp b/Engine/Source/Programs/Enterprise/Datasmith/DatasmithARCHICADExporter/Private/ProjectEvent.cpp
--- a/Engine/Source/Programs/Enterprise/Datasmith/DatasmithARCHICADExporter/Private/ProjectEvent.cpp
+++ b/Engine/Source/Programs/Enterprise/Datasmith/DatasmithARCHICADExporter/Private/ProjectEvent.cpp
@@ -37,8 +37,10 @@ FAssValueName::SAssValueName TAssEnumName< API_NotifyEventID >::AssEnumName[] =
 
 	ValueName(APINotify_SideviewCreated),
 	ValueName(APINotify_SideviewRebuilt),
+#if AC_VERSION < 26
 	ValueName(APINotify_PropertyDefinitionChanged),
 	ValueName(APINotify_ClassificationItemChanged),
+#endif
 	ValueName(APINotify_PropertyVisibilityChanged),
 	ValueName(APINotify_ClassificationVisibilityChanged),
 
diff --git a/Engine/Source/Programs/Enterprise/Datasmith/DatasmithARCHICADExporter/Private/SyncData.cpp b/Engine/Source/Programs/Enterprise/Datasmith/DatasmithARCHICADExporter/Private/SyncData.cpp
--- a/Engine/Source/Programs/Enterprise/Datasmith/DatasmithARCHICADExporter/Private/SyncData.cpp
+++ b/Engine/Source/Programs/Enterprise/Datasmith/DatasmithARCHICADExporter/Private/SyncData.cpp
@@ -745,7 +745,7 @@ void FSyncData::FElement::Process(FProcessInfo* IOProcessInfo)
 			IOProcessInfo->ElementID.InitHeader(GSGuid2APIGuid(ElementId));
 			CheckModificationStamp(IOProcessInfo->ElementID.GetHeader().modiStamp);
 
-			TypeID = IOProcessInfo->ElementID.GetHeader().typeID;
+			TypeID = IOProcessInfo->ElementID.GetTypeID();
 
 			UE_AC_STAT(IOProcessInfo->SyncContext.Stats.TotalOwnerCreated++);
 			TSharedRef< IDatasmithActorElement > NewActor =
@@ -780,7 +780,7 @@ void FSyncData::FElement::Process(FProcessInfo* IOProcessInfo)
 			IOProcessInfo->ElementID.InitElement(this);
 			IOProcessInfo->ElementID.InitHeader();
 
-			TypeID = IOProcessInfo->ElementID.GetHeader().typeID;
+			TypeID = IOProcessInfo->ElementID.GetTypeID();
 
 			ModelerAPI::Transformation LocalToWorld =
 				IOProcessInfo->ElementID.GetElement3D().GetElemLocalToWorldTransformation();
@@ -873,7 +873,7 @@ bool FSyncData::FElement::AddTags(FSyncDatabase* IOSyncDatabase)
 		Zap(&APIElement);
 		APIElement.header.guid = GSGuid2APIGuid(ElementId);
 		GSErrCode GSErr = ACAPI_Element_Get(&APIElement, 0);
-		UE_AC_Assert(APIElement.header.typeID == TypeID);
+		UE_AC_Assert(GET_HEADER_TYPEID(APIElement.header) == TypeID);
 
 		if (GSErr == NoError)
 		{
@@ -1549,7 +1549,7 @@ FSyncData::FHotLinkInstance::FHotLinkInstance(const GS::Guid& InGuid, FSyncDatab
 
 	API_Element hotlinkElem;
 	Zap(&hotlinkElem);
-	hotlinkElem.header.typeID = API_HotlinkID;
+	GET_HEADER_TYPEID(hotlinkElem.header) = API_HotlinkID;
 	hotlinkElem.header.guid = GSGuid2APIGuid(ElementId);
 	GSErrCode err = ACAPI_Element_Get(&hotlinkElem);
 	if (err == NoError)
@@ -1588,7 +1588,7 @@ void FSyncData::FHotLinkInstance::Process(FProcessInfo* IOProcessInfo)
 
 		API_Element hotlinkElem;
 		Zap(&hotlinkElem);
-		hotlinkElem.header.typeID = API_HotlinkID;
+		GET_HEADER_TYPEID(hotlinkElem.header) = API_HotlinkID;
 		hotlinkElem.header.guid = GSGuid2APIGuid(ElementId);
 		GSErrCode err = ACAPI_Element_Get(&hotlinkElem);
 		if (err == NoError)
diff --git a/Engine/Source/Programs/Enterprise/Datasmith/DatasmithARCHICADExporter/Private/Utils/APIEnvir.h b/Engine/Source/Programs/Enterprise/Datasmith/DatasmithARCHICADExporter/Private/Utils/APIEnvir.h
--- a/Engine/Source/Programs/Enterprise/Datasmith/DatasmithARCHICADExporter/Private/Utils/APIEnvir.h
+++ b/Engine/Source/Programs/Enterprise/Datasmith/DatasmithARCHICADExporter/Private/Utils/APIEnvir.h
@@ -14,7 +14,9 @@ DISABLE_SDK_WARNINGS_START
 #undef PI
 
 #include "Sight.hpp"
+#if AC_VERSION < 26
 #include "AttributeReader.hpp"
+#endif
 #include "Model.hpp"
 
 DISABLE_SDK_WARNINGS_END
@@ -61,3 +63,10 @@ END_NAMESPACE_UE_AC
 /* 0 - Right value UVEdit function
  * 1 - Make texture rotation and sizing compatible with Twinmotion */
 #define PIVOT_0_5_0_5 0
+
+#if AC_VERSION < 26
+#define GET_HEADER_TYPEID(_header_) _header_.typeID
+#else
+#define GET_HEADER_TYPEID(_header_) _header_.type.typeID
+#endif
+
diff --git a/Engine/Source/Programs/Enterprise/Datasmith/DatasmithARCHICADExporter/Private/Utils/Element2String.cpp b/Engine/Source/Programs/Enterprise/Datasmith/DatasmithARCHICADExporter/Private/Utils/Element2String.cpp
--- a/Engine/Source/Programs/Enterprise/Datasmith/DatasmithARCHICADExporter/Private/Utils/Element2String.cpp
+++ b/Engine/Source/Programs/Enterprise/Datasmith/DatasmithARCHICADExporter/Private/Utils/Element2String.cpp
@@ -19,7 +19,7 @@ utf8_string GetHotLinkInfo(const API_Element& InElement)
 		{
 			API_Element HotlinkElem;
 			Zap(&HotlinkElem);
-			HotlinkElem.header.typeID = API_HotlinkID;
+			GET_HEADER_TYPEID(HotlinkElem.header) = API_HotlinkID;
 			HotlinkElem.header.guid = HotlinkElemGuid;
 			GSErr = ACAPI_Element_Get(&HotlinkElem);
 			if (GSErr == NoError)
@@ -55,7 +55,7 @@ utf8_string FElement2String::GetElementAsShortString(const API_Element& InElemen
 	FElementTools::GetInfoString(InElement.header.guid, &ElementInfo);
 
 	ElementString += Utf8StringFormat("Element \"%s\", type=%s, Guid={%s}, Layer=\"%s\", Floor=%d\n",
-									  ElementInfo.ToUtf8(), FElementTools::TypeName(InElement.header.typeID),
+									  ElementInfo.ToUtf8(), FElementTools::TypeName(GET_HEADER_TYPEID(InElement.header)),
 									  APIGuidToString(InElement.header.guid).ToUtf8(),
 									  GetLayerName(InElement.header.layer).ToUtf8(), InElement.header.floorInd);
 
@@ -124,7 +124,7 @@ utf8_string FElement2String::GetElementAsString(const API_Element& InElement)
 
 	ElementString += Utf8StringFormat(
 		"\tElement \"%s\", type=%s, Guid={%s}\n\t\t\tMemoMD5={%s} Layer=\"%s\", Floor=%d\n", ElementInfo.ToUtf8(),
-		FElementTools::TypeName(InElement.header.typeID), APIGuidToString(InElement.header.guid).ToUtf8(),
+		FElementTools::TypeName(GET_HEADER_TYPEID(InElement.header)), APIGuidToString(InElement.header.guid).ToUtf8(),
 		APIGuidToString(MemoMD5).ToUtf8(), GetLayerName(InElement.header.layer).ToUtf8(), InElement.header.floorInd);
 	utf8_string HotLinkInfo = GetHotLinkInfo(InElement);
 	if (HotLinkInfo.size())
@@ -199,7 +199,7 @@ utf8_string FElement2String::GetElementAsString(const API_Element& InElement)
 	const API_ObjectType*	   ObjectType = nullptr;
 	const API_OpeningBaseType* OpeningBase = nullptr;
 	utf8_string				   ObjectData;
-	switch (InElement.header.typeID)
+	switch (GET_HEADER_TYPEID(InElement.header))
 	{
 		case API_WallID:
 			{
@@ -286,11 +286,19 @@ utf8_string FElement2String::GetElementAsString(const API_Element& InElement)
 		{
 			ObjectData += Utf8StringFormat(", pos={%lf, %lf}", ObjectType->pos.x, ObjectType->pos.y);
 		}
+#if AC_VERSION < 26
 		if (ObjectType->ownerID || ObjectType->owner != APINULLGuid)
 		{
 			ObjectData += Utf8StringFormat(", ownerID=%d, owner={%s}", ObjectType->ownerID,
-										   APIGuidToString(ObjectType->owner).ToUtf8());
+				APIGuidToString(ObjectType->owner).ToUtf8());
+	}
+#else
+		if (ObjectType->ownerType.typeID || ObjectType->owner != APINULLGuid)
+		{
+			ObjectData += Utf8StringFormat(", ownerID=%d, owner={%s}", ObjectType->ownerType.typeID,
+				APIGuidToString(ObjectType->owner).ToUtf8());
 		}
+#endif
 	}
 	if (OpeningBase != nullptr)
 	{
diff --git a/Engine/Source/Programs/Enterprise/Datasmith/DatasmithARCHICADExporter/Private/Utils/ElementTools.cpp b/Engine/Source/Programs/Enterprise/Datasmith/DatasmithARCHICADExporter/Private/Utils/ElementTools.cpp
--- a/Engine/Source/Programs/Enterprise/Datasmith/DatasmithARCHICADExporter/Private/Utils/ElementTools.cpp
+++ b/Engine/Source/Programs/Enterprise/Datasmith/DatasmithARCHICADExporter/Private/Utils/ElementTools.cpp
@@ -89,13 +89,15 @@ FAssValueName::SAssValueName TAssEnumName< API_ElemTypeID >::AssEnumName[] = {
 	{ API_RailingRailConnectionID, 		"RailingRailConnection" },
 	{ API_RailingEndFinishID, 			"RailingEndFinish" },
 
+#if AC_VERSION < 26
 	{ API_AnalyticalSupportID, 			"AnalyticalSupport" },
 	{ API_AnalyticalLinkID, 			"AnalyticalLink" },
+#endif
 
 	{ API_BeamSegmentID, 				"BeamSegment" },
 	{ API_ColumnSegmentID, 				"ColumnSegment" },
 	{ API_OpeningID, 					"Opening" },
-#if AC_VERSION > 24
+#if AC_VERSION > 24 && AC_VERSION < 26
 	{ API_AnalyticalPointLoadID, 		"AnalyticalPointLoad" },
 	{ API_AnalyticalEdgeLoadID, 		"AnalyticalEdgeLoad" },
 	{ API_AnalyticalSurfaceLoadID, 		"AnalyticalSurfaceLoad" },
@@ -183,9 +185,9 @@ const utf8_t* FElementTools::TypeName(const API_Guid& InElementGuid)
 	{
 		UE_AC_DebugF("FElementTools::TypeName - Can't get element header {%s} Error=%d\n",
 					 APIGuidToString(InElementGuid).ToUtf8(), GSErr);
-		ElementHead.typeID = API_ZombieElemID;
+		GET_HEADER_TYPEID(ElementHead) = API_ZombieElemID;
 	}
-	return TypeName(ElementHead.typeID);
+	return TypeName(GET_HEADER_TYPEID(ElementHead));
 }
 
 // Tool: Return the variation as string
@@ -224,7 +226,7 @@ utf8_string FElementTools::GetVariationAsString(API_ElemVariationID InVariation)
 // Tool: return libpart index (or 0 if no libpart)
 GS::Int32 FElementTools::GetLibPartIndex(const API_Element& InElement)
 {
-	switch (InElement.header.typeID)
+	switch (GET_HEADER_TYPEID(InElement.header))
 	{
 		case API_WindowID:
 		case API_DoorID:
@@ -336,7 +338,7 @@ GS::Guid FElementTools::GetLibPartId(const API_Elem_Head& InElement)
 // Tool: return element's owner guid
 API_Guid FElementTools::GetOwner(const API_Element& InElement)
 {
-	size_t Offset = GetOwnerOffset(InElement.header.typeID);
+	size_t Offset = GetOwnerOffset(GET_HEADER_TYPEID(InElement.header));
 	if (Offset)
 	{
 		return *reinterpret_cast< const API_Guid* >(reinterpret_cast< const char* >(&InElement) + Offset);
@@ -450,8 +452,10 @@ size_t FElementTools::GetOwnerOffset(API_ElemTypeID InTypeID)
 		TableString[API_RailingRailConnectionID] = offsetof(API_Element, railingRailConnection.owner);
 		TableString[API_RailingEndFinishID] = offsetof(API_Element, railingEndFinish.owner);
 
+#if AC_VERSION < 26
 		TableString[API_AnalyticalSupportID] = offsetof(API_Element, analyticalSupport.head);
 		TableString[API_AnalyticalLinkID] = offsetof(API_Element, analyticalLink.head);
+#endif
 		TableString[API_ColumnSegmentID] = offsetof(API_Element, columnSegment.owner);
 		TableString[API_BeamSegmentID] = offsetof(API_Element, beamSegment.owner);
 		TableString[API_OpeningID] = offsetof(API_Element, opening.owner);
