Datasmith Revit Exporter Actor with children deletion fix

diff --git a/Engine/Source/Programs/Enterprise/Datasmith/DatasmithRevitExporter/Private/DatasmithRevitDirectLink.cs b/Engine/Source/Programs/Enterprise/Datasmith/DatasmithRevitExporter/Private/DatasmithRevitDirectLink.cs
--- a/Engine/Source/Programs/Enterprise/Datasmith/DatasmithRevitExporter/Private/DatasmithRevitDirectLink.cs
+++ b/Engine/Source/Programs/Enterprise/Datasmith/DatasmithRevitExporter/Private/DatasmithRevitDirectLink.cs
@@ -125,7 +125,7 @@ namespace DatasmithRevitExporter
 					FDocumentData.FBaseElementData ElementData = CachedElements[ElemId];
 					CachedElements.Remove(ElemId);
 					ElementData.Parent?.ChildElements.Remove(ElementData);
-					DatasmithScene.RemoveActor(ElementData.ElementActor);
+					DatasmithScene.RemoveActor(ElementData.ElementActor, FDatasmithFacadeScene.EActorRemovalRule.KeepChildrenAndKeepRelativeTransform);
 					ExportedActorsMap.Remove(ElementData.ElementActor.GetName());
 				}
 			}
