[3ds Max Exporter] Fixed texture element name sanitation

- decimal symbol in the locale was leaking into name avoiding sanitize

diff --git a/Engine/Source/Programs/Enterprise/Datasmith/DatasmithMaxExporter/Private/DatasmithMaxScanline.cpp b/Engine/Source/Programs/Enterprise/Datasmith/DatasmithMaxExporter/Private/DatasmithMaxScanline.cpp
--- a/Engine/Source/Programs/Enterprise/Datasmith/DatasmithMaxExporter/Private/DatasmithMaxScanline.cpp
+++ b/Engine/Source/Programs/Enterprise/Datasmith/DatasmithMaxExporter/Private/DatasmithMaxScanline.cpp
@@ -162,7 +162,9 @@ FString GetActualBitmapNameImpl(T* InBitmap)
 	}
 
 	float Gamma = FDatasmithMaxMatHelper::GetBitmapGamma(InBitmap);
-	FString ActualBitmapName = FDatasmithUtils::SanitizeObjectName(FPaths::GetBaseFilename(ActualBitmapPath)) + TEXT("_") + FString::SanitizeFloat(Gamma).Replace(TEXT("."), TEXT("_")) + FDatasmithMaxMatWriter::TextureSuffix;
+
+	// note: call SanitizeObjectName for stringified float - converted floating point can contain invalid character too(e.g. when decimal symbol is invalid in the locale)
+	FString ActualBitmapName = FDatasmithUtils::SanitizeObjectName(FPaths::GetBaseFilename(ActualBitmapPath) + TEXT("_") + FDatasmithUtils::SanitizeObjectName(FString::SanitizeFloat(Gamma)) + FDatasmithMaxMatWriter::TextureSuffix);
 
 	return ActualBitmapName;
 }
diff --git a/Engine/Source/Programs/Enterprise/Datasmith/DatasmithMaxExporter/Private/MaxMaterialsToUEPbr/DatasmithMaxCoronaTexmapToUEPbr.cpp b/Engine/Source/Programs/Enterprise/Datasmith/DatasmithMaxExporter/Private/MaxMaterialsToUEPbr/DatasmithMaxCoronaTexmapToUEPbr.cpp
--- a/Engine/Source/Programs/Enterprise/Datasmith/DatasmithMaxExporter/Private/MaxMaterialsToUEPbr/DatasmithMaxCoronaTexmapToUEPbr.cpp
+++ b/Engine/Source/Programs/Enterprise/Datasmith/DatasmithMaxExporter/Private/MaxMaterialsToUEPbr/DatasmithMaxCoronaTexmapToUEPbr.cpp
@@ -274,8 +274,8 @@ IDatasmithMaterialExpression* FDatasmithMaxCoronalBitmapToUEPbr::Convert( FDatas
 {
 	DatasmithMaxTexmapParser::FCoronaBitmapParameters CoronaBitmapParameters = DatasmithMaxTexmapParser::ParseCoronaBitmap( InTexmap );
 
-	FString ActualBitmapName = FDatasmithUtils::SanitizeObjectName(FPaths::GetBaseFilename(CoronaBitmapParameters.Path)) + TEXT("_") +
-		FString::SanitizeFloat(CoronaBitmapParameters.Gamma).Replace(TEXT("."), TEXT("_")) + FDatasmithMaxMatWriter::TextureSuffix;
+	FString ActualBitmapName = FDatasmithUtils::SanitizeObjectName(FPaths::GetBaseFilename(CoronaBitmapParameters.Path) + TEXT("_") +
+		FString::SanitizeFloat(CoronaBitmapParameters.Gamma) + FDatasmithMaxMatWriter::TextureSuffix);
 
 	const bool bIsSRGB = false;
 	const bool bUseAlphaAsMono = false;
