[Solidworks] Lights and Assembly configuration fixes

- Suppressed parts with geometry are exported as mesh actors(instead of empty actors as before). Temporarily unsuppressing those parts for the duration of export and reverting suppression state back.

- Fixed name sanitation which was breaking some scenes with variants.

- Light Label uses user name of the light

- fixed spotlight direction

diff --git a/Engine/Source/Programs/Enterprise/Datasmith/DatasmithSolidworksExporter/Private/AssemblyDocument.cs b/Engine/Source/Programs/Enterprise/Datasmith/DatasmithSolidworksExporter/Private/AssemblyDocument.cs
--- a/Engine/Source/Programs/Enterprise/Datasmith/DatasmithSolidworksExporter/Private/AssemblyDocument.cs
+++ b/Engine/Source/Programs/Enterprise/Datasmith/DatasmithSolidworksExporter/Private/AssemblyDocument.cs
@@ -85,7 +85,10 @@ namespace DatasmithSolidworks
 			// actor hierarchy has been exported
 			Dictionary<Component2, string> MeshesToExportMap = new Dictionary<Component2, string>();
 		
-			ExportComponentRecursive(Root, null, ref MeshesToExportMap);
+			// Store original component suppression state to revert it after export
+			// Component needs to be fully resolved to have its data loaded and even suppressed components are needed for configurations export(in case they are visible in other configs)
+			List<Tuple<Component2, int>> ComponentsSuppression = new List<Tuple<Component2, int>>();
+			ExportComponentRecursive(Root, null, MeshesToExportMap, ComponentsSuppression);
 
 			// Export materials
 			SetExportStatus($"Component Materials");
@@ -160,6 +163,11 @@ namespace DatasmithSolidworks
 					}
 				}
 			}
+
+			foreach(var ComponentSuppression in ComponentsSuppression)
+			{
+				ComponentSuppression.Item1.SetSuppression2(ComponentSuppression.Item2);
+			}
 		}
 
 		public override bool HasMaterialUpdates()
@@ -275,7 +283,7 @@ namespace DatasmithSolidworks
 			return DatasmithTransform;
 		}
 
-		private void ExportComponentRecursive(Component2 InComponent, Component2 InParent, ref Dictionary<Component2, string> OutMeshesToExportMap)
+		private void ExportComponentRecursive(Component2 InComponent, Component2 InParent, Dictionary<Component2, string> OutMeshesToExportMap, List<Tuple<Component2, int>> OutComponentSuppression)
 		{
 			bool bHasDirtyTransform = false;
 
@@ -308,6 +316,15 @@ namespace DatasmithSolidworks
 
 				SyncState.ComponentsTransformsMap[InComponent.Name2] = ActorExportInfo.Transform;
 
+				// When exporting configurations we need to load all components to make their data available for export
+				// Suppressed or lightweight components don't have this
+				int Suppression = InComponent.GetSuppression();
+				if (bHasConfigurations && (Suppression !=  (int)swComponentSuppressionState_e.swComponentFullyResolved))
+				{
+					InComponent.SetSuppression2((int)swComponentSuppressionState_e.swComponentFullyResolved);
+					OutComponentSuppression.Add(new Tuple<Component2, int>(InComponent, Suppression));
+				}
+
 				if (!InComponent.IsSuppressed())
 				{
 					dynamic ComponentVisibility = InComponent.GetVisibilityInAsmDisplayStates((int)swDisplayStateOpts_e.swThisDisplayState, null);
@@ -327,7 +344,9 @@ namespace DatasmithSolidworks
 
 				bool bNeedsGeometryExport = false;
 
-				if (InComponent.GetModelDoc2() is PartDoc) 
+				// ComponentDoc is null if component is suppressed or lightweight
+				ModelDoc2 ComponentDoc = InComponent.GetModelDoc2();
+				if (ComponentDoc is PartDoc) 
 				{
 					bool bFirstExport = !SyncState.ExportedComponentsMap.ContainsKey(InComponent.Name2);
 					bNeedsGeometryExport = bFirstExport;
@@ -344,10 +363,8 @@ namespace DatasmithSolidworks
 
 				if (bNeedsGeometryExport)
 				{
-					object ComponentDoc = (object)InComponent.GetModelDoc2();
-
 					//TODO this will be null for new part, think of more solid solution
-					string PartPath = (ComponentDoc as ModelDoc2).GetPathName();
+					string PartPath = ComponentDoc.GetPathName();
 					if (!SyncState.PartsMap.ContainsKey(InComponent.Name2))
 					{
 						// New part
@@ -379,7 +396,7 @@ namespace DatasmithSolidworks
 				foreach (object Obj in Children)
 				{
 					Component2 Child = (Component2)Obj;
-					ExportComponentRecursive(Child, InComponent, ref OutMeshesToExportMap);
+					ExportComponentRecursive(Child, InComponent, OutMeshesToExportMap, OutComponentSuppression);
 				}
 			}
 		}
diff --git a/Engine/Source/Programs/Enterprise/Datasmith/DatasmithSolidworksExporter/Private/DatasmithExporter.cs b/Engine/Source/Programs/Enterprise/Datasmith/DatasmithSolidworksExporter/Private/DatasmithExporter.cs
--- a/Engine/Source/Programs/Enterprise/Datasmith/DatasmithSolidworksExporter/Private/DatasmithExporter.cs
+++ b/Engine/Source/Programs/Enterprise/Datasmith/DatasmithSolidworksExporter/Private/DatasmithExporter.cs
@@ -151,7 +151,7 @@ namespace DatasmithSolidworks
 		public void ExportLight(FLight InLight)
 		{
 			FDatasmithActorExportInfo ExportInfo = new FDatasmithActorExportInfo();
-			ExportInfo.Label = InLight.LightName;
+			ExportInfo.Label = InLight.LightLabel;
 			ExportInfo.Name = InLight.LightName;
 			ExportInfo.bVisible = true;
 
@@ -175,7 +175,7 @@ namespace DatasmithSolidworks
 				case FLight.EType.Spot:
 				{
 					LightPosition = InLight.SpotLightPosition;
-					LightDirection = (InLight.SpotLightTarget - InLight.SpotLightPosition).Normalized();
+					LightDirection = (InLight.SpotLightPosition - InLight.SpotLightTarget ).Normalized();  // Inverted direction
 					ExportInfo.Type = EActorType.SpotLightActor;
 				}
 				break;
@@ -338,10 +338,12 @@ namespace DatasmithSolidworks
 
 		FDatasmithFacadeActorBinding GetActorBinding(string InActorName, FDatasmithFacadeVariant InVariant)
 		{
+			string ActorNameSanitized = FDatasmithExporter.SanitizeName(InActorName);
+
 			for (int BindingIndex = 0; BindingIndex < InVariant.GetActorBindingsCount(); ++BindingIndex)
 			{
 				FDatasmithFacadeActorBinding Binding = InVariant.GetActorBinding(BindingIndex);
-				if (Binding.GetName() == InActorName)
+				if (Binding.GetName() == ActorNameSanitized)
 				{
 					return Binding;
 				}
@@ -351,9 +353,8 @@ namespace DatasmithSolidworks
 
 			FDatasmithFacadeActor Actor = null;
 
-			if (ExportedActorsMap.ContainsKey(InActorName))
+			if (ExportedActorsMap.TryGetValue(ActorNameSanitized, out Tuple<EActorType, FDatasmithFacadeActor> ActorInfo))
 			{
-				Tuple<EActorType, FDatasmithFacadeActor> ActorInfo = ExportedActorsMap[InActorName];
 				Actor = ActorInfo.Item2;
 			}
 			else
diff --git a/Engine/Source/Programs/Enterprise/Datasmith/DatasmithSolidworksExporter/Private/Document.cs b/Engine/Source/Programs/Enterprise/Datasmith/DatasmithSolidworksExporter/Private/Document.cs
--- a/Engine/Source/Programs/Enterprise/Datasmith/DatasmithSolidworksExporter/Private/Document.cs
+++ b/Engine/Source/Programs/Enterprise/Datasmith/DatasmithSolidworksExporter/Private/Document.cs
@@ -19,6 +19,7 @@ namespace DatasmithSolidworks
 		public bool bDirectLinkAutoSync { get; set; } = false;
 		public int DirectLinkSyncCount { get; private set; } = 0;
 		public ConcurrentDictionary<int, FMaterial> ExportedMaterialsMap = new ConcurrentDictionary<int, FMaterial>();
+		public bool bHasConfigurations = false;
 		
 		protected FDatasmithFacadeScene DatasmithScene = null;
 		protected FDatasmithExporter Exporter = null;
@@ -35,6 +36,7 @@ namespace DatasmithSolidworks
 		private bool bExitMaterialUpdateThread = false;
 		private uint FaceCounter = 1; // Face Id generator
 
+
 		public FDocument(int InDocId, ModelDoc2 InSwDoc, FDatasmithExporter InExporter)
 		{
 			DocId = InDocId;
@@ -99,7 +101,7 @@ namespace DatasmithSolidworks
 			}
 		}
 
-		private void ExportConfigurations()
+		private void ExportConfigurations(List<FConfigurationData> Configs)
 		{
 			// Remove any existing Datasmith LevelVariantSets as we'll re-add data
 			while (DatasmithScene.GetLevelVariantSetsCount() > 0)
@@ -110,8 +112,6 @@ namespace DatasmithSolidworks
 				}
 			}
 
-			List<FConfigurationData> Configs = FConfigurationExporter.ExportConfigurations(this);
-
 			Exporter.ExportMaterials(ExportedMaterialsMap);
 
 			if (Configs != null)
@@ -169,7 +169,6 @@ namespace DatasmithSolidworks
 #if DEBUG
 				Stopwatch Watch = Stopwatch.StartNew();
 #endif
-
 				ExportToDatasmithScene();
 
 				ExportLights();
@@ -201,11 +200,13 @@ namespace DatasmithSolidworks
 
 				Exporter = new FDatasmithExporter(DatasmithScene);
 
-				ExportToDatasmithScene();
+				List<FConfigurationData> Configs = FConfigurationExporter.ExportConfigurations(this);
+				bHasConfigurations = (Configs != null) && (Configs.Count != 0);
 
+				ExportToDatasmithScene();
 			
 				ExportLights();
-				ExportConfigurations();
+				ExportConfigurations(Configs);
 
 				DatasmithScene.PreExport();
 				DatasmithScene.ExportScene(DatasmithFileExportPath);
diff --git a/Engine/Source/Programs/Enterprise/Datasmith/DatasmithSolidworksExporter/Private/Lights.cs b/Engine/Source/Programs/Enterprise/Datasmith/DatasmithSolidworksExporter/Private/Lights.cs
--- a/Engine/Source/Programs/Enterprise/Datasmith/DatasmithSolidworksExporter/Private/Lights.cs
+++ b/Engine/Source/Programs/Enterprise/Datasmith/DatasmithSolidworksExporter/Private/Lights.cs
@@ -18,6 +18,7 @@ namespace DatasmithSolidworks
 		public EType LightType { get; private set; }
 
 		public string LightName { get; private set; }
+		public string LightLabel { get; private set; }
 
 		public bool bIsEnabled { get; private set; }
 
@@ -62,10 +63,11 @@ namespace DatasmithSolidworks
 			}
 		}
 
-		private FLight(EType InType, string InName, bool bInIsEnabled, double InIntensisty, int InColorref)
+		private FLight(EType InType, string InName, string InLabel, bool bInIsEnabled, double InIntensisty, int InColorref)
 		{
 			LightType = InType;
 			LightName = InName;
+			LightLabel = InLabel;
 			bIsEnabled = bInIsEnabled;
 			Intensity = (float)InIntensisty;
 
@@ -80,9 +82,9 @@ namespace DatasmithSolidworks
 		// Specific properties
 		private float[] Properties = new float[7];
 	
-		static public FLight MakeSpotLight(string InName, bool bInIsEnabled, FVec3 InPosition, FVec3 InTarget, double InConeAngle, double InIntensity, int InColor)
+		static public FLight MakeSpotLight(string InName, string InLabel, bool bInIsEnabled, FVec3 InPosition, FVec3 InTarget, double InConeAngle, double InIntensity, int InColor)
 		{
-			FLight SpotLight = new FLight(EType.Spot, InName, bInIsEnabled, InIntensity, InColor);
+			FLight SpotLight = new FLight(EType.Spot, InName, InLabel, bInIsEnabled, InIntensity, InColor);
 
 			// Position
 			SpotLight.Properties[0] = InPosition.X;
@@ -98,9 +100,9 @@ namespace DatasmithSolidworks
 			return SpotLight;
 		}
 
-		static public FLight MakePointLight(string InName, bool bInIsEnabled, FVec3 InPosition, double InIntensity, int InColor)
+		static public FLight MakePointLight(string InName, string InLabel, bool bInIsEnabled, FVec3 InPosition, double InIntensity, int InColor)
 		{
-			FLight PointLight = new FLight(EType.Point, InName, bInIsEnabled, InIntensity, InColor);
+			FLight PointLight = new FLight(EType.Point, InName, InLabel, bInIsEnabled, InIntensity, InColor);
 
 			// Position
 			PointLight.Properties[0] = InPosition.X;
@@ -110,9 +112,9 @@ namespace DatasmithSolidworks
 			return PointLight;
 		}
 
-		static public FLight MakeDirLight(string InName, bool bInIsEnabled, FVec3 InDirection, double InIntensity, int InColor)
+		static public FLight MakeDirLight(string InName, string InLabel, bool bInIsEnabled, FVec3 InDirection, double InIntensity, int InColor)
 		{
-			FLight DirLight = new FLight(EType.Directional, InName, bInIsEnabled, InIntensity, InColor);
+			FLight DirLight = new FLight(EType.Directional, InName, InLabel, bInIsEnabled, InIntensity, InColor);
 
 			// Position
 			DirLight.Properties[0] = InDirection.X;
@@ -153,6 +155,7 @@ namespace DatasmithSolidworks
 			for (int LightIndex = 0; LightIndex < InDoc.GetLightSourceCount(); ++LightIndex)
 			{
 				string LightName = InDoc.GetLightSourceName(LightIndex);
+				string LightLabel = InDoc.LightSourceUserName[LightIndex];
 				double[] Props = InDoc.LightSourcePropertyValues[LightIndex] as double[];
 
 				int LightType;
@@ -177,7 +180,8 @@ namespace DatasmithSolidworks
 
 						if (InDoc.GetSpotlightProperties(LightName, ref AmbientIntensity, ref DiffuseIntensity, ref SpecularIntensity, ref Color, ref bIsEnabled, ref bIsFixed, ref PosX, ref PosY, ref PosZ, ref DirX, ref DirY, ref DirZ, ref ConeAngle))
 						{
-							Light = FLight.MakeSpotLight(LightName, bIsEnabled, new FVec3(PosX, PosY, PosZ), new FVec3(DirX, DirY, DirZ), ConeAngle, DiffuseIntensity, Color);
+							Light = FLight.MakeSpotLight(LightName, LightLabel, bIsEnabled, new FVec3(PosX, PosY, PosZ), new FVec3(DirX, DirY, DirZ), ConeAngle, DiffuseIntensity, Color);
+							double[] Props2 = InDoc.LightSourcePropertyValues[LightIndex] as double[];
 						}
 					}
 					break;
@@ -187,7 +191,7 @@ namespace DatasmithSolidworks
 
 						if (InDoc.GetPointLightProperties(LightName, ref AmbientIntensity, ref DiffuseIntensity, ref SpecularIntensity, ref Color, ref bIsEnabled, ref bIsFixed, ref PosX, ref PosY, ref PosZ))
 						{
-							Light = FLight.MakePointLight(LightName, bIsEnabled, new FVec3(PosX, PosY, PosZ), DiffuseIntensity, Color);
+							Light = FLight.MakePointLight(LightName, LightLabel, bIsEnabled, new FVec3(PosX, PosY, PosZ), DiffuseIntensity, Color);
 						}
 					}
 					break;
@@ -196,7 +200,7 @@ namespace DatasmithSolidworks
 						double TargetX = 0.0, TargetY = 0.0, TargetZ = 0.0;
 						if (InDoc.GetDirectionLightProperties(LightName, ref AmbientIntensity, ref DiffuseIntensity, ref SpecularIntensity, ref Color, ref bIsEnabled, ref bIsFixed, ref TargetX, ref TargetY, ref TargetZ))
 						{
-							Light = FLight.MakeDirLight(LightName, bIsEnabled, new FVec3(TargetX, TargetY, TargetZ), DiffuseIntensity, Color);
+							Light = FLight.MakeDirLight(LightName, LightLabel, bIsEnabled, new FVec3(TargetX, TargetY, TargetZ), DiffuseIntensity, Color);
 						}
 					}
 					break;
