Datasmith Sketchup Mac - Adding support for Apple Silicon for Datasmith Sketchup 2022

The Datasmith Sketchup 2022 Mac plugin is now compiled and packaged as a Universal Binaries (x64_86/Arm64) and should work on either Intel or Apple Silicon mac,  without using Rosetta.
Only the Sketchup 2022 SDK is provided as universal binaries,  so other versions of our exporter are still x86_64 only and do require Rosetta to run on Apple Silicon.

This should solve some stability issues Apple Silicon users experience while using Sketchup with Rosetta enabled.

diff --git a/Engine/Source/Programs/Enterprise/Datasmith/DatasmithSketchUpRubyExporter/Build/build_mac.py b/Engine/Source/Programs/Enterprise/Datasmith/DatasmithSketchUpRubyExporter/Build/build_mac.py
--- a/Engine/Source/Programs/Enterprise/Datasmith/DatasmithSketchUpRubyExporter/Build/build_mac.py
+++ b/Engine/Source/Programs/Enterprise/Datasmith/DatasmithSketchUpRubyExporter/Build/build_mac.py
@@ -12,6 +12,7 @@ parser = argparse.ArgumentParser(fromfile_prefix_chars='@')
 parser.add_argument('--sdks-root', type=Path, required=True)
 parser.add_argument('--sketchup-version', required=True)
 parser.add_argument('--sketchup-sdk-version', required=True)
+parser.add_argument('--target_archs', nargs='+', required=True)
 parser.add_argument('--datasmithsdk-lib', type=Path, required=True)
 parser.add_argument('--output-path', type=Path, required=True)
 parser.add_argument('--intermediate-path', type=Path, required=True)
@@ -33,6 +34,13 @@ assert UE_Engine.parts[-1] == 'Engine', UE_Engine
 SU_SDK_VERSION = args.sketchup_sdk_version
 SU_VERSION = args.sketchup_version
 
+TARGET_ARCHS = args.target_archs;
+
+arch_target_dict = {
+    'x86_64': 'x86_64-apple-macos10.15',
+    'arm64': 'arm64-apple-macos11'
+}
+
 output_path = args.output_path/'Plugin/UnrealDatasmithSketchUp'
 obj_dir = args.intermediate_path
 ###########
@@ -65,7 +73,7 @@ def parse_paths(paths_string):
                 assert c == ' '
 
 def add_include_paths(cmd, paths):
-    for p in include_paths:
+    for p in paths:
         cmd.append('--include-directory')
         cmd.append(f'{p}')
         if not os.path.isdir(p):
@@ -119,110 +127,153 @@ UE_Imath = f'{UE_ThirdParty}/Imath/Deploy/Imath-3.1.3'
 
 UE_INCLUDES_PATH = f'''"{UE_Core}" "{UE_Core}/Internationalization" "{UE_Core}/Async" "{UE_Core}/Containers" "{UE_Core}/Delegates" "{UE_Core}/GenericPlatform" "{UE_Core}/HAL" "{UE_Core}/Logging" "{UE_Core}/Math" "{UE_Core}/Misc" "{UE_Core}/Modules" "{UE_Core}/Modules/Boilerplate" "{UE_Core}/ProfilingDebugging" "{UE_Core}/Serialization" "{UE_Core}/Serialization/Csv" "{UE_Core}/Stats" "{UE_Core}/Templates" "{UE_Core}/UObject" "{UE_Runtime}/CoreUObject/Public" "{UE_Project}/Interfaces" "{UE_Runtime}/TraceLog/Public"  "{UE_Runtime}/Datasmith/DatasmithCore/Public" "{UE_Runtime}/Datasmith/DirectLink/Public" "{UE_Runtime}/Messaging/Public" "{UE_Runtime}/Launch/Resources" "{UE_Engine}/Source/Developer/Datasmith/DatasmithExporter/Public" "{UE_Engine}/Source/Developer/Datasmith/DatasmithExporterUI/Public"'''
 
-CXX = ['xcrun', 'clang', 
-    '-x', 'objective-c++', 
-    '-arch', 'x86_64',
-    '-std=c++17',
-    '-stdlib=libc++',
-#    '-isysroot', '/Applications/Xcode.app/Contents/Developer/Platforms/MacOSX.platform/Developer/SDKs/MacOSX11.1.sdk',
-    '-target', 'x86_64-apple-macos10.15',
-    ]
-
-CXX_FLAGS = []
-
-CXX_FLAGS += [
-    '-O2', # using -Os or -O0 creates symbol relocations in .o files that fails linking SketchUp api, i.e. ld: illegal text-relocation to '_SUSceneGetName'... for architecture x86_64
-]
-
-CXX_FLAGS += [
-    # just adding
-    '-fno-common',
-    '-fmodules',
-    '-fpascal-strings',
-    '-fvisibility=hidden',
-    '-fstrict-aliasing',
-    '-fasm-blocks',
-    ]
-# todo: Debug/release
-
-include_paths = []
-include_paths += parse_paths(UE_INCLUDES_PATH)
-include_paths.append(f'{UE_Imath}/include')
-include_paths.append(f'{SU_SDK_PATH}/samples/common/ThirdParty/ruby/include/mac')
-include_paths.append(f'{SU_SDK_PATH}/samples/common/ThirdParty/ruby/include/mac/x86_64-darwin') # todo: might not be needed - sample uses x86_64-darwin17 which doesn't exist
-
 # Framework search path needed to lookup SketchUpApi and Ruby
 framework_search_paths = []
 framework_search_paths.append(f'{SU_SDK_PATH}')
 framework_search_paths.append(f'{SU_SDK_PATH}/samples/common/ThirdParty/ruby/lib/mac')
 
-definitions = [f'SKP_SDK_{SU_VERSION}', 'macintosh=1', '__register_unsupported__']
-UE_PREPROCESSOR_DEFINITIONS = 'UE_BUILD_DEVELOPMENT=1 UE_BUILD_MINIMAL=1 WITH_EDITOR=0 WITH_EDITORONLY_DATA=0 WITH_SERVER_CODE=1 WITH_ENGINE=0 WITH_UNREAL_DEVELOPER_TOOLS=0 WITH_PLUGIN_SUPPORT=0 IS_MONOLITHIC=1 IS_PROGRAM=1 PLATFORM_MAC=1 PLATFORM_APPLE=1 UE_BUILD_DEVELOPMENT_WITH_DEBUGGAME=0 UBT_COMPILED_PLATFORM=Mac CORE_API=DLLIMPORT COREUOBJECT_API=DLLIMPORT DATASMITHEXPORTER_API=DLLIMPORT DATASMITHCORE_API=DLLIMPORT DIRECTLINK_API=DLLIMPORT DATASMITHEXPORTERUI_API=DLLIMPORT'
-definitions += UE_PREPROCESSOR_DEFINITIONS.split()
+def create_compiler_for_arch(arch):
+    CXX = ['xcrun', 'clang', 
+    '-x', 'objective-c++', 
+    '-arch', arch,
+    '-std=c++17',
+    '-stdlib=libc++',
+    #    '-isysroot', '/Applications/Xcode.app/Contents/Developer/Platforms/MacOSX.platform/Developer/SDKs/MacOSX11.1.sdk',
+    '-target', arch_target_dict[arch],
+    ]
 
-compile_cmd_template = CXX + CXX_FLAGS
-add_include_paths(compile_cmd_template, include_paths)            
-add_framework_search_paths(compile_cmd_template, framework_search_paths)
-for d in definitions:
-    compile_cmd_template.append(f'-D{d}')
 
-compiler = Compiler(compile_cmd_template)
+    CXX_FLAGS = []
 
-# Compile source files
-plugin_private_dir = plugin_src_root/'Private'
-exporter_source_files = [p.parts[-1] for p in plugin_private_dir.iterdir() if p.suffix == '.cpp']
+    CXX_FLAGS += [
+        '-O2', # using -Os or -O0 creates symbol relocations in .o files that fails linking SketchUp api, i.e. ld: illegal text-relocation to '_SUSceneGetName'... for architecture x86_64
+    ]
 
-log_debug(exporter_source_files)
-sources = {'Private': (plugin_private_dir, exporter_source_files)}
-for name, (root, files) in sources.items():
-    group_obj_dir = obj_dir/name	
-    if not args.dry_run:
-        group_obj_dir.mkdir(parents=True, exist_ok=True)
-    for f in files:
-        compiler.compile_source(Path(root)/f, (group_obj_dir/f).with_suffix('.o'))
-compiler.done()
+    CXX_FLAGS += [
+        # just adding
+        '-fno-common',
+        '-fmodules',
+        '-fpascal-strings',
+        '-fvisibility=hidden',
+        '-fstrict-aliasing',
+        '-fasm-blocks',
+        ]
+    # todo: Debug/release
+
+    include_paths = []
+    include_paths += parse_paths(UE_INCLUDES_PATH)
+    include_paths.append(f'{UE_Imath}/include')
+    include_paths.append(f'{SU_SDK_PATH}/samples/common/ThirdParty/ruby/include/mac')
+    include_paths.append(f'{SU_SDK_PATH}/samples/common/ThirdParty/ruby/include/mac/{arch}-darwin') # todo: might not be needed - sample uses x86_64-darwin17 which doesn't exist
+
+
+    definitions = [f'SKP_SDK_{SU_VERSION}', 'macintosh=1', '__register_unsupported__']
+    UE_PREPROCESSOR_DEFINITIONS = 'UE_BUILD_DEVELOPMENT=1 UE_BUILD_MINIMAL=1 WITH_EDITOR=0 WITH_EDITORONLY_DATA=0 WITH_SERVER_CODE=1 WITH_ENGINE=0 WITH_UNREAL_DEVELOPER_TOOLS=0 WITH_PLUGIN_SUPPORT=0 IS_MONOLITHIC=1 IS_PROGRAM=1 PLATFORM_MAC=1 PLATFORM_APPLE=1 UE_BUILD_DEVELOPMENT_WITH_DEBUGGAME=0 UBT_COMPILED_PLATFORM=Mac CORE_API=DLLIMPORT COREUOBJECT_API=DLLIMPORT DATASMITHEXPORTER_API=DLLIMPORT DATASMITHCORE_API=DLLIMPORT DIRECTLINK_API=DLLIMPORT DATASMITHEXPORTERUI_API=DLLIMPORT'
+    definitions += UE_PREPROCESSOR_DEFINITIONS.split()
+
+    compile_cmd_template = CXX + CXX_FLAGS
+    add_include_paths(compile_cmd_template, include_paths)            
+    add_framework_search_paths(compile_cmd_template, framework_search_paths)
+    for d in definitions:
+        compile_cmd_template.append(f'-D{d}')
+
+    return Compiler(compile_cmd_template)
+
+arch_compiler_dict = {}
+
+for arch in TARGET_ARCHS:
+
+    target = arch_target_dict.get(arch, None)
+
+    if target is None:
+        log_warn('Cannot compile unknown architecture: '+ arch)
+        continue
+
+    compiler = create_compiler_for_arch(arch)
+    arch_compiler_dict[arch] = compiler    
+
+    # Compile source files
+    plugin_private_dir = plugin_src_root/'Private'
+    exporter_source_files = [p.parts[-1] for p in plugin_private_dir.iterdir() if p.suffix == '.cpp']
+
+    log_debug(exporter_source_files)
+    sources = {'Private': (plugin_private_dir, exporter_source_files)}
+    for name, (root, files) in sources.items():
+        group_obj_dir = obj_dir/arch/name	
+        if not args.dry_run:
+            group_obj_dir.mkdir(parents=True, exist_ok=True)
+        for f in files:
+            compiler.compile_source(Path(root)/f, (group_obj_dir/f).with_suffix('.o'))
+    compiler.done()
 ########
 
 # Link the bundle
 DatasmithSDKlib = args.datasmithsdk_lib
 assert os.path.isfile(DatasmithSDKlib), DatasmithSDKlib
 
-link_cmd = [
-    'xcrun',
-    'clang++', 
-    '-target', 'x86_64-apple-macos10.15',
-    '-Xlinker', '-export_dynamic',
-    '-Xlinker', '-no_deduplicate',
-    '-stdlib=libc++',
-    '-demangle',
-    '-dynamic',
-    ]
+def get_output_path(arch):
+    return str(output_path/('DatasmithSketchUp.bundle')) if arch is None else str(output_path/('DatasmithSketchUp_' + arch + '.bundle'))
 
-link_cmd.append('-bundle') # link ruby extension as a MacOs 'bundle' not dynamic lib
+def link(arch, target):
+    link_cmd = [
+        'xcrun',
+        'clang++', 
+        '-target', target,
+        '-Xlinker', '-export_dynamic',
+        '-Xlinker', '-no_deduplicate',
+        '-stdlib=libc++',
+        '-demangle',
+        '-dynamic',
+        ]
 
-#link_cmd += ['-framework', 'Ruby'] # link to RUby framework (found by framework_search_paths)
+    link_cmd.append('-bundle') # link ruby extension as a MacOs 'bundle' not dynamic lib
 
-# Handle not linking against SketchUpAPI and Ruby frameworks. Undefined will be solved at runtime
-link_cmd += [
-    '-undefined', 'dynamic_lookup', # We are referencing SketchUp API symbols that are defined in the SketchUp app
-    # looks like the following(referencing SketchUp app) is not needed when we have the above(-undefined dynamic_lookup)
-    # Although it's useful to uncomment to catch symbols that are missing and won't be found in the SketchUp app
-    #'-bundle_loader', f'/Applications/SketchUp {SU_VERSION}/SketchUp.app/Contents/MacOS/SketchUp', # this seems enough to resolve SU api symbols during linking(no need for -F)
-    ]
+    #link_cmd += ['-framework', 'Ruby'] # link to RUby framework (found by framework_search_paths)
 
-link_cmd += compiler.objs
-link_cmd += [DatasmithSDKlib] # link DatasmithSDK by full path(no 'lib' prefix to use search with -l)
+    # Handle not linking against SketchUpAPI and Ruby frameworks. Undefined will be solved at runtime
+    link_cmd += [
+        '-undefined', 'dynamic_lookup', # We are referencing SketchUp API symbols that are defined in the SketchUp app
+        # looks like the following(referencing SketchUp app) is not needed when we have the above(-undefined dynamic_lookup)
+        # Although it's useful to uncomment to catch symbols that are missing and won't be found in the SketchUp app
+        #'-bundle_loader', f'/Applications/SketchUp {SU_VERSION}/SketchUp.app/Contents/MacOS/SketchUp', # this seems enough to resolve SU api symbols during linking(no need for -F)
+        ]
 
-link_cmd += ['-o', str(output_path/'DatasmithSketchUp.bundle')]
+    link_cmd += arch_compiler_dict[arch].objs
+    link_cmd += [DatasmithSDKlib] # link DatasmithSDK by full path(no 'lib' prefix to use search with -l)
 
-add_framework_search_paths(link_cmd, framework_search_paths)
+    link_cmd += ['-o', get_output_path(arch)]
 
-if args.verbose:
-    log_info(link_cmd)
-if not args.dry_run:
-    output_path.mkdir(parents=True, exist_ok=True)
-    obj_dir.mkdir(parents=True, exist_ok=True)
-    subprocess.check_call(link_cmd)
+    add_framework_search_paths(link_cmd, framework_search_paths)
+
+    if args.verbose:
+        log_info(link_cmd)
+    if not args.dry_run:
+        output_path.mkdir(parents=True, exist_ok=True)
+        obj_dir.mkdir(parents=True, exist_ok=True)
+        subprocess.check_call(link_cmd)
+
+for arch in TARGET_ARCHS:
+    target = arch_target_dict.get(arch, None)
+
+    if target is not None:
+        link(arch, target)
+    else:
+        log_warn('Cannot link unknown architecture: '+ arch)
+
+if len(TARGET_ARCHS) > 1:
+    bundles = []
+    for arch in TARGET_ARCHS:
+        bundles.append(get_output_path(arch))
+    # using the lipo utility tool to merge multiple binaries into a single universal binay.
+    lipo_cmd = ['lipo', '-create', '-output', get_output_path(None)] + bundles
+    subprocess.check_call(lipo_cmd)
+
+    # remove the source bundles to clear up the directory.
+    for bundle in bundles:
+       os.remove(bundle)
+elif len(TARGET_ARCHS) == 1:
+    # remove the architecture from the bundle name, only a single bundle was generated, no need to make a distinction.
+    os.rename(get_output_path(TARGET_ARCHS[0]),get_output_path(None))
 
 if args.verbose:
     log_debug('SUCCESS!')
diff --git a/Engine/Source/Programs/Enterprise/Datasmith/DatasmithSketchUpRubyExporter/Build/build_mac.sh b/Engine/Source/Programs/Enterprise/Datasmith/DatasmithSketchUpRubyExporter/Build/build_mac.sh
--- a/Engine/Source/Programs/Enterprise/Datasmith/DatasmithSketchUpRubyExporter/Build/build_mac.sh
+++ b/Engine/Source/Programs/Enterprise/Datasmith/DatasmithSketchUpRubyExporter/Build/build_mac.sh
@@ -19,46 +19,59 @@ popd
 mkdir -p /tmp/Python3/Deploy
 ln -f -s "$EnginePath/Binaries/ThirdParty/Python3/Mac/lib" /tmp/Python3/Deploy
 
-IntermediatePath="$EnginePath/Intermediate/Build/Mac/x86_64/SketchUp"
-rm -rf "$IntermediatePath"
-mkdir -p "$IntermediatePath"
+Intermediatex86_64Path="$EnginePath/Intermediate/Build/Mac/x86_64/SketchUp"
+IntermediateUniversalPath="$EnginePath/Intermediate/Build/Mac/arm64+x86_64/SketchUp"
 
-dylibLibFreeImage=libfreeimage-3.18.0.dylib
-dylibLibTbb=libtbb.dylib
-dylibLibTbbMalloc=libtbbmalloc.dylib
+SetupIntermediateDirectory() {
+    IntermediatePath="${1}"
 
-# Copy and fixup dylibs
-Dylibs="$IntermediatePath/Dylibs"
-rm -rf "$Dylibs"
-mkdir -p "$Dylibs"
+    rm -rf "$IntermediatePath"
+    mkdir -p "$IntermediatePath"
 
-cp "$EnginePath/Binaries/ThirdParty/FreeImage/Mac/$dylibLibFreeImage" "$Dylibs"
-cp "$EnginePath/Binaries/Mac/DatasmithSDK/DatasmithSDK.dylib" "$Dylibs"
-cp "$EnginePath/Binaries/ThirdParty/Intel/TBB/Mac/$dylibLibTbb" "$Dylibs"
-cp "$EnginePath/Binaries/ThirdParty/Intel/TBB/Mac/$dylibLibTbbMalloc" "$Dylibs"
+    dylibLibFreeImage=libfreeimage-3.18.0.dylib
+    dylibLibTbb=libtbb.dylib
+    dylibLibTbbMalloc=libtbbmalloc.dylib
 
-chmod 777 "$Dylibs/$dylibLibFreeImage"
-chmod 777 "$Dylibs/$dylibLibTbb"
-chmod 777 "$Dylibs/$dylibLibTbbMalloc"
+    # Copy and fixup dylibs
+    Dylibs="$IntermediatePath/Dylibs"
+    rm -rf "$Dylibs"
+    mkdir -p "$Dylibs"
 
-install_name_tool -id @loader_path/Dylibs/DatasmithSDK.dylib "$Dylibs/DatasmithSDK.dylib"
-install_name_tool -id @loader_path/Dylibs/$dylibLibFreeImage "$Dylibs/$dylibLibFreeImage"
-install_name_tool -change @rpath/$dylibLibFreeImage @loader_path/$dylibLibFreeImage "$Dylibs/DatasmithSDK.dylib"
+    cp "$EnginePath/Binaries/ThirdParty/FreeImage/Mac/$dylibLibFreeImage" "$Dylibs"
+    cp "$EnginePath/Binaries/Mac/DatasmithSDK/DatasmithSDK.dylib" "$Dylibs"
+    cp "$EnginePath/Binaries/ThirdParty/Intel/TBB/Mac/$dylibLibTbb" "$Dylibs"
+    cp "$EnginePath/Binaries/ThirdParty/Intel/TBB/Mac/$dylibLibTbbMalloc" "$Dylibs"
+
+    chmod 777 "$Dylibs/$dylibLibFreeImage"
+    chmod 777 "$Dylibs/$dylibLibTbb"
+    chmod 777 "$Dylibs/$dylibLibTbbMalloc"
+
+    install_name_tool -id @loader_path/Dylibs/DatasmithSDK.dylib "$Dylibs/DatasmithSDK.dylib"
+    install_name_tool -id @loader_path/Dylibs/$dylibLibFreeImage "$Dylibs/$dylibLibFreeImage"
+    install_name_tool -change @rpath/$dylibLibFreeImage @loader_path/$dylibLibFreeImage "$Dylibs/DatasmithSDK.dylib"
+}
+
+# Technicaly our dylibs are already Universal Arm64+x86_64 binaries, so they shouldn't really go into the x86_64 folder,
+# but for the sake of simplicity we keep all the x86_64 dependencies into the x86_64 Intermediate folder with the rest.
+SetupIntermediateDirectory "$Intermediatex86_64Path"
+SetupIntermediateDirectory "$IntermediateUniversalPath"
 
 BuildSketchUpPlugin() {
     SUVERSION=${1}
     SUSDKVERSION=${2}
+	ARCHS=${3}
+    IntermediateDir="${4}/$SUVERSION"
+    Dylibs="${4}/Dylibs"
 
     # Compile plugin bundle
     # BuildDir="$PluginSrcPath/.build/$SUVERSION"
     BuildDir="$EnginePath/Binaries/Mac/SketchUp/$SUVERSION"
-    IntermediateDir="$IntermediatePath/$SUVERSION"
     rm -rf "$BuildDir"
     mkdir -p "$BuildDir"
     #DEBUG_BUILD_SCRIPT="--no-compile --verbose"
     #DEBUG_BUILD_SCRIPT="--verbose"
     DEBUG_BUILD_SCRIPT=
-    "$EnginePath/Binaries/ThirdParty/Python3/Mac/bin/python3" "$PluginSrcPath/Build/build_mac.py" --multithread --sdks-root="$UE_SDKS_ROOT" --sketchup-version=$SUVERSION --sketchup-sdk-version=$SUSDKVERSION --output-path="$BuildDir" --intermediate-path="$IntermediateDir" --datasmithsdk-lib="$Dylibs/DatasmithSDK.dylib" $DEBUG_BUILD_SCRIPT
+    "$EnginePath/Binaries/ThirdParty/Python3/Mac/bin/python3" "$PluginSrcPath/Build/build_mac.py" --multithread --sdks-root="$UE_SDKS_ROOT" --sketchup-version=$SUVERSION --sketchup-sdk-version=$SUSDKVERSION --target_archs $ARCHS --output-path="$BuildDir" --intermediate-path="$IntermediateDir" --datasmithsdk-lib="$Dylibs/DatasmithSDK.dylib" $DEBUG_BUILD_SCRIPT
 
     # Copy plugin files as they should b ein the plugin folder:
     # ruby code
@@ -75,16 +88,19 @@ BuildSketchUpPlugin() {
     "$EnginePath/Binaries/ThirdParty/Python3/Mac/bin/python3" "$PluginSrcPath/Build/create_version_file.py" "$EnginePath" "$BuildDir/Plugin/UnrealDatasmithSketchUp/version"
 }
 
-BuildSketchUpPlugin 2019 SDK_Mac_2019-3-252
-BuildSketchUpPlugin 2020 SDK_Mac_2020-2-171
-BuildSketchUpPlugin 2021 SDK_Mac_2021-0-338
-BuildSketchUpPlugin 2022 SDK_Mac_2022-0-315
+BuildSketchUpPlugin 2019 SDK_Mac_2019-3-252 "x86_64" "$Intermediatex86_64Path"
+BuildSketchUpPlugin 2020 SDK_Mac_2020-2-171 "x86_64" "$Intermediatex86_64Path"
+BuildSketchUpPlugin 2021 SDK_Mac_2021-0-338 "x86_64" "$Intermediatex86_64Path"
+BuildSketchUpPlugin 2022 SDK_Mac_2022-0-353 "x86_64 arm64" "$IntermediateUniversalPath"
 
 # install_name_tool -change @rpath/DatasmithSDK.dylib @loader_path/Dylibs/DatasmithSDK.dylib DatasmithSketchUp.bundle 
 # install_name_tool -change @loader_path/libfreeimage-3.18.0.dylib @loader_path/Dylibs/libfreeimage-3.18.0.dylib DatasmithSketchUp.bundle
 
-rm -f "$EnginePath/Intermediate/Build/Mac/x86_64/SketchUp/Packages"
-mkdir -p "$EnginePath/Intermediate/Build/Mac/x86_64/SketchUp/Packages"
+rm -f "$Intermediatex86_64Path/Packages"
+mkdir -p "$Intermediatex86_64Path/Packages"
+
+rm -f "$IntermediateUniversalPath/Packages"
+mkdir -p "$IntermediateUniversalPath/Packages"
 
 rm -rf /tmp/Python3
 
