[Solidworks] Default directlional lights excluded from export

diff --git a/Engine/Source/Programs/Enterprise/Datasmith/DatasmithSolidworksExporter/Private/Lights.cs b/Engine/Source/Programs/Enterprise/Datasmith/DatasmithSolidworksExporter/Private/Lights.cs
--- a/Engine/Source/Programs/Enterprise/Datasmith/DatasmithSolidworksExporter/Private/Lights.cs
+++ b/Engine/Source/Programs/Enterprise/Datasmith/DatasmithSolidworksExporter/Private/Lights.cs
@@ -197,12 +197,17 @@ namespace DatasmithSolidworks
 					break;
 					case DIR_LIGHT:
 					{
-						double TargetX = 0.0, TargetY = 0.0, TargetZ = 0.0;
-						if (InDoc.GetDirectionLightProperties(LightName, ref AmbientIntensity, ref DiffuseIntensity, ref SpecularIntensity, ref Color, ref bIsEnabled, ref bIsFixed, ref TargetX, ref TargetY, ref TargetZ))
+						if (LightName != "Directional-1") // Don't export three default directional lights(they all have same internal name)
 						{
-							Light = FLight.MakeDirLight(LightName, LightLabel, bIsEnabled, new FVec3(TargetX, TargetY, TargetZ), DiffuseIntensity, Color);
-						}
-					}
+							double TargetX = 0.0, TargetY = 0.0, TargetZ = 0.0;
+							if (InDoc.GetDirectionLightProperties(LightName, ref AmbientIntensity, ref DiffuseIntensity,
+								    ref SpecularIntensity, ref Color, ref bIsEnabled, ref bIsFixed, ref TargetX,
+								    ref TargetY, ref TargetZ))
+							{
+								Light = FLight.MakeDirLight(LightName, LightLabel, bIsEnabled,
+									new FVec3(TargetX, TargetY, TargetZ), DiffuseIntensity, Color);
+							}
+						}					}
 					break;
 				}
 
