[Solidworks] Configurations for Part geometry exported

- when a part file has distinct geometry in some of its configurations then geometry variants are created

diff --git a/Engine/Source/Programs/Enterprise/Datasmith/DatasmithSolidworksExporter/Private/Addin.cs b/Engine/Source/Programs/Enterprise/Datasmith/DatasmithSolidworksExporter/Private/Addin.cs
--- a/Engine/Source/Programs/Enterprise/Datasmith/DatasmithSolidworksExporter/Private/Addin.cs
+++ b/Engine/Source/Programs/Enterprise/Datasmith/DatasmithSolidworksExporter/Private/Addin.cs
@@ -294,7 +294,7 @@ namespace DatasmithSolidworks
 			{
 				switch (InDoc.GetType())
 				{
-					case (int)swDocumentTypes_e.swDocPART: CurrentDocument = new FPartDocument(DocId, InDoc as PartDoc, null, null, null); break;
+					case (int)swDocumentTypes_e.swDocPART: CurrentDocument = new FPartDocument(DocId, InDoc as PartDoc, null, null, new FComponentName()); break;
 					case (int)swDocumentTypes_e.swDocASSEMBLY: CurrentDocument = new FAssemblyDocument(DocId, InDoc as AssemblyDoc, null); break;
 					default: throw new Exception("Unsupported document type");
 				}
diff --git a/Engine/Source/Programs/Enterprise/Datasmith/DatasmithSolidworksExporter/Private/AssemblyDocument.cs b/Engine/Source/Programs/Enterprise/Datasmith/DatasmithSolidworksExporter/Private/AssemblyDocument.cs
--- a/Engine/Source/Programs/Enterprise/Datasmith/DatasmithSolidworksExporter/Private/AssemblyDocument.cs
+++ b/Engine/Source/Programs/Enterprise/Datasmith/DatasmithSolidworksExporter/Private/AssemblyDocument.cs
@@ -24,20 +24,20 @@ namespace DatasmithSolidworks
 		public class FSyncState
 		{
 			public Dictionary<string, FPartDocument> PartsMap = new Dictionary<string, FPartDocument>();
-			public ConcurrentDictionary<string, FObjectMaterials> ComponentsMaterialsMap = new ConcurrentDictionary<string, FObjectMaterials>();
-			public Dictionary<string, string> ComponentToPartMap = new Dictionary<string, string>();
-			public Dictionary<string, float[]> ComponentsTransformsMap = new Dictionary<string, float[]>();
+			public ConcurrentDictionary<FComponentName, FObjectMaterials> ComponentsMaterialsMap = new ConcurrentDictionary<FComponentName, FObjectMaterials>();
+			public Dictionary<FComponentName, string> ComponentToPartMap = new Dictionary<FComponentName, string>();
+			public Dictionary<FComponentName, float[]> ComponentsTransformsMap = new Dictionary<FComponentName, float[]>();
 
-			/** All updated(not dirty) components. Updated component will be removed from dirty on export end*/
-			public HashSet<string> CleanComponents = new HashSet<string>();
+			/// All updated(not dirty) components. Updated component will be removed from dirty on export end
+			public HashSet<FComponentName> CleanComponents = new HashSet<FComponentName>();
 
 			/** Flags for each component that needs update. Cleared when export completes*/
-			public Dictionary<string, uint> DirtyComponents = new Dictionary<string, uint>();
+			public Dictionary<FComponentName, uint> DirtyComponents = new Dictionary<FComponentName, uint>();
 
-			public HashSet<string> ComponentsToDelete = new HashSet<string>();
-			public Dictionary<string, Component2> ExportedComponentsMap = new Dictionary<string, Component2>();
+			public HashSet<FComponentName> ComponentsToDelete = new HashSet<FComponentName>();
+			public Dictionary<FComponentName, Component2> ExportedComponentsMap = new Dictionary<FComponentName, Component2>();
 			/** Stores which mesh was exported for each component*/
-			public Dictionary<string, string> ComponentNameToMeshNameMap = new Dictionary<string, string>();
+			public Dictionary<FComponentName, string> ComponentNameToMeshNameMap = new Dictionary<FComponentName, string>();
 		}
 
 		public AssemblyDoc SwAsmDoc { get; private set; } = null;
@@ -49,7 +49,11 @@ namespace DatasmithSolidworks
 			SwAsmDoc = InSwDoc;
 		}
 
-		public override void ExportToDatasmithScene()
+		public override void PreExport(FMeshes Meshes)
+		{
+		}
+
+		public override void ExportToDatasmithScene(FMeshes meshes)
 		{
 			FSyncState OldSyncState = SyncState;
 
@@ -59,9 +63,9 @@ namespace DatasmithSolidworks
 			}
 
 			SetExportStatus("Actors");
-			foreach (string CompName in SyncState.ComponentsToDelete)
+			foreach (FComponentName CompName in SyncState.ComponentsToDelete)
 			{
-				string ActorName = FDatasmithExporter.SanitizeName(CompName);
+				FActorName ActorName = CompName.GetActorName();
 				SyncState.ComponentToPartMap.Remove(CompName);
 				SyncState.ExportedComponentsMap.Remove(CompName);
 
@@ -83,7 +87,7 @@ namespace DatasmithSolidworks
 
 			// Track components that need their mesh exported: we want to do that in parallel after the 
 			// actor hierarchy has been exported
-			Dictionary<Component2, string> MeshesToExportMap = new Dictionary<Component2, string>();
+			Dictionary<Component2, FActorName> MeshesToExportMap = new Dictionary<Component2, FActorName>();
 		
 			// Store original component suppression state to revert it after export
 			// Component needs to be fully resolved to have its data loaded and even suppressed components are needed for configurations export(in case they are visible in other configs)
@@ -93,16 +97,17 @@ namespace DatasmithSolidworks
 			// Export materials
 			SetExportStatus($"Component Materials");
 
-			HashSet<string> ComponentNamesToExportSet = new HashSet<string>();
+			HashSet<FComponentName> ComponentNamesToExportSet = new HashSet<FComponentName>();
 			foreach (var KVP in MeshesToExportMap)
 			{
-				if (!ComponentNamesToExportSet.Contains(KVP.Key.Name))
+				FComponentName ComponentName = new FComponentName(KVP.Key);
+				if (!ComponentNamesToExportSet.Contains(ComponentName))
 				{
-					ComponentNamesToExportSet.Add(KVP.Key.Name);
+					ComponentNamesToExportSet.Add(ComponentName);
 				}
 			}
 
-			ConcurrentDictionary<string, FObjectMaterials> ModifiedComponentsMaterials = FObjectMaterials.LoadAssemblyMaterials(this, ComponentNamesToExportSet, swDisplayStateOpts_e.swThisDisplayState, null);
+			ConcurrentDictionary<FComponentName, FObjectMaterials> ModifiedComponentsMaterials = FObjectMaterials.LoadAssemblyMaterials(this, ComponentNamesToExportSet, swDisplayStateOpts_e.swThisDisplayState, null);
 
 			if (ModifiedComponentsMaterials != null)
 			{
@@ -123,7 +128,7 @@ namespace DatasmithSolidworks
 				Component2 Comp = KVP.Key;
 
 				FObjectMaterials ComponentMaterials = null;
-				SyncState.ComponentsMaterialsMap?.TryGetValue(Comp.Name2, out ComponentMaterials);
+				SyncState.ComponentsMaterialsMap?.TryGetValue(new FComponentName(Comp), out ComponentMaterials);
 
 				ConcurrentBag<FBody> Bodies = FBody.FetchBodies(Comp);
 				FMeshData MeshData = FStripGeometry.CreateMeshData(Bodies, ComponentMaterials);
@@ -141,7 +146,7 @@ namespace DatasmithSolidworks
 			foreach (var KVP in CreatedMeshes)
 			{
 				string MeshName = Exporter.AddMesh(KVP.Value);
-				SyncState.ComponentNameToMeshNameMap[MeshName] = KVP.Key.Name2;
+				SyncState.ComponentNameToMeshNameMap[new FComponentName(KVP.Key)] = MeshName;
 			}
 
 			SyncState.ComponentsToDelete.Clear();
@@ -173,11 +178,11 @@ namespace DatasmithSolidworks
 		public override bool HasMaterialUpdates()
 		{
 			// Dig into part level materials (they wont be read by LoadDocumentMaterials)
-			HashSet<string> AllExportedComponents = new HashSet<string>();
+			HashSet<FComponentName> AllExportedComponents = new HashSet<FComponentName>();
 			AllExportedComponents.UnionWith(SyncState.CleanComponents);
 			AllExportedComponents.UnionWith(SyncState.DirtyComponents.Keys);
 
-			ConcurrentDictionary<string, FObjectMaterials> CurrentDocMaterialsMap = FObjectMaterials.LoadAssemblyMaterials(this, AllExportedComponents, swDisplayStateOpts_e.swThisDisplayState, null);
+			ConcurrentDictionary<FComponentName, FObjectMaterials> CurrentDocMaterialsMap = FObjectMaterials.LoadAssemblyMaterials(this, AllExportedComponents, swDisplayStateOpts_e.swThisDisplayState, null);
 
 			if (CurrentDocMaterialsMap == null && SyncState.ComponentsMaterialsMap == null)
 			{
@@ -205,15 +210,15 @@ namespace DatasmithSolidworks
 
 				if (CurrentDocMaterialsMap.Count != SyncState.ComponentsMaterialsMap.Count)
 				{
-					IEnumerable<string> Diff1 = CurrentDocMaterialsMap.Keys.Except(SyncState.ComponentsMaterialsMap.Keys);
-					IEnumerable<string> Diff2 = SyncState.ComponentsMaterialsMap.Keys.Except(CurrentDocMaterialsMap.Keys);
+					IEnumerable<FComponentName> Diff1 = CurrentDocMaterialsMap.Keys.Except(SyncState.ComponentsMaterialsMap.Keys);
+					IEnumerable<FComponentName> Diff2 = SyncState.ComponentsMaterialsMap.Keys.Except(CurrentDocMaterialsMap.Keys);
 
-					HashSet<string> DiffSet = new HashSet<string>();
+					HashSet<FComponentName> DiffSet = new HashSet<FComponentName>();
 					DiffSet.UnionWith(Diff1);
 					DiffSet.UnionWith(Diff2);
 
 					// Components in the DiffSet have their materials changed
-					foreach (string CompName in DiffSet)
+					foreach (FComponentName CompName in DiffSet)
 					{
 						bool bShouldSyncComponentMaterial = false;
 
@@ -283,38 +288,38 @@ namespace DatasmithSolidworks
 			return DatasmithTransform;
 		}
 
-		private void ExportComponentRecursive(Component2 InComponent, Component2 InParent, Dictionary<Component2, string> OutMeshesToExportMap, List<Tuple<Component2, int>> OutComponentSuppression)
+		private void ExportComponentRecursive(Component2 InComponent, Component2 InParent, Dictionary<Component2, FActorName> OutMeshesToExportMap, List<Tuple<Component2, int>> OutComponentSuppression)
 		{
 			bool bHasDirtyTransform = false;
 
-			if (SyncState.ComponentsTransformsMap.ContainsKey(InComponent.Name2))
+			FComponentName ComponentName = new FComponentName(InComponent);
+			if (SyncState.ComponentsTransformsMap.ContainsKey(ComponentName))
 			{
 				float[] ComponentTm = GetComponentDatasmithTransform(InComponent);
-				bHasDirtyTransform = !MathUtils.TransformsAreEqual(SyncState.ComponentsTransformsMap[InComponent.Name2], ComponentTm);
+				bHasDirtyTransform = !MathUtils.TransformsAreEqual(SyncState.ComponentsTransformsMap[ComponentName], ComponentTm);
 			}
 
-			if (bHasDirtyTransform || !SyncState.CleanComponents.Contains(InComponent.Name2))
+			if (bHasDirtyTransform || !SyncState.CleanComponents.Contains(ComponentName))
 			{
-				SetExportStatus(InComponent.Name2);
+				SetExportStatus(ComponentName.GetString());
 
 				FDatasmithActorExportInfo ActorExportInfo = new FDatasmithActorExportInfo();
 
-				string ComponentName = FDatasmithExporter.SanitizeName(InComponent.Name2);
-				string[] NameComponents = InComponent.Name2.Split('/');
+				FActorName ActorName = ComponentName.GetActorName();
 
-				ActorExportInfo.Label = NameComponents.Last();
-				ActorExportInfo.Name = ComponentName;
+				ActorExportInfo.Label = ComponentName.GetLabel();
+				ActorExportInfo.Name = ActorName;
 
 				if (InParent != null)
 				{
-					ActorExportInfo.ParentName = FDatasmithExporter.SanitizeName(InParent.Name2);
+					ActorExportInfo.ParentName = new FComponentName(InParent).GetActorName();
 				}
 				
 				ActorExportInfo.bVisible = true;
-				ActorExportInfo.Type = Exporter.GetExportedActorType(ComponentName) ?? EActorType.SimpleActor;
+				ActorExportInfo.Type = Exporter.GetExportedActorType(ActorName) ?? EActorType.SimpleActor;
 				ActorExportInfo.Transform = GetComponentDatasmithTransform(InComponent);
 
-				SyncState.ComponentsTransformsMap[InComponent.Name2] = ActorExportInfo.Transform;
+				SyncState.ComponentsTransformsMap[ComponentName] = ActorExportInfo.Transform;
 
 				// When exporting configurations we need to load all components to make their data available for export
 				// Suppressed or lightweight components don't have this
@@ -348,12 +353,12 @@ namespace DatasmithSolidworks
 				ModelDoc2 ComponentDoc = InComponent.GetModelDoc2();
 				if (ComponentDoc is PartDoc) 
 				{
-					bool bFirstExport = !SyncState.ExportedComponentsMap.ContainsKey(InComponent.Name2);
+					bool bFirstExport = !SyncState.ExportedComponentsMap.ContainsKey(ComponentName);
 					bNeedsGeometryExport = bFirstExport;
 
-					if (!bFirstExport && SyncState.DirtyComponents.ContainsKey(InComponent.Name2))
+					if (!bFirstExport && SyncState.DirtyComponents.ContainsKey(ComponentName))
 					{
-						uint DirtyState = SyncState.DirtyComponents[InComponent.Name2];
+						uint DirtyState = SyncState.DirtyComponents[ComponentName];
 						bNeedsGeometryExport =
 							((DirtyState & (1u << (int)EComponentDirtyState.Material)) != 0) ||
 							((DirtyState & (1u << (int)EComponentDirtyState.Geometry)) != 0) ||
@@ -365,15 +370,15 @@ namespace DatasmithSolidworks
 				{
 					//TODO this will be null for new part, think of more solid solution
 					string PartPath = ComponentDoc.GetPathName();
-					if (!SyncState.PartsMap.ContainsKey(InComponent.Name2))
+					if (!SyncState.PartsMap.ContainsKey(PartPath))
 					{
 						// New part
 						int PartDocId = Addin.Instance.GetDocumentId(ComponentDoc as ModelDoc2);
-						SyncState.PartsMap[PartPath] = new FPartDocument(PartDocId, ComponentDoc as PartDoc, Exporter, this, InComponent.Name2);
+						SyncState.PartsMap[PartPath] = new FPartDocument(PartDocId, ComponentDoc as PartDoc, Exporter, this, ComponentName);
 						SyncState.PartsMap[PartPath].Init();
 					}
 
-					SyncState.ComponentToPartMap[InComponent.Name2] = PartPath;
+					SyncState.ComponentToPartMap[ComponentName] = PartPath;
 
 					// This component has associated part document -- treat is as a mesh actor
 					ActorExportInfo.Type = EActorType.MeshActor;
@@ -383,10 +388,10 @@ namespace DatasmithSolidworks
 
 				Exporter.ExportOrUpdateActor(ActorExportInfo);
 
-				SyncState.CleanComponents.Add(InComponent.Name2);
+				SyncState.CleanComponents.Add(ComponentName);
 			}
 
-			SyncState.ExportedComponentsMap[InComponent.Name2] = InComponent;
+			SyncState.ExportedComponentsMap[ComponentName] = InComponent;
 
 			// Export component children
 			object[] Children = (object[])InComponent.GetChildren();
@@ -401,7 +406,7 @@ namespace DatasmithSolidworks
 			}
 		}
 
-		public void SetComponentDirty(string InComponent, EComponentDirtyState InState)
+		public void SetComponentDirty(FComponentName InComponent, EComponentDirtyState InState)
 		{
 			if (SyncState.CleanComponents.Contains(InComponent))
 			{
@@ -416,6 +421,11 @@ namespace DatasmithSolidworks
 			SetDirty(true);
 		}
 
+		public void SetComponentDirty(Component2 InComponent, EComponentDirtyState InState)
+		{
+			SetComponentDirty(new FComponentName(InComponent), InState);
+		}
+
 		public override void Init()
 		{
 			base.Init();
@@ -528,7 +538,7 @@ namespace DatasmithSolidworks
 		{
 			if (InComponentObj is Component2 Comp)
 			{
-				SetComponentDirty(Comp.Name2, EComponentDirtyState.Visibility);
+				SetComponentDirty(Comp, EComponentDirtyState.Visibility);
 			}
 			return 0;
 		}
@@ -537,7 +547,7 @@ namespace DatasmithSolidworks
 		{
 			if (InCompObject is Component2 Comp)
 			{
-				SetComponentDirty(Comp.Name2, EComponentDirtyState.Material);
+				SetComponentDirty(Comp, EComponentDirtyState.Material);
 			}
 			return 0;
 		}
@@ -585,7 +595,7 @@ namespace DatasmithSolidworks
 
 		int OnComponentConfigurationChangeNotify(string componentName, string oldConfigurationName, string newConfigurationName)
 		{
-			SetComponentDirty(componentName, EComponentDirtyState.Geometry);
+			SetComponentDirty(FComponentName.FromApiString(componentName), EComponentDirtyState.Geometry);
 			return 0;
 		}
 
@@ -597,7 +607,7 @@ namespace DatasmithSolidworks
 			{
 				Component2 Comp = KVP.Value;
 				float[] PrevCompTransform = null;
-				if (SyncState.ComponentsTransformsMap.TryGetValue(Comp.Name2, out PrevCompTransform))
+				if (SyncState.ComponentsTransformsMap.TryGetValue(new FComponentName(Comp), out PrevCompTransform))
 				{
 					float[] CompTransform = GetComponentDatasmithTransform(Comp);
 
@@ -645,7 +655,7 @@ namespace DatasmithSolidworks
 		{
 			if (InComponentObj is Component2 Comp)
 			{
-				SetComponentDirty(Comp.Name2, EComponentDirtyState.Visibility);
+				SetComponentDirty(new FComponentName(Comp), EComponentDirtyState.Visibility);
 			}
 			return 0;
 		}
@@ -657,24 +667,28 @@ namespace DatasmithSolidworks
 
 		int OnDeleteItemNotify(int InEntityType, string InItemName)
 		{
+			FComponentName ComponentName = FComponentName.FromApiString(InItemName);
+
 			if (InEntityType == (int)swNotifyEntityType_e.swNotifyComponent && 
-				SyncState.ComponentToPartMap.ContainsKey(InItemName) && 
-				!SyncState.ComponentsToDelete.Contains(InItemName))
+				SyncState.ComponentToPartMap.ContainsKey(ComponentName) && 
+				!SyncState.ComponentsToDelete.Contains(ComponentName))
 			{
-				SyncState.ComponentsToDelete.Add(InItemName);
-				SetComponentDirty(InItemName, EComponentDirtyState.Delete);
+				SyncState.ComponentsToDelete.Add(ComponentName);
+				SetComponentDirty(ComponentName, EComponentDirtyState.Delete);
 			}
 			return 0;
 		}
 
 		int OnRenameItemNotify(int InEntityType, string InOldName, string InNewName)
 		{
+			FComponentName ComponentName = FComponentName.FromApiString(InOldName);
+
 			if (InEntityType == (int)swNotifyEntityType_e.swNotifyComponent &&
-				SyncState.ComponentToPartMap.ContainsKey(InOldName) &&
-				!SyncState.ComponentsToDelete.Contains(InOldName))
+			    SyncState.ComponentToPartMap.ContainsKey(ComponentName) &&
+			    !SyncState.ComponentsToDelete.Contains(ComponentName))
 			{
-				SyncState.ComponentsToDelete.Add(InOldName);
-				SetComponentDirty(InOldName, EComponentDirtyState.Delete);
+				SyncState.ComponentsToDelete.Add(ComponentName);
+				SetComponentDirty(ComponentName, EComponentDirtyState.Delete);
 			}
 			return 0;
 		}
@@ -750,7 +764,7 @@ namespace DatasmithSolidworks
 				IComponent2 Comp = ObjComp as IComponent2;
 				if (Comp != null)
 				{
-					SetComponentDirty(Comp.Name2, EComponentDirtyState.Transform);
+					SetComponentDirty(new FComponentName(Comp), EComponentDirtyState.Transform);
 				}
 			}
 			return 0;
diff --git a/Engine/Source/Programs/Enterprise/Datasmith/DatasmithSolidworksExporter/Private/ConfigurationExporter.cs b/Engine/Source/Programs/Enterprise/Datasmith/DatasmithSolidworksExporter/Private/ConfigurationExporter.cs
--- a/Engine/Source/Programs/Enterprise/Datasmith/DatasmithSolidworksExporter/Private/ConfigurationExporter.cs
+++ b/Engine/Source/Programs/Enterprise/Datasmith/DatasmithSolidworksExporter/Private/ConfigurationExporter.cs
@@ -6,30 +6,195 @@ using SolidWorks.Interop.sldworks;
 using SolidWorks.Interop.swconst;
 using System.Runtime.InteropServices;
 using System.Collections.Concurrent;
+using System.Linq;
 
 namespace DatasmithSolidworks
 {
+	[ComVisible(false)]
 	public class FConfigurationData
 	{
 		public string Name;
 		public bool bIsDisplayStateConfiguration = false;
-		public Dictionary<string, float[]> ComponentTransform = new Dictionary<string, float[]>();
-		public Dictionary<string, bool> ComponentVisibility = new Dictionary<string, bool>();
-		public Dictionary<string, FObjectMaterials> ComponentMaterials = new Dictionary<string, FObjectMaterials>();
+		public Dictionary<FComponentName, float[]> ComponentTransform = new Dictionary<FComponentName, float[]>();
+		public Dictionary<FComponentName, bool> ComponentVisibility = new Dictionary<FComponentName, bool>();
+		public Dictionary<FComponentName, FObjectMaterials> ComponentMaterials = new Dictionary<FComponentName, FObjectMaterials>();
+
+		[ComVisible(false)]
+		public struct FComponentGeometryVariant
+		{
+			public FActorName VisibleActor; // Actor enabled for this variant
+			public List<FActorName> All; // All actors used for this component in variants
+		}
+
+		public Dictionary<FComponentName, FComponentGeometryVariant> ComponentGeometry = new Dictionary<FComponentName, FComponentGeometryVariant>();
 
 		public bool IsEmpty()
 		{
-			return (
-				ComponentTransform.Count == 0 && 
-				ComponentVisibility.Count == 0 && 
-				ComponentMaterials.Count == 0);
+			return (ComponentTransform.Count == 0
+			        && ComponentVisibility.Count == 0 
+			        && ComponentMaterials.Count == 0
+			        && ComponentGeometry.Count == 0);
 		}
 	};
 
+	[ComVisible(false)]
+	public class FMeshes
+	{
+		private Dictionary<string, IConfiguration> Configurations = new Dictionary<string, IConfiguration>();
+
+		private HashSet<FMesh> Meshes = new HashSet<FMesh>();
+
+		[ComVisible(false)]
+		public struct FMesh
+		{
+			public FMeshData MeshData;
+
+			private readonly int Hash;
+
+			public FMesh(FMeshData InMeshData)
+			{
+
+				MeshData = InMeshData;
+
+				Hash = 0;
+
+				foreach (FVec3 Vertex in MeshData.Vertices)
+				{
+					Hash ^= Vertex.GetHashCode();
+				}
+
+				foreach (FVec3 Normal in MeshData.Normals)
+				{
+					Hash ^= Normal.GetHashCode();
+				}
+
+				foreach (FVec2 TexCoord in MeshData.TexCoords)
+				{
+					Hash ^= TexCoord.GetHashCode();
+				}
+
+				foreach (FTriangle Triangle in MeshData.Triangles)
+				{
+					Hash ^= Triangle.GetHashCode();
+				}
+			}
+
+			public override bool Equals(object Obj)
+			{
+				return Obj is FMesh Other && Equals(Other);
+			}
+
+			public bool Equals(FMesh Other)
+			{
+				
+				if (Hash != Other.Hash)
+				{
+					return false;
+				}
+
+				if (!MeshData.Vertices.SequenceEqual(Other.MeshData.Vertices))
+				{
+					return false;
+				}
+
+				if (!MeshData.Normals.SequenceEqual(Other.MeshData.Normals))
+				{
+					return false;
+				}
+
+				if (!MeshData.TexCoords.SequenceEqual(Other.MeshData.TexCoords))
+				{
+					return false;
+				}
+
+				if (!MeshData.Triangles.SequenceEqual(Other.MeshData.Triangles))
+				{
+					return false;
+				}
+
+				return true;
+			}
+
+			public static bool operator ==(FMesh A, FMesh B)
+			{
+				return A.Equals(B);
+			}
+
+			public static bool operator !=(FMesh A, FMesh B)
+			{
+				return !(A == B);
+			}
+
+			public override int GetHashCode()
+			{
+				return Hash;
+			}
+		}
+
+		[ComVisible(false)]
+		public interface IConfiguration
+		{
+			void AddMesh(Component2 Component, FMeshData Mesh);
+			bool GetMeshForComponent(FComponentName ComponentName, out FMesh Result);
+		}
+
+		public IConfiguration GetMeshesConfiguration(string ConfigurationName)
+		{
+			if (Configurations.TryGetValue(ConfigurationName, out var Configuration))
+			{
+				return Configuration;
+			}			
+
+			Configuration = new FConfigurationImpl(this, ConfigurationName);
+			Configurations.Add(ConfigurationName, Configuration);
+			return Configuration;
+		}
+
+		public bool IsSameMesh(FComponentName ComponentName, string ConfigNameA, string ConfigNameB)
+		{
+			return GetMeshesConfiguration(ConfigNameA).GetMeshForComponent(ComponentName, out FMesh MeshA)
+			       && GetMeshesConfiguration(ConfigNameB).GetMeshForComponent(ComponentName, out FMesh MeshB)
+			       && MeshA == MeshB;
+		}
+
+
+		public FMeshData GetMeshData(string ConfigName, FComponentName ComponentName)
+		{
+			return GetMeshesConfiguration(ConfigName).GetMeshForComponent(ComponentName, out FMesh Mesh) ? Mesh.MeshData : null;
+		}
+
+		private class FConfigurationImpl : IConfiguration
+		{
+			private Dictionary<FComponentName, FMesh> MeshForComponent = new Dictionary<FComponentName, FMesh>();
+
+			public FConfigurationImpl(FMeshes Meshes, string Name)
+			{
+			}
+
+			public void AddMesh(Component2 Component, FMeshData Mesh)
+			{
+				MeshForComponent.Add(new FComponentName(Component), new FMesh(Mesh));
+			}
+
+			public bool GetMeshForComponent(FComponentName ComponentName, out FMesh Result)
+			{
+				return MeshForComponent.TryGetValue(ComponentName, out Result);
+			}
+
+		}
+	}
+
 	[ComVisible(false)]
 	public class FConfigurationExporter
 	{
-		public static List<FConfigurationData> ExportConfigurations(FDocument InDoc)
+		private readonly FMeshes Meshes;
+
+		public FConfigurationExporter(FMeshes Meshes)
+		{
+			this.Meshes = Meshes;
+		}
+
+		public List<FConfigurationData> ExportConfigurations(FDocument InDoc)
 		{
 			string[] CfgNames = InDoc.SwDoc?.GetConfigurationNames();
 
@@ -39,7 +204,7 @@ namespace DatasmithSolidworks
 			}
 
 			FConfigurationTree.FComponentTreeNode CombinedTree = new FConfigurationTree.FComponentTreeNode();
-			CombinedTree.ComponentName = "CombinedTree";
+			CombinedTree.ComponentName = FComponentName.FromCustomString("CombinedTree");
 
 			// Ensure recursion will not stop on the root node (this may happen if it is not explicitly marked as visible)
 			CombinedTree.bVisibilitySame = true;
@@ -69,15 +234,17 @@ namespace DatasmithSolidworks
 
 				FConfigurationTree.FComponentTreeNode ConfigNode = new FConfigurationTree.FComponentTreeNode();
 				ConfigNode.Children = new List<FConfigurationTree.FComponentTreeNode>();
-				ConfigNode.ComponentName = ConfigName;
+				ConfigNode.ComponentName = FComponentName.FromCustomString(ConfigName);
+
 
 				// Build the tree and get default materials (which aren't affected by any configuration)
 				// Use GetRootComponent3() with Resolve = true to ensure suppressed components will be loaded
+				// todo: docs says that Part document SW returns null. Not the case. But probably better to add a guard
 				CollectComponentsRecursive(InDoc, swConfiguration.GetRootComponent3(true), ConfigNode);
 
 				ExportedConfigurationNames.Add(ConfigName);
 
-				Dictionary<string, FObjectMaterials> MaterialsMap = GetComponentMaterials(InDoc, DisplayStates != null ? DisplayStates[0] : null, swConfiguration);
+				Dictionary<FComponentName, FObjectMaterials> MaterialsMap = GetComponentMaterials(InDoc, DisplayStates != null ? DisplayStates[0] : null, swConfiguration);
 				SetComponentTreeMaterials(ConfigNode, MaterialsMap, null, false);
 
 				if (DisplayStates != null)
@@ -113,14 +280,14 @@ namespace DatasmithSolidworks
 				{
 					string DisplayStateConfigName = $"DisplayState_{FDatasmithExporter.SanitizeName(DisplayState)}";
 					DisplayStateConfigurations.Add(DisplayStateConfigName);
-					Dictionary<string, FObjectMaterials> MaterialsMap = GetComponentMaterials(InDoc, DisplayState, OriginalConfiguration);
+					Dictionary<FComponentName, FObjectMaterials> MaterialsMap = GetComponentMaterials(InDoc, DisplayState, OriginalConfiguration);
 
 					SetComponentTreeMaterials(CombinedTree, MaterialsMap, DisplayStateConfigName, true);
 				}
 			}
 
 			// Remove configuration data when it's the same
-			FConfigurationTree.Compress(CombinedTree);
+			FConfigurationTree.Compress(CombinedTree, this);
 
 			List<FConfigurationData> FlatConfigurationData = new List<FConfigurationData>();
 
@@ -130,7 +297,7 @@ namespace DatasmithSolidworks
 				FConfigurationData CfgData = new FConfigurationData();
 				CfgData.Name = CfgName;
 				CfgData.bIsDisplayStateConfiguration = false;
-				FConfigurationTree.FillConfigurationData(CombinedTree, CfgName, CfgData, false);
+				FConfigurationTree.FillConfigurationData(this, CombinedTree, CfgName, CfgData, false);
 
 				if (!CfgData.IsEmpty())
 				{
@@ -142,7 +309,7 @@ namespace DatasmithSolidworks
 				FConfigurationData CfgData = new FConfigurationData();
 				CfgData.Name = DisplayStateConfigName;
 				CfgData.bIsDisplayStateConfiguration = true;
-				FConfigurationTree.FillConfigurationData(CombinedTree, DisplayStateConfigName, CfgData, true);
+				FConfigurationTree.FillConfigurationData(this, CombinedTree, DisplayStateConfigName, CfgData, true);
 
 				if (!CfgData.IsEmpty())
 				{
@@ -153,9 +320,9 @@ namespace DatasmithSolidworks
 			return FlatConfigurationData;
 		}
 
-		private static Dictionary<string, FObjectMaterials> GetComponentMaterials(FDocument InDoc, string InDisplayState, IConfiguration InConfiguration)
+		private static Dictionary<FComponentName, FObjectMaterials> GetComponentMaterials(FDocument InDoc, string InDisplayState, IConfiguration InConfiguration)
 		{
-			Dictionary<string, FObjectMaterials> MaterialsMap = new Dictionary<string, FObjectMaterials>();
+			Dictionary<FComponentName, FObjectMaterials> MaterialsMap = new Dictionary<FComponentName, FObjectMaterials>();
 
 			swDisplayStateOpts_e Option = InDisplayState != null ? swDisplayStateOpts_e.swSpecifyDisplayState : swDisplayStateOpts_e.swThisDisplayState;
 			string[] DisplayStates = InDisplayState != null ? new string[] { InDisplayState } : null;
@@ -164,14 +331,14 @@ namespace DatasmithSolidworks
 			{
 				FAssemblyDocument AsmDoc = InDoc as FAssemblyDocument;
 
-				HashSet<string> InComponentsSet = new HashSet<string>();
+				HashSet<FComponentName> InComponentsSet = new HashSet<FComponentName>();
 
-				foreach (string CompName in AsmDoc.SyncState.ExportedComponentsMap.Keys)
+				foreach (FComponentName CompName in AsmDoc.SyncState.ExportedComponentsMap.Keys)
 				{
 					InComponentsSet.Add(CompName);
 				}
 
-				ConcurrentDictionary<string, FObjectMaterials> ComponentMaterials =
+				ConcurrentDictionary<FComponentName, FObjectMaterials> ComponentMaterials =
 					FObjectMaterials.LoadAssemblyMaterials(AsmDoc, InComponentsSet, Option, DisplayStates);
 
 				if (ComponentMaterials != null)
@@ -187,13 +354,14 @@ namespace DatasmithSolidworks
 				FPartDocument PartDocument = InDoc as FPartDocument;
 				FObjectMaterials PartMaterials = FObjectMaterials.LoadPartMaterials(InDoc, PartDocument.SwPartDoc, Option, DisplayStates);
 				Component2 Comp = InConfiguration.GetRootComponent3(true);
-				MaterialsMap.Add(Comp?.Name2 ?? "", PartMaterials);
+
+				MaterialsMap.Add(FComponentName.FromCustomString(Comp?.Name2 ?? ""), PartMaterials);
 			}
 
 			return MaterialsMap;
 		}
 
-		private static void SetComponentTreeMaterials(FConfigurationTree.FComponentTreeNode InComponentTree, Dictionary<string, FObjectMaterials> InComponentMaterialsMap, string InConfigurationName, bool bIsDisplayState)
+		private static void SetComponentTreeMaterials(FConfigurationTree.FComponentTreeNode InComponentTree, Dictionary<FComponentName, FObjectMaterials> InComponentMaterialsMap, string InConfigurationName, bool bIsDisplayState)
 		{
 			FConfigurationTree.FComponentConfig TargetConfig = null;
 
@@ -231,7 +399,7 @@ namespace DatasmithSolidworks
 			InParentNode.Children.Add(NewNode);
 
 			// Basic properties
-			NewNode.ComponentName = InComponent.Name2;
+			NewNode.ComponentName = new FComponentName(InComponent);
 			NewNode.ComponentID = InComponent.GetID();
 			NewNode.CommonConfig.bVisible = InComponent.Visible != (int)swComponentVisibilityState_e.swComponentHidden;
 			NewNode.CommonConfig.bSuppressed = InComponent.IsSuppressed();
@@ -263,7 +431,7 @@ namespace DatasmithSolidworks
 
 			// Process children components
 			var Children = (Object[])InComponent.GetChildren();
-			if (Children != null && Children.Length > 0)
+			if (Children != null && Children.Length > 0)  // Children can be null sometimes and sometimes zero-length array
 			{
 				NewNode.Children = new List<FConfigurationTree.FComponentTreeNode>();
 				foreach (object ObjChild in Children)
@@ -278,5 +446,15 @@ namespace DatasmithSolidworks
 				});
 			}
 		}
+
+		public bool IsSameMesh(FComponentName ComponentName, string ConfigNameA, string ConfigNameB)
+		{
+			return Meshes.IsSameMesh(ComponentName, ConfigNameA, ConfigNameB);
+		}
+
+		public FActorName GetMeshActorName(FComponentName ComponentName, string ConfigName)
+		{
+			return FActorName.FromString(ComponentName.GetString() + "_" + ConfigName);
+		}
 	}
 }
\ No newline at end of file
diff --git a/Engine/Source/Programs/Enterprise/Datasmith/DatasmithSolidworksExporter/Private/ConfigurationTree.cs b/Engine/Source/Programs/Enterprise/Datasmith/DatasmithSolidworksExporter/Private/ConfigurationTree.cs
--- a/Engine/Source/Programs/Enterprise/Datasmith/DatasmithSolidworksExporter/Private/ConfigurationTree.cs
+++ b/Engine/Source/Programs/Enterprise/Datasmith/DatasmithSolidworksExporter/Private/ConfigurationTree.cs
@@ -2,6 +2,7 @@
 
 using SolidWorks.Interop.sldworks;
 using System.Collections.Generic;
+using System.Diagnostics;
 using System.Linq;
 using System.Runtime.InteropServices;
 
@@ -41,7 +42,7 @@ namespace DatasmithSolidworks
 		// A node contains all the configuration data for a single component
 		public class FComponentTreeNode
 		{
-			public string ComponentName;
+			public FComponentName ComponentName;
 			public int ComponentID;
 
 			// Common configuration data
@@ -54,6 +55,9 @@ namespace DatasmithSolidworks
 			public bool bVisibilitySame = true;
 			public bool bSuppressionSame = true;
 			public bool bMaterialSame = true;
+			public bool bGeometrySame = true;
+
+			public HashSet<FActorName> Meshes = null;
 
 			public List<FComponentTreeNode> Children;
 
@@ -158,7 +162,7 @@ namespace DatasmithSolidworks
 			}
 		}
 
-		static public void Compress(FComponentTreeNode InNode)
+		static public void Compress(FComponentTreeNode InNode, FConfigurationExporter ConfigurationExporter)
 		{
 			bool CheckMaterialsEqual(List<FComponentConfig> ConfigList)
 			{
@@ -224,6 +228,8 @@ namespace DatasmithSolidworks
 				bool bSuppressed = InNode.Configurations[0].bSuppressed;
 				bool bVisibilitySame = true;
 				bool bSuppressionSame = true;
+				bool bGeometrySame = true;
+
 				for (int Idx = 1; Idx < InNode.Configurations.Count; Idx++)
 				{
 					if (InNode.Configurations[Idx].bVisible != bVisible)
@@ -235,11 +241,30 @@ namespace DatasmithSolidworks
 					{
 						bSuppressionSame = false;
 					}
+
+					if (!ConfigurationExporter.IsSameMesh(InNode.ComponentName, InNode.Configurations[0].ConfigName, InNode.Configurations[Idx].ConfigName))
+					{
+						bGeometrySame = false;
+					}
+				}
+
+				if (!bGeometrySame)
+				{
+					Debug.Assert(InNode.Meshes == null);
+					InNode.Meshes = new HashSet<FActorName>();
+
+					foreach (FComponentConfig ComponentConfig in InNode.Configurations)
+					{
+						InNode.Meshes.Add(ConfigurationExporter.GetMeshActorName(InNode.ComponentName,
+							ComponentConfig.ConfigName));
+					}
 				}
+
 				// Propagate common values
 				InNode.bMaterialSame = bAllMaterialsAreSame;
 				InNode.bVisibilitySame = bVisibilitySame;
 				InNode.bSuppressionSame = bSuppressionSame;
+				InNode.bGeometrySame = bGeometrySame;
 
 				if (bVisibilitySame)
 				{
@@ -253,7 +278,7 @@ namespace DatasmithSolidworks
 				//todo: store bAll...Same in ConfigData
 
 				// If EVERYTHING is same, just remove all configurations at all
-				if (bAllTransformsAreSame && bAllMaterialsAreSame && bVisibilitySame && bSuppressionSame)
+				if (bAllTransformsAreSame && bAllMaterialsAreSame && bVisibilitySame && bSuppressionSame && bGeometrySame)
 				{
 					InNode.Configurations = null;
 				}
@@ -276,16 +301,16 @@ namespace DatasmithSolidworks
 			{
 				foreach (FComponentTreeNode Child in InNode.Children)
 				{
-					Compress(Child);
+					Compress(Child, ConfigurationExporter);
 				}
 			}
 		}
 
-		static public void FillConfigurationData(FComponentTreeNode InNode, string InConfigurationName, FConfigurationData OutData, bool bIsDisplayState)
+		static public void FillConfigurationData(FConfigurationExporter ConfigurationExporter, FComponentTreeNode InNode, string InConfigurationName, FConfigurationData OutData, bool bIsDisplayState)
 		{
 			FComponentConfig NodeConfig = InNode.GetConfiguration(InConfigurationName, bIsDisplayState);
 
-			// Visibility or suppression flags are set per node, and not propagated to childre.
+			// Visibility or suppression flags are set per node, and not propagated to children.
 			// Also, suppression doesn't mark node as invisible. We should process these separately
 			// to exclude any variant information from invisible nodes and their children.
 			if ((InNode.bVisibilitySame && !InNode.CommonConfig.bVisible) || (NodeConfig != null && !NodeConfig.bVisible) ||
@@ -313,6 +338,15 @@ namespace DatasmithSolidworks
 				{
 					OutData.ComponentVisibility.Add(InNode.ComponentName, !NodeConfig.bSuppressed);
 				}
+				if (!InNode.bGeometrySame)
+				{
+					OutData.ComponentGeometry.Add(InNode.ComponentName, new FConfigurationData.FComponentGeometryVariant
+						{
+							VisibleActor = ConfigurationExporter.GetMeshActorName(InNode.ComponentName, InConfigurationName),
+							All = InNode.Meshes.ToList()
+						}
+					);
+				}
 				if (!InNode.bMaterialSame)
 				{
 					FObjectMaterials Materials = NodeConfig.Materials;
@@ -333,9 +367,9 @@ namespace DatasmithSolidworks
 			{
 				foreach (FComponentTreeNode Child in InNode.Children)
 				{
-					FillConfigurationData(Child, InConfigurationName, OutData, bIsDisplayState);
+					FillConfigurationData(ConfigurationExporter, Child, InConfigurationName, OutData, bIsDisplayState);
 				}
 			}
 		}
 	}
-}
\ No newline at end of file
+}
diff --git a/Engine/Source/Programs/Enterprise/Datasmith/DatasmithSolidworksExporter/Private/DatasmithExporter.cs b/Engine/Source/Programs/Enterprise/Datasmith/DatasmithSolidworksExporter/Private/DatasmithExporter.cs
--- a/Engine/Source/Programs/Enterprise/Datasmith/DatasmithSolidworksExporter/Private/DatasmithExporter.cs
+++ b/Engine/Source/Programs/Enterprise/Datasmith/DatasmithSolidworksExporter/Private/DatasmithExporter.cs
@@ -5,6 +5,8 @@ using System;
 using System.Collections.Concurrent;
 using System.Collections.Generic;
 using System.IO;
+using System.Linq;
+using System.Runtime.InteropServices;
 using System.Threading.Tasks;
 
 namespace DatasmithSolidworks
@@ -18,12 +20,156 @@ namespace DatasmithSolidworks
 		DirLightActor
 	};
 
+	// Wrapping Component Name and Actor Name to distinguish one from the other in the code. This helps to
+	//  - make explicit which name where
+	//  - avoid using sanitized name to check component's name and vice versa
+
+	[ComVisible(false)]
+	public struct FActorName
+	{
+		private string Value;
+
+		public FActorName(FComponentName ComponentName)
+		{
+			Value = FDatasmithExporter.SanitizeName(ComponentName.GetString());
+		}
+
+		public static FActorName FromString(string Value)
+		{
+			return new FActorName(Value);
+		}
+
+		private FActorName(string InValue)
+		{
+			Value = InValue;
+		}
+
+		public bool IsValid()
+		{
+			return !string.IsNullOrEmpty(Value);
+		}
+
+		public string GetString()
+		{
+			return Value;
+		}
+
+		public override string ToString()
+		{
+			return Value;
+		}
+
+		public override bool Equals(object Obj)
+		{
+			if (Obj is FActorName Other)
+			{
+				return this == Other;
+			}			
+			return false;
+		}
+
+		public static bool operator ==(FActorName A, FActorName B)
+		{
+			return A.Value == B.Value;
+		}
+
+		public static bool operator !=(FActorName A, FActorName B)
+		{
+			return !(A == B);
+		}
+
+		public override int GetHashCode()
+		{
+			return Value.GetHashCode();
+		}
+	}
+
+	[ComVisible(false)]
+	public struct FComponentName
+	{
+		private string Value;
+
+		public FComponentName(Component2 Component)
+		{
+			Value = Component.Name2;
+		}
+
+		public FComponentName(IComponent2 Component)
+		{
+			Value = Component.Name2;
+		}
+
+		private FComponentName(string Name)
+		{
+			Value = Name;
+		}
+
+		/// Used to convert string to component name ONLY when received from Solidworks API
+		public static FComponentName FromApiString(string Name)
+		{
+			return new FComponentName(Name);
+		}
+
+		/// Used to convert from a custom string to represent something 'like' a Solidworks component
+		/// todo: may want to avoid trying to mix components and non-components into the same 'entity'?
+		public static FComponentName FromCustomString(string Value)
+		{
+			return new FComponentName(Value);
+		}
+
+		public bool IsValid() => !string.IsNullOrEmpty(Value);
+
+		public string GetString()
+		{
+			return Value;
+		}
+
+		public FActorName GetActorName()
+		{
+			return new FActorName(this);
+		}
+
+		public string GetLabel()
+		{
+			return Value.Split('/').Last();
+		}
+
+		public override string ToString()
+		{
+			return Value;
+		}
+
+		public override bool Equals(object Obj)
+		{
+			if (Obj is FComponentName Other)
+			{
+				return this == Other;
+			}			
+			return false;
+		}
+
+		public static bool operator ==(FComponentName A, FComponentName B)
+		{
+			return A.Value == B.Value;
+		}
+
+		public static bool operator !=(FComponentName A, FComponentName B)
+		{
+			return !(A == B);
+		}
+
+		public override int GetHashCode()
+		{
+			return Value.GetHashCode();
+		}
+	}
+
 	public class FDatasmithActorExportInfo
 	{
 		public EActorType Type;
 		public string Label;
-		public string Name;
-		public string ParentName;
+		public FActorName Name;
+		public FActorName ParentName;
 		public string MeshName;
 		public float[] Transform;
 		public bool bVisible;
@@ -31,7 +177,7 @@ namespace DatasmithSolidworks
 
 	public class FDatasmithExporter
 	{
-		private Dictionary<string, Tuple<EActorType, FDatasmithFacadeActor>> ExportedActorsMap = new Dictionary<string, Tuple<EActorType, FDatasmithFacadeActor>>();
+		private Dictionary<FActorName, Tuple<EActorType, FDatasmithFacadeActor>> ExportedActorsMap = new Dictionary<FActorName, Tuple<EActorType, FDatasmithFacadeActor>>();
 		private ConcurrentDictionary<string, Tuple<FDatasmithFacadeMeshElement, FDatasmithFacadeMesh>> ExportedMeshesMap = new ConcurrentDictionary<string, Tuple<FDatasmithFacadeMeshElement, FDatasmithFacadeMesh>>();
 		private ConcurrentDictionary<int, FDatasmithFacadeMaterialInstance> ExportedMaterialsMap = new ConcurrentDictionary<int, FDatasmithFacadeMaterialInstance>();
 		private ConcurrentDictionary<string, FDatasmithFacadeTexture> ExportedTexturesMap = new ConcurrentDictionary<string, FDatasmithFacadeTexture>();
@@ -45,7 +191,7 @@ namespace DatasmithSolidworks
 			DatasmithScene = InScene;
 		}
 
-		public EActorType? GetExportedActorType(string InActorName)
+		public EActorType? GetExportedActorType(FActorName InActorName)
 		{
 			Tuple<EActorType, FDatasmithFacadeActor> ActorInfo = null;
 			if (ExportedActorsMap.TryGetValue(InActorName, out ActorInfo))
@@ -79,19 +225,20 @@ namespace DatasmithSolidworks
 			{
 				switch (InExportInfo.Type)
 				{
-					case EActorType.SimpleActor: Actor = new FDatasmithFacadeActor(InExportInfo.Name); break;
-					case EActorType.MeshActor: Actor = new FDatasmithFacadeActorMesh(InExportInfo.Name); break;
-					case EActorType.PointLightActor: Actor = new FDatasmithFacadePointLight(InExportInfo.Name); break;
-					case EActorType.SpotLightActor: Actor = new FDatasmithFacadeSpotLight(InExportInfo.Name); break;
-					case EActorType.DirLightActor: Actor = new FDatasmithFacadeDirectionalLight(InExportInfo.Name); break;
+					case EActorType.SimpleActor: Actor = new FDatasmithFacadeActor(InExportInfo.Name.GetString()); break;
+					case EActorType.MeshActor: Actor = new FDatasmithFacadeActorMesh(InExportInfo.Name.GetString()); break;
+					case EActorType.PointLightActor: Actor = new FDatasmithFacadePointLight(InExportInfo.Name.GetString()); break;
+					case EActorType.SpotLightActor: Actor = new FDatasmithFacadeSpotLight(InExportInfo.Name.GetString()); break;
+					case EActorType.DirLightActor: Actor = new FDatasmithFacadeDirectionalLight(InExportInfo.Name.GetString()); break;
 				}
 
-				Actor.AddTag(InExportInfo.Name);
+				Actor.AddTag(InExportInfo.Name.GetString());
 
 				ExportedActorsMap[InExportInfo.Name] = new Tuple<EActorType, FDatasmithFacadeActor>(InExportInfo.Type, Actor);
 
 				Tuple<EActorType, FDatasmithFacadeActor> ParentExportInfo = null;
-				if (!string.IsNullOrEmpty(InExportInfo.ParentName) && ExportedActorsMap.TryGetValue(InExportInfo.ParentName, out ParentExportInfo))
+
+				if (InExportInfo.ParentName.IsValid() && ExportedActorsMap.TryGetValue(InExportInfo.ParentName, out ParentExportInfo))
 				{
 					FDatasmithFacadeActor ParentActor = ParentExportInfo.Item2;
 					ParentActor.AddChild(Actor);
@@ -137,7 +284,7 @@ namespace DatasmithSolidworks
 			}
 		}
 
-		public void RemoveActor(string InActorName)
+		public void RemoveActor(FActorName InActorName)
 		{
 			if (ExportedActorsMap.ContainsKey(InActorName))
 			{
@@ -152,7 +299,7 @@ namespace DatasmithSolidworks
 		{
 			FDatasmithActorExportInfo ExportInfo = new FDatasmithActorExportInfo();
 			ExportInfo.Label = InLight.LightLabel;
-			ExportInfo.Name = InLight.LightName;
+			ExportInfo.Name = FActorName.FromString(InLight.LightName);
 			ExportInfo.bVisible = true;
 
 			FVec3 LightPosition = null;
@@ -208,7 +355,7 @@ namespace DatasmithSolidworks
 			}
 		}
 
-		public bool ExportMesh(string InMeshName, FMeshData InData, string InUpdateMeshActor, out Tuple<FDatasmithFacadeMeshElement, FDatasmithFacadeMesh> OutMeshPair)
+		public bool ExportMesh(string InMeshName, FMeshData InData, FActorName InUpdateMeshActor, out Tuple<FDatasmithFacadeMeshElement, FDatasmithFacadeMesh> OutMeshPair)
 		{
 			InMeshName = SanitizeName(InMeshName);  //Compute mesh name early(it might be needed to remove old mesh)
 
@@ -287,7 +434,7 @@ namespace DatasmithSolidworks
 
 			DatasmithScene.ExportDatasmithMesh(MeshElement, Mesh);
 
-			if (!string.IsNullOrEmpty(InUpdateMeshActor))
+			if (InUpdateMeshActor.IsValid())
 			{
 				Tuple<EActorType, FDatasmithFacadeActor> ExportedActorInfo = null;
 				if (ExportedActorsMap.TryGetValue(InUpdateMeshActor, out ExportedActorInfo) && ExportedActorInfo.Item1 == EActorType.MeshActor)
@@ -306,6 +453,7 @@ namespace DatasmithSolidworks
 
 			if (InMetadata.OwnerType == FMetadata.EOwnerType.Actor)
 			{
+				
 				if (ExportedActorsMap.ContainsKey(InMetadata.OwnerName))
 				{
 					Tuple<EActorType, FDatasmithFacadeActor> ActorInfo = ExportedActorsMap[InMetadata.OwnerName];
@@ -336,14 +484,17 @@ namespace DatasmithSolidworks
 			}
 		}
 
-		FDatasmithFacadeActorBinding GetActorBinding(string InActorName, FDatasmithFacadeVariant InVariant)
+		FDatasmithFacadeActorBinding GetActorBinding(FComponentName InComponentName, FDatasmithFacadeVariant InVariant)
 		{
-			string ActorNameSanitized = FDatasmithExporter.SanitizeName(InActorName);
+			return GetActorBinding(InComponentName.GetActorName(), InVariant);
+		}
 
+		private FDatasmithFacadeActorBinding GetActorBinding(FActorName ActorName, FDatasmithFacadeVariant InVariant)
+		{
 			for (int BindingIndex = 0; BindingIndex < InVariant.GetActorBindingsCount(); ++BindingIndex)
 			{
 				FDatasmithFacadeActorBinding Binding = InVariant.GetActorBinding(BindingIndex);
-				if (Binding.GetName() == ActorNameSanitized)
+				if (ActorName.Equals(Binding.GetName()))
 				{
 					return Binding;
 				}
@@ -353,7 +504,7 @@ namespace DatasmithSolidworks
 
 			FDatasmithFacadeActor Actor = null;
 
-			if (ExportedActorsMap.TryGetValue(ActorNameSanitized, out Tuple<EActorType, FDatasmithFacadeActor> ActorInfo))
+			if (ExportedActorsMap.TryGetValue(ActorName, out Tuple<EActorType, FDatasmithFacadeActor> ActorInfo))
 			{
 				Actor = ActorInfo.Item2;
 			}
@@ -546,9 +697,28 @@ namespace DatasmithSolidworks
 
 			if (ConfigurationVariants != null)
 			{
+				// todo: visibility variants may drop ComponentName from FConfigurationData
+				// at this point only actor names may stay
 				ExportActorVisibilityVariants(ConfigurationVariants);
 				ExportMaterialVariants(ConfigurationVariants);
 				ExportTransformVariants(ConfigurationVariants);
+
+				// Geometry variants
+				foreach (Tuple<FConfigurationData, FDatasmithFacadeVariant> KVP in ConfigurationVariants)
+				{
+					FConfigurationData Config = KVP.Item1;
+					FDatasmithFacadeVariant Variant = KVP.Item2;
+
+					// Make visible mesh actor corresponding this configuration only
+					foreach (FConfigurationData.FComponentGeometryVariant GeometryVariant in Config.ComponentGeometry.Values)
+					{
+						foreach (FActorName ActorName in GeometryVariant.All)
+						{
+							bool bVisible = GeometryVariant.VisibleActor == ActorName;
+							GetActorBinding(ActorName, Variant)?.AddVisibilityCapture(bVisible);
+						}
+					}
+				}
 			}
 
 			if (DisplayStateVariants != null)
diff --git a/Engine/Source/Programs/Enterprise/Datasmith/DatasmithSolidworksExporter/Private/Document.cs b/Engine/Source/Programs/Enterprise/Datasmith/DatasmithSolidworksExporter/Private/Document.cs
--- a/Engine/Source/Programs/Enterprise/Datasmith/DatasmithSolidworksExporter/Private/Document.cs
+++ b/Engine/Source/Programs/Enterprise/Datasmith/DatasmithSolidworksExporter/Private/Document.cs
@@ -169,7 +169,8 @@ namespace DatasmithSolidworks
 #if DEBUG
 				Stopwatch Watch = Stopwatch.StartNew();
 #endif
-				ExportToDatasmithScene();
+				FMeshes Meshes = new FMeshes();
+				ExportToDatasmithScene(Meshes);
 
 				ExportLights();
 
@@ -200,10 +201,14 @@ namespace DatasmithSolidworks
 
 				Exporter = new FDatasmithExporter(DatasmithScene);
 
-				List<FConfigurationData> Configs = FConfigurationExporter.ExportConfigurations(this);
+				FMeshes Meshes = new FMeshes();
+
+				PreExport(Meshes);
+
+				List<FConfigurationData> Configs = new FConfigurationExporter(Meshes).ExportConfigurations(this);
 				bHasConfigurations = (Configs != null) && (Configs.Count != 0);
 
-				ExportToDatasmithScene();
+				ExportToDatasmithScene(Meshes);
 			
 				ExportLights();
 				ExportConfigurations(Configs);
@@ -233,7 +238,8 @@ namespace DatasmithSolidworks
 		}
 
 		public abstract bool HasMaterialUpdates();
-		public abstract void ExportToDatasmithScene();
+		public abstract void ExportToDatasmithScene(FMeshes Meshes);
+		public abstract void PreExport(FMeshes Meshes);  // Called before configurations are parsed to prepare meshes needed to identify if they create different configurations
 
 		public virtual void Init()
 		{
diff --git a/Engine/Source/Programs/Enterprise/Datasmith/DatasmithSolidworksExporter/Private/MetadataManager.cs b/Engine/Source/Programs/Enterprise/Datasmith/DatasmithSolidworksExporter/Private/MetadataManager.cs
--- a/Engine/Source/Programs/Enterprise/Datasmith/DatasmithSolidworksExporter/Private/MetadataManager.cs
+++ b/Engine/Source/Programs/Enterprise/Datasmith/DatasmithSolidworksExporter/Private/MetadataManager.cs
@@ -56,7 +56,7 @@ namespace DatasmithSolidworks
 			None
 		}
 
-		public string OwnerName { get; set; } = "";
+		public FActorName OwnerName { get; set; }
 		public EOwnerType OwnerType = EOwnerType.None;
 		public List<IMetadataPair> Pairs = new List<IMetadataPair>();
 
diff --git a/Engine/Source/Programs/Enterprise/Datasmith/DatasmithSolidworksExporter/Private/ObjectMaterials.cs b/Engine/Source/Programs/Enterprise/Datasmith/DatasmithSolidworksExporter/Private/ObjectMaterials.cs
--- a/Engine/Source/Programs/Enterprise/Datasmith/DatasmithSolidworksExporter/Private/ObjectMaterials.cs
+++ b/Engine/Source/Programs/Enterprise/Datasmith/DatasmithSolidworksExporter/Private/ObjectMaterials.cs
@@ -291,12 +291,12 @@ namespace DatasmithSolidworks
 			return ComponentMaterials;
 		}
 
-		public static ConcurrentDictionary<string, FObjectMaterials> LoadAssemblyMaterials(FAssemblyDocument InAsmDoc, HashSet<string> InComponentsSet, swDisplayStateOpts_e InDisplayState, string[] InDisplayStateNames)
+		public static ConcurrentDictionary<FComponentName, FObjectMaterials> LoadAssemblyMaterials(FAssemblyDocument InAsmDoc, HashSet<FComponentName> InComponentsSet, swDisplayStateOpts_e InDisplayState, string[] InDisplayStateNames)
 		{
 			IModelDocExtension Ext = InAsmDoc.SwDoc.Extension;
 			int NumMaterials = Ext.GetRenderMaterialsCount2((int)InDisplayState, InDisplayStateNames);
 
-			ConcurrentDictionary<string, FObjectMaterials> DocMaterials = new ConcurrentDictionary<string, FObjectMaterials>();
+			ConcurrentDictionary<FComponentName, FObjectMaterials> DocMaterials = new ConcurrentDictionary<FComponentName, FObjectMaterials>();
 
 			if (NumMaterials > 0)
 			{
@@ -322,7 +322,7 @@ namespace DatasmithSolidworks
 							{
 								FObjectMaterials ComponentMaterials = new FObjectMaterials(InAsmDoc, CompDoc, ref InAsmDoc.ExportedMaterialsMap);
 								ComponentMaterials.SetComponentMaterial(RenderMat, InAsmDoc.SwDoc);
-								DocMaterials[Comp.Name2] = ComponentMaterials;
+								DocMaterials[new FComponentName(Comp)] = ComponentMaterials;
 							}
 							continue;
 						}
@@ -332,7 +332,9 @@ namespace DatasmithSolidworks
 							FObjectMaterials PartMaterials = LoadPartMaterials(InAsmDoc, Doc as PartDoc, InDisplayState, InDisplayStateNames);
 							if (PartMaterials != null)
 							{
-								DocMaterials[(Doc as ModelDoc2).GetPathName()] = PartMaterials;
+								// xxx: why trying to present PahtName as component name?
+								// xxx: how it can be a PartDoc in Assembly function
+								DocMaterials[FComponentName.FromCustomString((Doc as ModelDoc2).GetPathName())] = PartMaterials;
 							}
 							continue;
 						}
diff --git a/Engine/Source/Programs/Enterprise/Datasmith/DatasmithSolidworksExporter/Private/PartDocument.cs b/Engine/Source/Programs/Enterprise/Datasmith/DatasmithSolidworksExporter/Private/PartDocument.cs
--- a/Engine/Source/Programs/Enterprise/Datasmith/DatasmithSolidworksExporter/Private/PartDocument.cs
+++ b/Engine/Source/Programs/Enterprise/Datasmith/DatasmithSolidworksExporter/Private/PartDocument.cs
@@ -12,14 +12,14 @@ namespace DatasmithSolidworks
 	{
 		public PartDoc SwPartDoc { get; private set; } = null;
 		private FAssemblyDocument AsmDoc = null;
-		private string ComponentName = null;
+		private FComponentName ComponentName;
 		private string MeshName = null;
 
 		public FObjectMaterials ExportedPartMaterials { get; set; }
 
 		public string PathName { get; private set; }
 
-		public FPartDocument(int InDocId, PartDoc InPartDoc, FDatasmithExporter InExporter, FAssemblyDocument InAsmDoc, string InComponentName) 
+		public FPartDocument(int InDocId, PartDoc InPartDoc, FDatasmithExporter InExporter, FAssemblyDocument InAsmDoc, FComponentName InComponentName) 
 			: base(InDocId, InPartDoc as ModelDoc2, InExporter)
 		{
 			SwPartDoc = InPartDoc;
@@ -72,45 +72,119 @@ namespace DatasmithSolidworks
 		{
 			base.SetDirty(bInDirty);
 
-			if (bInDirty && AsmDoc != null && !string.IsNullOrEmpty(ComponentName))
+			if (bInDirty && AsmDoc != null && ComponentName.IsValid())
 			{
 				// Notify owner assembly that a component needs re-export
 				AsmDoc.SetComponentDirty(ComponentName, FAssemblyDocument.EComponentDirtyState.Geometry);
 			}
 		}
 
-		public override void ExportToDatasmithScene()
+		public override void PreExport(FMeshes Meshes)
 		{
-			string ActorName = Path.GetFileNameWithoutExtension(PathName);
+			ConfigurationManager ConfigManager = (SwPartDoc as ModelDoc2).ConfigurationManager;
+			IConfiguration OriginalConfiguration = ConfigManager.ActiveConfiguration;
 
-			SetExportStatus($"{ActorName} Materials");
+			foreach (string ConfigurationName in SwDoc.GetConfigurationNames())
+			{
+				SwDoc.ShowConfiguration2(ConfigurationName);
+
+				// todo: Solidworks api docs says that GetRootComponent3 is null for Part but our current 
+				//   parsing code relies on it's existence(and it exists in all cases). Probably better to refactor this without Component use for parts
+				//   This is possible, we don't need Component, just code would be be less 'generic'   
+				Component2 Component = ConfigManager.ActiveConfiguration.GetRootComponent3(true);
+
+				ConcurrentBag<FBody> Bodies = FBody.FetchBodies(SwPartDoc);
+				FMeshData MeshData = FStripGeometry.CreateMeshData(Bodies, ExportedPartMaterials);
+
+				if (MeshData != null)
+				{
+					Meshes.GetMeshesConfiguration(ConfigManager.ActiveConfiguration.Name).AddMesh(Component, MeshData);
+				}
+			}
+
+			if (OriginalConfiguration != null)
+			{
+				SwDoc.ShowConfiguration2(OriginalConfiguration.Name);
+			}
+		}
+
+		public override void ExportToDatasmithScene(FMeshes Meshes)
+		{
+			string PartName = Path.GetFileNameWithoutExtension(PathName);
+
+			SetExportStatus($"{PartName} Materials");
 
 			ExportedPartMaterials = FObjectMaterials.LoadPartMaterials(this, SwPartDoc, swDisplayStateOpts_e.swThisDisplayState, null);
 
 			Exporter.ExportMaterials(ExportedMaterialsMap);
 
-			SetExportStatus($"{ActorName} Meshes");
+			SetExportStatus($"{PartName} Meshes");
 
-			ConcurrentBag<FBody> Bodies = FBody.FetchBodies(SwPartDoc);
-			FMeshData MeshData = FStripGeometry.CreateMeshData(Bodies, ExportedPartMaterials);
+			if (bHasConfigurations)
+			{
+				ConfigurationManager ConfigManager = (SwPartDoc as ModelDoc2).ConfigurationManager;
+				string ActiveConfigurationName = ConfigManager.ActiveConfiguration.Name;
 
-			Exporter.RemoveMesh(MeshName);
-			MeshName = null;
-			bool bHasMesh = Exporter.ExportMesh($"{ActorName}_Mesh", MeshData, ActorName, out Tuple<FDatasmithFacadeMeshElement, FDatasmithFacadeMesh> NewMesh);
+				foreach (string ConfigurationName in SwDoc.GetConfigurationNames())
+				{
 
-			if (bHasMesh)
-			{
-				MeshName = Exporter.AddMesh(NewMesh);
+					// todo: might want to not use GetRootComponent3 for Part - docs says that it's null for Part
+					FComponentName RootComponentName = new FComponentName(ConfigManager.ActiveConfiguration.GetRootComponent3(true));
+					FMeshData MeshData = Meshes.GetMeshData(ConfigurationName, RootComponentName);
+					if (MeshData == null)
+					{
+						continue;
+					}
+					// todo: not make extra copy - use existing(i.e. instead of bHasConfigurations) 
+					FActorName ActorName = new FConfigurationExporter(Meshes).GetMeshActorName(RootComponentName, ConfigurationName);
+
+					bool bHasMesh = Exporter.ExportMesh($"{ActorName}_Mesh", MeshData, ActorName,
+						out Tuple<FDatasmithFacadeMeshElement, FDatasmithFacadeMesh> NewMesh);
+
+					string MeshNameForConfiguration = null; // todo: Sync not supported for variants, 
+					if (bHasMesh)
+					{
+						MeshNameForConfiguration = Exporter.AddMesh(NewMesh);
+					}
+
+					FDatasmithActorExportInfo ExportInfo = new FDatasmithActorExportInfo();
+					ExportInfo.Name = ActorName;
+					ExportInfo.bVisible = ConfigurationName == ActiveConfigurationName;
+					ExportInfo.Label = ActorName.GetString();
+					ExportInfo.MeshName = MeshNameForConfiguration;
+					ExportInfo.Type = bHasMesh ? EActorType.MeshActor : EActorType.SimpleActor;
+					Exporter.ExportOrUpdateActor(ExportInfo);
+				}
 			}
+			else
+			{
+				FActorName ActorName = FActorName.FromString(PartName);
+				ConcurrentBag<FBody> Bodies = FBody.FetchBodies(SwPartDoc);
+				FMeshData MeshData = FStripGeometry.CreateMeshData(Bodies, ExportedPartMaterials);
+
+				Exporter.RemoveMesh(MeshName);
+				MeshName = null;
+				bool bHasMesh = false;
+
+				if (MeshData != null)
+				{
+					bHasMesh = Exporter.ExportMesh($"{ActorName}_Mesh", MeshData, ActorName,
+						out Tuple<FDatasmithFacadeMeshElement, FDatasmithFacadeMesh> NewMesh);
 
-			FDatasmithActorExportInfo ExportInfo = new FDatasmithActorExportInfo();
-			ExportInfo.Name = ActorName;
-			ExportInfo.bVisible = true;
-			ExportInfo.Label = ActorName;
-			ExportInfo.MeshName = MeshName;
-			ExportInfo.Type = bHasMesh ? EActorType.MeshActor : EActorType.SimpleActor;
+					if (bHasMesh)
+					{
+						MeshName = Exporter.AddMesh(NewMesh);
+					}
+				}
 
-			Exporter.ExportOrUpdateActor(ExportInfo);
+				FDatasmithActorExportInfo ExportInfo = new FDatasmithActorExportInfo();
+				ExportInfo.Name = ActorName;
+				ExportInfo.bVisible = true;
+				ExportInfo.Label = ActorName.GetString();
+				ExportInfo.MeshName = MeshName;
+				ExportInfo.Type = bHasMesh ? EActorType.MeshActor : EActorType.SimpleActor;
+				Exporter.ExportOrUpdateActor(ExportInfo);
+			}
 		}
 
 		public override bool HasMaterialUpdates()
diff --git a/Engine/Source/Programs/Enterprise/Datasmith/DatasmithSolidworksExporter/Private/Triangle.cs b/Engine/Source/Programs/Enterprise/Datasmith/DatasmithSolidworksExporter/Private/Triangle.cs
--- a/Engine/Source/Programs/Enterprise/Datasmith/DatasmithSolidworksExporter/Private/Triangle.cs
+++ b/Engine/Source/Programs/Enterprise/Datasmith/DatasmithSolidworksExporter/Private/Triangle.cs
@@ -29,11 +29,48 @@ namespace DatasmithSolidworks
 			MaterialID = InMaterialID;
 		}
 
+		public static bool operator ==(FTriangle A, FTriangle B)
+		{
+			if (ReferenceEquals(A, B))
+			{
+				return true;
+			}
+
+			if (A is null || B is null)
+			{
+				return false;
+			}
+
+			return A.Index1 == B.Index1 && A.Index2 == B.Index2 && A.Index3 == B.Index3 && A.MaterialID == B.MaterialID;
+		}
+
+		public static bool operator !=(FTriangle A, FTriangle B)
+		{
+			return !(A == B);
+		}
+
+		public override bool Equals(object Obj)
+		{
+			return Obj is FTriangle Other && this == Other;
+		}
+
 		public FTriangle Offset(int InOffset)
 		{
 			return new FTriangle(Indices[0] + InOffset, Indices[1] + InOffset, Indices[2] + InOffset, MaterialID);
 		}
 
+		public override int GetHashCode() 
+		{
+			int Hash = 7;
+
+			foreach (int Index in Indices)
+			{
+				Hash = Hash * 17 + Index.GetHashCode();
+			}
+
+			return Hash;
+		}
+
 		public override string ToString()
 		{
 			return "" + Index1 + "," + Index2 + "," + Index3;
diff --git a/Engine/Source/Programs/Enterprise/Datasmith/DatasmithSolidworksExporter/Private/Vec2.cs b/Engine/Source/Programs/Enterprise/Datasmith/DatasmithSolidworksExporter/Private/Vec2.cs
--- a/Engine/Source/Programs/Enterprise/Datasmith/DatasmithSolidworksExporter/Private/Vec2.cs
+++ b/Engine/Source/Programs/Enterprise/Datasmith/DatasmithSolidworksExporter/Private/Vec2.cs
@@ -29,6 +29,39 @@ namespace DatasmithSolidworks
             Y = InY;
         }
 
+        public static bool operator ==(FVec2 A, FVec2 B)
+        {
+	        if (ReferenceEquals(A, B))
+	        {
+		        return true;
+	        }
+
+	        if (A is null || B is null)
+	        {
+		        return false;
+	        }
+
+	        return A.X == B.X && A.Y == B.Y;
+        }
+
+        public static bool operator !=(FVec2 A, FVec2 B)
+        {
+	        return !(A == B);
+        }
+
+        public override bool Equals(object Obj)
+        {
+	        return Obj is FVec2 Other && this == Other;
+        }
+
+        public override int GetHashCode()
+        {
+	        int Hash = 7;
+	        Hash = Hash * 17 + X.GetHashCode();
+	        Hash = Hash * 17 + Y.GetHashCode();
+	        return Hash;
+        }
+
         public static FVec2 Rotate(FVec2 InV, float InAngleRadians)
         {
             return new FVec2(
diff --git a/Engine/Source/Programs/Enterprise/Datasmith/DatasmithSolidworksExporter/Private/Vec3.cs b/Engine/Source/Programs/Enterprise/Datasmith/DatasmithSolidworksExporter/Private/Vec3.cs
--- a/Engine/Source/Programs/Enterprise/Datasmith/DatasmithSolidworksExporter/Private/Vec3.cs
+++ b/Engine/Source/Programs/Enterprise/Datasmith/DatasmithSolidworksExporter/Private/Vec3.cs
@@ -59,9 +59,29 @@ namespace DatasmithSolidworks
 			Z = (float)Data[2];
         }
 
-        public override bool Equals(object InObj)
+        public static bool operator ==(FVec3 A, FVec3 B)
         {
-            return base.Equals(InObj);
+	        if (ReferenceEquals(A, B))
+	        {
+				return true;
+	        }
+
+	        if (A is null || B is null)
+	        {
+		        return false;
+	        }
+
+	        return A.X == B.X && A.Y == B.Y && A.Z == B.Z;
+        }
+
+        public static bool operator !=(FVec3 A, FVec3 B)
+        {
+	        return !(A == B);
+        }
+
+        public override bool Equals(object Obj)
+        {
+            return Obj is FVec3 Other && this == Other;
         }
 
         public override int GetHashCode()
