Added support for ArchiCAD26 to the exporter on the Mac

diff --git a/Engine/Source/Programs/Enterprise/Datasmith/DatasmithARCHICADExporter/Build/Config/UE_AC.xcconfig b/Engine/Source/Programs/Enterprise/Datasmith/DatasmithARCHICADExporter/Build/Config/UE_AC.xcconfig
--- a/Engine/Source/Programs/Enterprise/Datasmith/DatasmithARCHICADExporter/Build/Config/UE_AC.xcconfig
+++ b/Engine/Source/Programs/Enterprise/Datasmith/DatasmithARCHICADExporter/Build/Config/UE_AC.xcconfig
@@ -23,7 +23,7 @@ VRevision = 0
 VBetaOrRevision = b // b for beta or . for revision
 VAddOn = $(VMajor).$(VMinor)$(VBetaOrRevision)$(VRevision)
 
-CurrentYear = 2021
+CurrentYear = 2022
 
 FRAMEWORK_SEARCH_PATHS = $(FRAMEWORK_SEARCH_PATHS) $(AC_SDK_PATH)/Support/Frameworks
 
diff --git a/Engine/Source/Programs/Enterprise/Datasmith/DatasmithARCHICADExporter/Build/Config/app_to_bundle.sh b/Engine/Source/Programs/Enterprise/Datasmith/DatasmithARCHICADExporter/Build/Config/app_to_bundle.sh
--- a/Engine/Source/Programs/Enterprise/Datasmith/DatasmithARCHICADExporter/Build/Config/app_to_bundle.sh
+++ b/Engine/Source/Programs/Enterprise/Datasmith/DatasmithARCHICADExporter/Build/Config/app_to_bundle.sh
@@ -11,4 +11,5 @@ pushd $EnginePath/Binaries/Mac/DatasmithARCHICADExporter
 mv DatasmithARCHICAD23Exporter.app DatasmithARCHICAD23Exporter.bundle
 mv DatasmithARCHICAD24Exporter.app DatasmithARCHICAD24Exporter.bundle
 mv DatasmithARCHICAD25Exporter.app DatasmithARCHICAD25Exporter.bundle
+mv DatasmithARCHICAD26Exporter.app DatasmithARCHICAD26Exporter.bundle
 popd
diff --git a/Engine/Source/Programs/Enterprise/Datasmith/DatasmithARCHICADExporter/Build/Config/bundle_to_app.sh b/Engine/Source/Programs/Enterprise/Datasmith/DatasmithARCHICADExporter/Build/Config/bundle_to_app.sh
--- a/Engine/Source/Programs/Enterprise/Datasmith/DatasmithARCHICADExporter/Build/Config/bundle_to_app.sh
+++ b/Engine/Source/Programs/Enterprise/Datasmith/DatasmithARCHICADExporter/Build/Config/bundle_to_app.sh
@@ -11,4 +11,5 @@ pushd $EnginePath/Binaries/Mac/DatasmithARCHICADExporter
 mv DatasmithARCHICAD23Exporter.bundle DatasmithARCHICAD23Exporter.app
 mv DatasmithARCHICAD24Exporter.bundle DatasmithARCHICAD24Exporter.app
 mv DatasmithARCHICAD25Exporter.bundle DatasmithARCHICAD25Exporter.app
+mv DatasmithARCHICAD26Exporter.bundle DatasmithARCHICAD26Exporter.app
 popd
diff --git a/Engine/Source/Programs/Enterprise/Datasmith/DatasmithARCHICADExporter/Build/Config/setupdylibs.sh b/Engine/Source/Programs/Enterprise/Datasmith/DatasmithARCHICADExporter/Build/Config/setupdylibs.sh
--- a/Engine/Source/Programs/Enterprise/Datasmith/DatasmithARCHICADExporter/Build/Config/setupdylibs.sh
+++ b/Engine/Source/Programs/Enterprise/Datasmith/DatasmithARCHICADExporter/Build/Config/setupdylibs.sh
@@ -3,10 +3,11 @@
 set -e
 
 ConfigPath=`dirname "$0"`
-projectPath=$ConfigPath/..
-
-RelativeEnginePath=$projectPath/../../../../../../../Engine
-EnginePath=`python -c "import os; print(os.path.realpath('$RelativeEnginePath'))"`
+pushd $ConfigPath/..
+projectPath=`pwd`
+cd ../../../../../..
+EnginePath=`pwd`
+popd
 
 echo "UE_SDKS_ROOT = ${UE_SDKS_ROOT}"
 
@@ -31,6 +32,9 @@ pushd "$UE_SDKS_ROOT/HostMac/Mac/Archicad"
 
 	chmod 777 25/Support/Tools/OSX/ResConv
 	xattr -r -d com.apple.quarantine 25/Support/Tools/OSX/ResConv
+
+	chmod 777 26/Support/Tools/OSX/ResConv
+	xattr -r -d com.apple.quarantine 26/Support/Tools/OSX/ResConv
 popd
 
 OurDylibFolder=$projectPath/Dylibs
diff --git a/Engine/Source/Programs/Enterprise/Datasmith/DatasmithARCHICADExporter/Private/ConnectionWindow.cpp b/Engine/Source/Programs/Enterprise/Datasmith/DatasmithARCHICADExporter/Private/ConnectionWindow.cpp
--- a/Engine/Source/Programs/Enterprise/Datasmith/DatasmithARCHICADExporter/Private/ConnectionWindow.cpp
+++ b/Engine/Source/Programs/Enterprise/Datasmith/DatasmithARCHICADExporter/Private/ConnectionWindow.cpp
@@ -379,6 +379,13 @@ class FConnectionDialog : public DG::Palette,
 	{
 		CacheFolderText.SetText(GetGSName(kName_CacheDirectory) + FSyncDatabase::GetCachePath());
 	}
+
+#if PLATFORM_MAC & AC_VERSION > 25
+	virtual void ItemMouseExited(const DG::ItemMouseMoveEvent& /*ev*/) override {}
+	virtual void ItemMouseEntered(const DG::ItemMouseMoveEvent& /*ev*/) override {}
+	virtual short SpecMouseExited(const DG::ItemMouseMoveEvent& /*ev*/) override { return 0; }
+	virtual short SpecMouseEntered(const DG::ItemMouseMoveEvent& /*ev*/) override { return 0; }
+#endif
 };
 
 FConnectionDialog::FConnectionDialog()
diff --git a/Engine/Source/Programs/Enterprise/Datasmith/DatasmithARCHICADExporter/Private/Palette.cpp b/Engine/Source/Programs/Enterprise/Datasmith/DatasmithARCHICADExporter/Private/Palette.cpp
--- a/Engine/Source/Programs/Enterprise/Datasmith/DatasmithARCHICADExporter/Private/Palette.cpp
+++ b/Engine/Source/Programs/Enterprise/Datasmith/DatasmithARCHICADExporter/Private/Palette.cpp
@@ -81,6 +81,13 @@ class FPaletteWindow : public DG::Palette,
 		DetachFromAllItems(*this);
 	}
 
+#if PLATFORM_MAC & AC_VERSION > 25
+	virtual void ItemMouseExited(const DG::ItemMouseMoveEvent& /*ev*/) override {}
+	virtual void ItemMouseEntered(const DG::ItemMouseMoveEvent& /*ev*/) override {}
+	virtual short SpecMouseExited(const DG::ItemMouseMoveEvent& /*ev*/) override { return 0; }
+	virtual short SpecMouseEntered(const DG::ItemMouseMoveEvent& /*ev*/) override { return 0; }
+#endif
+
   private:
 	virtual void PanelOpened(const DG::PanelOpenEvent& /*ev*/) override
 	{
diff --git a/Engine/Source/Programs/Enterprise/Datasmith/DatasmithARCHICADExporter/Private/ReportWindow.cpp b/Engine/Source/Programs/Enterprise/Datasmith/DatasmithARCHICADExporter/Private/ReportWindow.cpp
--- a/Engine/Source/Programs/Enterprise/Datasmith/DatasmithARCHICADExporter/Private/ReportWindow.cpp
+++ b/Engine/Source/Programs/Enterprise/Datasmith/DatasmithARCHICADExporter/Private/ReportWindow.cpp
@@ -119,7 +119,14 @@ class FReportDialog : public DG::Palette,
             MessagesTextEdit.SetSelection(Selection);
         }
 	}
-};
+
+#if PLATFORM_MAC & AC_VERSION > 25
+	virtual void ItemMouseExited(const DG::ItemMouseMoveEvent& /*ev*/) override {}
+	virtual void ItemMouseEntered(const DG::ItemMouseMoveEvent& /*ev*/) override {}
+	virtual short SpecMouseExited(const DG::ItemMouseMoveEvent& /*ev*/) override { return 0; }
+	virtual short SpecMouseEntered(const DG::ItemMouseMoveEvent& /*ev*/) override { return 0; }
+#endif
+					  };
 
 FReportDialog::FReportDialog()
 	: DG::Palette(ACAPI_GetOwnResModule(), LocalizeResId(kDlgReport), ACAPI_GetOwnResModule(), PaletteGuid)
