Revit Exporter Plugin Crash on Decals Export fix,
Temporary fix for Revit plugin sends constant updates, even when nothing changed

diff --git a/Engine/Source/Programs/Enterprise/Datasmith/DatasmithRevitExporter/Private/DatasmithRevitDirectLink.cs b/Engine/Source/Programs/Enterprise/Datasmith/DatasmithRevitExporter/Private/DatasmithRevitDirectLink.cs
--- a/Engine/Source/Programs/Enterprise/Datasmith/DatasmithRevitExporter/Private/DatasmithRevitDirectLink.cs
+++ b/Engine/Source/Programs/Enterprise/Datasmith/DatasmithRevitExporter/Private/DatasmithRevitDirectLink.cs
@@ -216,8 +216,8 @@ namespace DatasmithRevitExporter
 				return;
 			}
 
-			// We want to keep this number low in order to ensure functioning Revit GUI
-			int MetadataCountPerIdleEvent = 2;
+			// MetadataCountPerIdleEvent value will be revisited in a future release since it also requires changes in the import of metadata.
+			int MetadataCountPerIdleEvent = int.MaxValue;
 
 			string EnvBatchSize = Environment.GetEnvironmentVariable("REVIT_DIRECTLINK_METADATA_BATCH_SIZE");
 			if (!string.IsNullOrEmpty(EnvBatchSize))
diff --git a/Engine/Source/Programs/Enterprise/Datasmith/DatasmithRevitExporter/Private/DatasmithRevitDocumentData.cs b/Engine/Source/Programs/Enterprise/Datasmith/DatasmithRevitExporter/Private/DatasmithRevitDocumentData.cs
--- a/Engine/Source/Programs/Enterprise/Datasmith/DatasmithRevitExporter/Private/DatasmithRevitDocumentData.cs
+++ b/Engine/Source/Programs/Enterprise/Datasmith/DatasmithRevitExporter/Private/DatasmithRevitDocumentData.cs
@@ -1660,6 +1660,11 @@ namespace DatasmithRevitExporter
 					IList<FDecalMaterial> DecalMaterials = new List<FDecalMaterial>();
 					foreach (KeyValuePair<ElementId, Asset> DecalIdAndAssetPair in ReferencingDecalIdAndAssetPairs)
 					{
+						if (DecalMaterialsMap.ContainsKey(DecalIdAndAssetPair.Key))
+						{
+							continue;
+						}
+
 						FDecalMaterial DecalMaterial = FDecalMaterial.Create(InMaterialNode, CurrentMaterial, DecalMaterials.Count, DecalIdAndAssetPair.Value);
 
 						if (DecalMaterial != null)
