[3ds Max Exporter] DirectLink sync fixed for those modifiers that are not sending a CHANGE event

- instead begin edint event is used
- this was reproduced with Body Cutter - when event wasn't sent to the object that a body cutter is made from

diff --git a/Engine/Source/Programs/Enterprise/Datasmith/DatasmithMaxExporter/Private/DatasmithMaxDirectLinkNotifications.cpp b/Engine/Source/Programs/Enterprise/Datasmith/DatasmithMaxExporter/Private/DatasmithMaxDirectLinkNotifications.cpp
--- a/Engine/Source/Programs/Enterprise/Datasmith/DatasmithMaxExporter/Private/DatasmithMaxDirectLinkNotifications.cpp
+++ b/Engine/Source/Programs/Enterprise/Datasmith/DatasmithMaxExporter/Private/DatasmithMaxDirectLinkNotifications.cpp
@@ -450,6 +450,12 @@ public:
 
 		break;
 
+		case REFMSG_BEGIN_EDIT: 
+			// Workaround for missing events for some modifiers
+			// E.g. Body Cutter won't sent a CHANGE event an object is converted to it
+			SceneTracker.NodeGeometryChanged(&Node);
+		break;
+
 		case REFMSG_NODE_MATERIAL_CHANGED:
 			SceneTracker.NodeMaterialAssignmentChanged(&Node);
 		break;
