Set BuildSettingsVersion.V2 rather than Latest on some program targets, as V3 is currently being worked on and is not ready to be enabled

#rnx

diff --git a/Engine/Source/Programs/Enterprise/Datasmith/DatasmithSDK/DatasmithSDK.Target.cs b/Engine/Source/Programs/Enterprise/Datasmith/DatasmithSDK/DatasmithSDK.Target.cs
--- a/Engine/Source/Programs/Enterprise/Datasmith/DatasmithSDK/DatasmithSDK.Target.cs
+++ b/Engine/Source/Programs/Enterprise/Datasmith/DatasmithSDK/DatasmithSDK.Target.cs
@@ -10,7 +10,7 @@ public class DatasmithSDKTarget : TargetRules
 	{
 		Type = TargetType.Program;
 		IncludeOrderVersion = EngineIncludeOrderVersion.Latest;
-		DefaultBuildSettings = BuildSettingsVersion.Latest;
+		DefaultBuildSettings = BuildSettingsVersion.V2;
 		SolutionDirectory = "Programs/Datasmith";
 		bBuildInSolutionByDefault = false;
 
