[3ds Max Exporter] Fixed mesh quality for Body Cutter

- TriObject is used to extract geometry

- For some object types in Max GetRenderMesh doesn't give away the finest mesh even when viewport mesh is set to highest quality. It's the case for Body Cutter

diff --git a/Engine/Source/Programs/Enterprise/Datasmith/DatasmithMaxExporter/Private/DatasmithMaxDirectLinkMesh.cpp b/Engine/Source/Programs/Enterprise/Datasmith/DatasmithMaxExporter/Private/DatasmithMaxDirectLinkMesh.cpp
--- a/Engine/Source/Programs/Enterprise/Datasmith/DatasmithMaxExporter/Private/DatasmithMaxDirectLinkMesh.cpp
+++ b/Engine/Source/Programs/Enterprise/Datasmith/DatasmithMaxExporter/Private/DatasmithMaxDirectLinkMesh.cpp
@@ -45,18 +45,95 @@ public:
 	}
 };
 
+Object* GetBaseObject(INode* Node, TimeValue Time);
+int SetObjectParamValue(TimeValue CurrentTime, Object *Obj, const FString& ParamName, int DesiredValue);
+
 FRenderMeshForConversion GetMeshForGeomObject(TimeValue CurrentTime, INode* Node, const FTransform& Pivot)
 {
 	// todo: baseline exporter uses GetBaseObject which takes result of EvalWorldState
 	// and searched down DerivedObject pipeline(by taking GetObjRef) 
 	// This is STRANGE as EvalWorldState shouldn't return DerivedObject in the first place(it should return result of pipeline evaluation)
 
-	BOOL bNeedsDelete;
+	BOOL bNeedsDelete = false;
 	Interval ValidityInterval;
 	ValidityInterval.SetInfinite();
-	Mesh* RenderMesh = GetMeshFromRenderMesh(CurrentTime, Node, ValidityInterval, bNeedsDelete);
 
-	FRenderMeshForConversion Result(Node, RenderMesh, bNeedsDelete);
+	TriObject* ObjectToDelete = nullptr;
+
+	Mesh* RenderMesh = nullptr;
+
+	{
+		Object* Obj = GetBaseObject(Node, CurrentTime);
+
+		// Read validity interval before changing display settings
+		// Else VRay will somehow react to a call to ObjectValidity and change RenderMesh contents to what it was originally.
+		// E.g. display was set 'Box', we change display to 'Mesh' here before calling GetRenderMesh to retrieve actual mesh later(rather than a simple box)
+		//  but if ObjectValidity is called after GetRenderMesh that RenderMesh will become box again!
+		Interval ObjectValidity = Obj->ObjectValidity(CurrentTime);
+
+		const Class_ID& ObjectClassID = Obj->ClassID();
+		const FString VRayProxyParamName(TEXT("display"));
+		const FString BodyObjectViewportMeshParamName(TEXT("RenderViewportMeshRA"));
+		int PreviousMeshDisplayValue = 0;
+
+		const int VRAY_PROXY_DISPLAY_AS_MESH = 4; // Value to set on a VRay Mesh Proxy to get the mesh
+		if (ObjectClassID == VRAYPROXY_CLASS_ID)
+		{
+			// Need the high resolution render mesh associated with the VRay Mesh Proxy for the export
+			PreviousMeshDisplayValue = SetObjectParamValue(CurrentTime, Obj, VRayProxyParamName, VRAY_PROXY_DISPLAY_AS_MESH);
+		}
+		else if(ObjectClassID == BODYOBJECT_CLASS_ID)
+		{
+			// Need to make sure we are using the viewport mesh on BodyObject, otherwise the RenderMesh gives a tessellated low resolution mesh.
+			PreviousMeshDisplayValue = SetObjectParamValue(CurrentTime, Obj, BodyObjectViewportMeshParamName, 1);
+		}
+
+		GeomObject* GeomObj = static_cast<GeomObject*>(Obj);
+		if (GeomObj)
+		{
+			if (GeomObj->CanConvertToType(Class_ID(TRIOBJ_CLASS_ID, 0)))
+			{
+				// GeomObj->GetRenderMesh for some objects types returns low quality mesh no matter settings
+				// Example: a body cutter object, its GetRenderMesh gives away a coarse mesh even when its viewport settings are set to Fine(and it's displayed Fine in viewport)
+				// Looks like TriObject should be used in the first place as a source of best quality mesh
+
+				TriObject* TriObj = static_cast<TriObject*>(GeomObj->ConvertToType(CurrentTime, Class_ID(TRIOBJ_CLASS_ID, 0)));
+
+				if (TriObj != GeomObj)
+				{
+					// In case tri object is a different object from node's GeomObj then use it to get mesh
+					RenderMesh = &(TriObj->mesh);
+					bNeedsDelete = false; // Mesh doesn't need to be deleted - TriObject does
+
+					check(TriObj != GeomObj); // Extra check that we won't delete live object 
+					ObjectToDelete = TriObj; // When converted tri object is different from th eobject itself it needs to be disposed
+				}
+			}
+
+			if (!RenderMesh)  // Fallback to GetRenderMesh if TriObject fails
+			{
+				FNullView View;
+				RenderMesh = GeomObj->GetRenderMesh(CurrentTime, Node, View, bNeedsDelete);
+			}
+		}
+
+		// Restore display state if different from mesh display
+		if (ObjectClassID == VRAYPROXY_CLASS_ID && PreviousMeshDisplayValue != VRAY_PROXY_DISPLAY_AS_MESH)
+		{
+			SetObjectParamValue(CurrentTime, Obj, VRayProxyParamName, PreviousMeshDisplayValue);
+		}
+		else if(ObjectClassID == BODYOBJECT_CLASS_ID && PreviousMeshDisplayValue != 1)
+		{
+			SetObjectParamValue(CurrentTime, Obj, BodyObjectViewportMeshParamName, PreviousMeshDisplayValue);
+		}
+
+		if (RenderMesh)
+		{
+			ValidityInterval &= ObjectValidity; // Update validity only when actual mesh is returned
+		}
+	}
+
+	FRenderMeshForConversion Result(Node, RenderMesh, bNeedsDelete, ObjectToDelete);
 	Result.SetValidityInterval(ValidityInterval);
 	Result.SetPivot(Pivot);
 	return MoveTemp(Result);
@@ -432,8 +509,6 @@ bool CreateDatasmithMeshFromMaxMesh(FDatasmithMesh& DatasmithMesh, const MeshCon
 	return bResult;
 }
 
-// todo: copied from DatasmithMaxMeshExporter.cpp
-
 Object* GetBaseObject(INode* Node, TimeValue Time)
 {
 	ObjectState ObjState = Node->EvalWorldState(Time);
@@ -452,7 +527,6 @@ Object* GetBaseObject(INode* Node, TimeValue Time)
 
 	return Obj;
 }
-#define VRAY_PROXY_DISPLAY_AS_MESH	4	// Value to set on a VRay Mesh Proxy to get the mesh
 
 int SetObjectParamValue(TimeValue CurrentTime, Object *Obj, const FString& ParamName, int DesiredValue)
 {
@@ -483,60 +557,6 @@ int SetObjectParamValue(TimeValue CurrentTime, Object *Obj, const FString& Param
 	return PrevDisplayValue;
 }
 
-
-Mesh* GetMeshFromRenderMesh(TimeValue CurrentTime, INode* Node, Interval& ValidityInterval, BOOL& bNeedsDelete)
-{
-	Object* Obj = GetBaseObject(Node, CurrentTime);
-
-	// Read validity interval before changing display settings
-	// Else VRay will somehow react to a call to ObjectValidity and change RenderMesh contents to what it was originally.
-	// E.g. display was set 'Box', we change display to 'Mesh' here before calling GetRenderMesh to retrieve actual mesh later(rather than a simple box)
-	//  but if ObjectValidity is called after GetRenderMesh that RenderMesh will become box again!
-	Interval ObjectValidity = Obj->ObjectValidity(CurrentTime);
-
-	const Class_ID& ObjectClassID = Obj->ClassID();
-	const FString VRayProxyParamName(TEXT("display"));
-	const FString BodyObjectViewportMeshParamName(TEXT("RenderViewportMeshRA"));
-	int PreviousMeshDisplayValue = 0;
-
-	if (ObjectClassID == VRAYPROXY_CLASS_ID)
-	{
-		// Need the high resolution render mesh associated with the VRay Mesh Proxy for the export
-		PreviousMeshDisplayValue = SetObjectParamValue(CurrentTime, Obj, VRayProxyParamName, VRAY_PROXY_DISPLAY_AS_MESH);
-	}
-	else if(ObjectClassID == BODYOBJECT_CLASS_ID)
-	{
-		// Need to make sure we are using the viewport mesh on BodyObject, otherwise the RenderMesh gives a tessellated low resolution mesh.
-		PreviousMeshDisplayValue = SetObjectParamValue(CurrentTime, Obj, BodyObjectViewportMeshParamName, 1);
-	}
-
-	GeomObject* GeomObj = static_cast<GeomObject*>(Obj);
-	if (GeomObj == nullptr)
-	{
-		return nullptr;
-	}
-
-	FNullView View;
-	Mesh* RenderMesh = GeomObj->GetRenderMesh(CurrentTime, Node, View, bNeedsDelete);
-
-	// Restore display state if different from mesh display
-	if (ObjectClassID == VRAYPROXY_CLASS_ID && PreviousMeshDisplayValue != VRAY_PROXY_DISPLAY_AS_MESH)
-	{
-		SetObjectParamValue(CurrentTime, Obj, VRayProxyParamName, PreviousMeshDisplayValue);
-	}
-	else if(ObjectClassID == BODYOBJECT_CLASS_ID && PreviousMeshDisplayValue != 1)
-	{
-		SetObjectParamValue(CurrentTime, Obj, BodyObjectViewportMeshParamName, PreviousMeshDisplayValue);
-	}
-
-	if (RenderMesh)
-	{
-		ValidityInterval &= ObjectValidity; // Update validity only when actual mesh is returned
-	}
-	return RenderMesh;
-}
-
-
 }
 }
 
diff --git a/Engine/Source/Programs/Enterprise/Datasmith/DatasmithMaxExporter/Public/DatasmithMaxGeomUtils.h b/Engine/Source/Programs/Enterprise/Datasmith/DatasmithMaxExporter/Public/DatasmithMaxGeomUtils.h
--- a/Engine/Source/Programs/Enterprise/Datasmith/DatasmithMaxExporter/Public/DatasmithMaxGeomUtils.h
+++ b/Engine/Source/Programs/Enterprise/Datasmith/DatasmithMaxExporter/Public/DatasmithMaxGeomUtils.h
@@ -30,15 +30,17 @@ public:
 		: Node(nullptr)
 		, MaxMesh(nullptr)
 		, bNeedsDelete(false)
+		, ObjectToDelete(nullptr)
 		, Pivot(FTransform::Identity)
 	{
 		ValidityInterval.SetInfinite();
 	}
 
-	explicit FRenderMeshForConversion(INode* InNode, Mesh* InMaxMesh, bool bInNeedsDelete)
+	explicit FRenderMeshForConversion(INode* InNode, Mesh* InMaxMesh, bool bInNeedsDelete, TriObject* InObjectToDelete=nullptr)
 		: Node(InNode)
 		, MaxMesh(InMaxMesh)
 		, bNeedsDelete(bInNeedsDelete)
+		, ObjectToDelete(InObjectToDelete)
 		, Pivot(FTransform::Identity)
 	{
 		ValidityInterval.SetInfinite();
@@ -61,6 +63,9 @@ public:
 		//Transfer ownership
 		bNeedsDelete = Other.bNeedsDelete;
 		Other.bNeedsDelete = false;
+
+		ObjectToDelete = Other.ObjectToDelete;
+		Other.ObjectToDelete = nullptr;
 	}
 
 	~FRenderMeshForConversion()
@@ -69,6 +74,10 @@ public:
 		{
 			MaxMesh->DeleteThis();
 		}
+		if (ObjectToDelete)
+		{
+			ObjectToDelete->DeleteMe();
+		}
 	}
 
 	bool IsValid() const
@@ -106,17 +115,18 @@ public:
 		return ValidityInterval;
 	}
 private:
-	INode* Node;
-	Mesh* MaxMesh;
-	bool bNeedsDelete;
-	FTransform Pivot;
+	INode* Node; // Scene node mesh is taken from 
+	Mesh* MaxMesh; // the 3ds Max Mesh which holds the geometry
+	bool bNeedsDelete; // Whether Mesh lifetime is in our hands
+	TriObject* ObjectToDelete; // When the Mesh is retrieved from some temporary TriObject this TriObject needs to be deleted(after Mesh is used)
+
+	FTransform Pivot; 
 	Interval ValidityInterval;
 };
 
 bool ConvertRailClone(DatasmithMaxDirectLink::ISceneTracker& SceneTracker, DatasmithMaxDirectLink::FNodeTracker& NodeTracker);
 bool ConvertForest(DatasmithMaxDirectLink::ISceneTracker& Scene, DatasmithMaxDirectLink::FNodeTracker& NodeTracker);
 
-Mesh* GetMeshFromRenderMesh(TimeValue CurrentTime, INode* Node, Interval& ValidityInterval, BOOL& bNeedsDelete);
 FRenderMeshForConversion GetMeshForGeomObject(TimeValue CurrentTime, INode* Node, const FTransform& Pivot = FTransform::Identity); // Extract mesh using already evaluated object
 
 FRenderMeshForConversion GetMeshForCollision(TimeValue CurrentTime, DatasmithMaxDirectLink::ISceneTracker& SceneTracker, INode* Node, bool bBakePivot);
