Revit Datasmith Direct Link sending metadata udpates when autosync is disabled fix.

diff --git a/Engine/Source/Programs/Enterprise/Datasmith/DatasmithRevitExporter/Private/DatasmithRevitCommand.cs b/Engine/Source/Programs/Enterprise/Datasmith/DatasmithRevitExporter/Private/DatasmithRevitCommand.cs
--- a/Engine/Source/Programs/Enterprise/Datasmith/DatasmithRevitExporter/Private/DatasmithRevitCommand.cs
+++ b/Engine/Source/Programs/Enterprise/Datasmith/DatasmithRevitExporter/Private/DatasmithRevitCommand.cs
@@ -123,6 +123,8 @@ namespace DatasmithRevitExporter
 				}
 			}
 
+			FDocument.ActiveDocument?.ActiveDirectLinkInstance?.ExportMetadataBatch();
+
 			return Result.Succeeded;
 		}
 	}
@@ -342,6 +344,8 @@ namespace DatasmithRevitExporter
 				}
 			}
 
+			FDocument.ActiveDocument?.ActiveDirectLinkInstance?.ExportMetadataBatch();
+
 			return Result.Succeeded;
 		}
 
diff --git a/Engine/Source/Programs/Enterprise/Datasmith/DatasmithRevitExporter/Private/DatasmithRevitDirectLink.cs b/Engine/Source/Programs/Enterprise/Datasmith/DatasmithRevitExporter/Private/DatasmithRevitDirectLink.cs
--- a/Engine/Source/Programs/Enterprise/Datasmith/DatasmithRevitExporter/Private/DatasmithRevitDirectLink.cs
+++ b/Engine/Source/Programs/Enterprise/Datasmith/DatasmithRevitExporter/Private/DatasmithRevitDirectLink.cs
@@ -210,25 +210,6 @@ namespace DatasmithRevitExporter
 			{
 				ActiveInstance.RunAutoSync();
 			}
-
-			if (ActiveInstance.SyncCount == 0)
-			{
-				return;
-			}
-
-			// MetadataCountPerIdleEvent value will be revisited in a future release since it also requires changes in the import of metadata.
-			int MetadataCountPerIdleEvent = int.MaxValue;
-
-			string EnvBatchSize = Environment.GetEnvironmentVariable("REVIT_DIRECTLINK_METADATA_BATCH_SIZE");
-			if (!string.IsNullOrEmpty(EnvBatchSize))
-			{
-				if (int.TryParse(EnvBatchSize, out MetadataCountPerIdleEvent))
-				{
-					MetadataCountPerIdleEvent = Math.Max(1, MetadataCountPerIdleEvent);
-				}
-			}
-
-			ActiveInstance.ExportMetadataBatch(MetadataCountPerIdleEvent);
 		}
 
 		public static void OnDocumentChanged(
@@ -748,13 +729,16 @@ namespace DatasmithRevitExporter
 			bSyncInProgress = false;
 		}
 
-		void ExportMetadataBatch(int ExportBatchSize)
+		//The number of metadata transfers in one go will be revisited in a future release since it also requires changes in the import of metadata.
+		// For now however all metadata will be transfered in one go.
+		// ExportMetadataBatch is called from the DatasmithSyncRevitCommand.OnExecute/DatasmithExportRevitCommand.OnExecute functions (instead of getting triggered when Revit is idle).
+		public void ExportMetadataBatch()
 		{
 			int CurrentBatchSize = 0;
 
 			Action<FCachedDocumentData> AddElements = (FCachedDocumentData CacheData) => 
 			{
-				while (CacheData.ElementsWithoutMetadataQueue.Count > 0 && CurrentBatchSize < ExportBatchSize)
+				while (CacheData.ElementsWithoutMetadataQueue.Count > 0)
 				{
 					var Entry = CacheData.ElementsWithoutMetadataQueue.Dequeue();
 
@@ -824,11 +808,6 @@ namespace DatasmithRevitExporter
 			foreach (var Cache in CachesToExport)
 			{
 				AddElements(Cache);
-
-				if (CurrentBatchSize >= ExportBatchSize)
-				{
-					break;
-				}
 			}
 
 			if (CurrentBatchSize > 0)
