[3ds Max Exporter] Removed file extension from scene/DirectLink name

diff --git a/Engine/Source/Programs/Enterprise/Datasmith/DatasmithMaxExporter/Private/DatasmithMaxExporterWithDirectLink.cpp b/Engine/Source/Programs/Enterprise/Datasmith/DatasmithMaxExporter/Private/DatasmithMaxExporterWithDirectLink.cpp
--- a/Engine/Source/Programs/Enterprise/Datasmith/DatasmithMaxExporter/Private/DatasmithMaxExporterWithDirectLink.cpp
+++ b/Engine/Source/Programs/Enterprise/Datasmith/DatasmithMaxExporter/Private/DatasmithMaxExporterWithDirectLink.cpp
@@ -3482,7 +3482,7 @@ public:
 		ExportedScene.PreExport();
 
 		SetOutputPath(GetDirectlinkCacheDirectory());
-		FString SceneName = FPaths::GetCleanFilename(GetCOREInterface()->GetCurFileName().data());
+		FString SceneName = FPaths::GetBaseFilename(GetCOREInterface()->GetCurFileName().data());
 		SetName(*SceneName);
 
 		DirectLinkImpl.Reset(new FDatasmithDirectLink);
