[SketchUp]  Fixed texture not being exported occasionally

- assets folder created ahead of using SU api to write image file

diff --git a/Engine/Source/Programs/Enterprise/Datasmith/DatasmithSketchUpRubyExporter/Private/DatasmithSketchUpTexture.cpp b/Engine/Source/Programs/Enterprise/Datasmith/DatasmithSketchUpRubyExporter/Private/DatasmithSketchUpTexture.cpp
--- a/Engine/Source/Programs/Enterprise/Datasmith/DatasmithSketchUpRubyExporter/Private/DatasmithSketchUpTexture.cpp
+++ b/Engine/Source/Programs/Enterprise/Datasmith/DatasmithSketchUpRubyExporter/Private/DatasmithSketchUpTexture.cpp
@@ -15,6 +15,7 @@
 #include "IDatasmithSceneElements.h"
 
 #include "Misc/Paths.h"
+#include "HAL/FileManager.h"
 
 // SketchUp SDK.
 #include "DatasmithSketchUpSDKBegins.h"
@@ -24,6 +25,7 @@
 
 #include "DatasmithSketchUpSDKCeases.h"
 
+
 using namespace DatasmithSketchUp;
 
 TSharedPtr<FTexture> FTextureCollection::FindOrAdd(SUTextureRef TextureRef)
@@ -55,6 +57,17 @@ FTexture* FTextureCollection::AddColorizedTexture(SUTextureRef TextureRef, FStri
 
 void FTextureCollection::Update()
 {
+	if (!TexturesMap.IsEmpty())
+	{
+		// Making sure _Assets folder is present
+		// SU API need to have folder created before writing image files
+		// It's not neccessarily present at this point - it's only created when meshes are exported and they are being exported in a separate thread
+		if (!FPaths::FileExists(Context.GetAssetsOutputPath()))
+		{
+			IFileManager::Get().MakeDirectory(Context.GetAssetsOutputPath());
+		}
+	}
+
 	for (TPair<FTextureIDType, TSharedPtr<FTexture>>& TextureNameAndTextureImageFile : TexturesMap)
 	{
 		TSharedPtr<FTexture> Texture = TextureNameAndTextureImageFile.Value;
