Revit - moving decals result in missing decal actor after synching via direct link fix,
(also fixes addition/creation of decal placements via directlink)

diff --git a/Engine/Source/Programs/Enterprise/Datasmith/DatasmithRevitExporter/Private/DatasmithRevitDirectLink.cs b/Engine/Source/Programs/Enterprise/Datasmith/DatasmithRevitExporter/Private/DatasmithRevitDirectLink.cs
--- a/Engine/Source/Programs/Enterprise/Datasmith/DatasmithRevitExporter/Private/DatasmithRevitDirectLink.cs
+++ b/Engine/Source/Programs/Enterprise/Datasmith/DatasmithRevitExporter/Private/DatasmithRevitDirectLink.cs
@@ -194,6 +194,8 @@ namespace DatasmithRevitExporter
 			}
 		}
 
+		public Dictionary<ElementId, ElementId> DecalIdToOwnerObjectIdMap = new Dictionary<ElementId, ElementId>();
+
 		public static void OnApplicationIdle()
 		{
 			FDirectLink ActiveInstance = FDocument.ActiveDocument?.ActiveDirectLinkInstance ?? null;
@@ -250,6 +252,16 @@ namespace DatasmithRevitExporter
 					}
 					else
 					{
+						if (FUtils.IsElementDecal(ModifiedElement))
+						{
+							//decal:
+							//modifying decal does not modify the owner object
+							if (ActiveInstance.DecalIdToOwnerObjectIdMap.TryGetValue(ElemId, out var OwnerObjectElementId))
+							{
+								ActiveInstance.RootCache.SetElementModified(true, OwnerObjectElementId);
+							}
+						}
+
 						// Handles a case where Revit won't notify us about modified mullions and their transform remains obsolte, thus wrong.
 						ElementCategoryFilter Filter = new ElementCategoryFilter(BuiltInCategory.OST_CurtainWallMullions);
 						IList<ElementId> DependentElements = ModifiedElement.GetDependentElements(Filter);
@@ -266,6 +278,64 @@ namespace DatasmithRevitExporter
 				}
 			}
 
+			foreach (ElementId ElemId in InArgs.GetAddedElementIds())
+			{
+				Element AddedElement = ActiveInstance.RootCache.SourceDocument.GetElement(ElemId);
+
+				if (FUtils.IsElementDecal(AddedElement))
+				{
+					FilteredElementCollector Collector = new FilteredElementCollector(ActiveInstance.RootCache.SourceDocument, ActiveInstance.RootCache.SourceDocument.ActiveView.Id);
+					IList<Element> AllElementsInView = Collector.ToElements();
+
+					bool ElementFound = false;
+
+					//find the Element that is dependent on the newly added decal:
+					foreach (Element ElementInView in AllElementsInView)
+					{
+#if REVIT_API_2023
+						if (ElementInView.Category != null && ElementInView.Category.BuiltInCategory != BuiltInCategory.OST_Levels)
+#else
+						if (ElementInView.Category != null && (BuiltInCategory)ElementInView.Category.Id.IntegerValue != BuiltInCategory.OST_Levels)
+#endif
+						{
+							IList<ElementId> DependentElements = ElementInView.GetDependentElements(null);
+							foreach (ElementId DependentElementId in DependentElements)
+							{
+								if (DependentElementId == ElemId)
+								{
+									ActiveInstance.RootCache.SetElementModified(true, ElementInView.Id);
+									ElementFound = true;
+									break;
+								}
+							}
+
+							if (ElementFound)
+							{
+								break;
+							}
+						}
+					}
+				}
+			}
+
+			foreach (ElementId ElemId in InArgs.GetDeletedElementIds())
+			{
+				//checking if decal was removed:
+				if (ActiveInstance.DecalIdToOwnerObjectIdMap.ContainsKey(ElemId))
+				{
+					ActiveInstance.DecalIdToOwnerObjectIdMap.Remove(ElemId);
+				}
+				//checking if owner object was removed:
+				if (ActiveInstance.DecalIdToOwnerObjectIdMap.ContainsValue(ElemId))
+				{
+					KeyValuePair<ElementId, ElementId>[] PairsToRemove = ActiveInstance.DecalIdToOwnerObjectIdMap.Where(CurrentPair => CurrentPair.Value == ElemId).ToArray();
+					foreach (KeyValuePair<ElementId, ElementId> PairToRemove in PairsToRemove)
+					{
+						ActiveInstance.DecalIdToOwnerObjectIdMap.Remove(PairToRemove.Key);
+					}
+				}
+			}
+
 			ActiveInstance.bHasChanges = ActiveInstance.bHasChanges || InArgs.GetDeletedElementIds().Any();
 			ActiveInstance.bHasChanges = ActiveInstance.bHasChanges || InArgs.GetAddedElementIds().Any();
 		}
@@ -321,7 +391,7 @@ namespace DatasmithRevitExporter
 #elif REVIT_API_2023
 				CmdGUID = "CB3186CC-1714-497A-9A54-A5D4B726524A";
 #else
-				#error This version of Revit is not supported yet.
+#error This version of Revit is not supported yet.
 #endif
 
 				if (UIApp == null)
diff --git a/Engine/Source/Programs/Enterprise/Datasmith/DatasmithRevitExporter/Private/DatasmithRevitDocumentData.cs b/Engine/Source/Programs/Enterprise/Datasmith/DatasmithRevitExporter/Private/DatasmithRevitDocumentData.cs
--- a/Engine/Source/Programs/Enterprise/Datasmith/DatasmithRevitExporter/Private/DatasmithRevitDocumentData.cs
+++ b/Engine/Source/Programs/Enterprise/Datasmith/DatasmithRevitExporter/Private/DatasmithRevitDocumentData.cs
@@ -1359,6 +1359,27 @@ namespace DatasmithRevitExporter
 			Transform InWorldTransform
 		)
 		{
+#if REVIT_API_2023
+			if (DirectLink != null && InElement.Category != null && InElement.Category.BuiltInCategory != BuiltInCategory.OST_Levels)
+#else
+			if (DirectLink != null && InElement.Category != null && (BuiltInCategory)InElement.Category.Id.IntegerValue != BuiltInCategory.OST_Levels)
+#endif
+			{
+				//Check if any of its children is Decal:
+				//If so track them, so that we can use the owner object elementId to update the Decals (when modified, for ep: changing its location)
+				IList<ElementId> DependentElements = InElement.GetDependentElements(null);
+				foreach (ElementId DependentElementId in DependentElements)
+				{
+					if (FUtils.IsElementIdDecal(InElement.Document, DependentElementId))
+					{
+						if (DependentElementId != InElement.Id && !DirectLink.DecalIdToOwnerObjectIdMap.ContainsKey(DependentElementId))
+						{
+							DirectLink.DecalIdToOwnerObjectIdMap.Add(DependentElementId, InElement.Id);
+						}
+					}
+				}
+			}
+
 			DirectLink?.MarkForExport(InElement);
 
 			FElementData ElementData = null;
@@ -1423,6 +1444,14 @@ namespace DatasmithRevitExporter
 						ElementData.InitializePivotPlacement(ref InWorldTransform);
 						ElementData.InitializeElement(InWorldTransform, ElementData);
 						ElementData.MeshMaterialsMap.Clear();
+
+						if (ElementData.bIsDecalElement)
+						{
+							if (!DecalElementsMap.ContainsKey(InElement.Id))
+							{
+								DecalElementsMap.Add(InElement.Id, ElementData);
+							}
+						}
 					}
 					else
 					{
diff --git a/Engine/Source/Programs/Enterprise/Datasmith/DatasmithRevitExporter/Private/DatasmithRevitUtils.cs b/Engine/Source/Programs/Enterprise/Datasmith/DatasmithRevitExporter/Private/DatasmithRevitUtils.cs
--- a/Engine/Source/Programs/Enterprise/Datasmith/DatasmithRevitExporter/Private/DatasmithRevitUtils.cs
+++ b/Engine/Source/Programs/Enterprise/Datasmith/DatasmithRevitExporter/Private/DatasmithRevitUtils.cs
@@ -171,5 +171,23 @@ namespace DatasmithRevitExporter
 
 			OutDecalDimensions = new XYZ(DecalQuad[0].Length * CENTIMETERS_PER_FOOT * 0.5, DecalQuad[1].Length * CENTIMETERS_PER_FOOT * 0.5, DimensionZ);
 		}
+
+		public static bool IsElementIdDecal(Document InDocument, ElementId InElementId)
+		{
+			Element Element = InDocument.GetElement(InElementId);
+			return IsElementDecal(Element);
+		}
+
+		public static bool IsElementDecal(Element InElement)
+		{
+			ElementType ElementType = InElement.Document.GetElement(InElement.GetTypeId()) as ElementType;
+
+			if (ElementType != null && (ElementType.FamilyName == "Decal"))
+			{
+				return true;
+			}
+
+			return false;
+		}
 	}
 }
