Revit - Decals on curved entities will sometimes scale inside the exported mesh

diff --git a/Engine/Source/Programs/Enterprise/Datasmith/DatasmithRevitExporter/Private/DatasmithRevitUtils.cs b/Engine/Source/Programs/Enterprise/Datasmith/DatasmithRevitExporter/Private/DatasmithRevitUtils.cs
--- a/Engine/Source/Programs/Enterprise/Datasmith/DatasmithRevitExporter/Private/DatasmithRevitUtils.cs
+++ b/Engine/Source/Programs/Enterprise/Datasmith/DatasmithRevitExporter/Private/DatasmithRevitUtils.cs
@@ -113,6 +113,8 @@ namespace DatasmithRevitExporter
 
 			GeometryElement GeomElement = InDecalElement.get_Geometry(new Options());
 
+			List<XYZ> MidPointsToTest = new List<XYZ>();
+
 			foreach (GeometryObject GeomObj in GeomElement)
 			{
 				if (GeomObj is Line QuadLine)
@@ -127,6 +129,8 @@ namespace DatasmithRevitExporter
 						XYZ EndPoint = QuadCurve.GetEndPoint(1);
 						Line BoundLine = Line.CreateBound(StartPoint, EndPoint);
 						DecalQuad.Add(BoundLine);
+
+						MidPointsToTest.AddRange(QuadCurve.Tessellate());
 					}
 					catch {}
 				}
@@ -158,15 +162,29 @@ namespace DatasmithRevitExporter
 
 			Plane TestPlane = Plane.CreateByThreePoints(TopLeft, TopRight, BottomRight);
 			UV TestUv;
-			double Distance;
-			TestPlane.Project(BottomLeft, out TestUv, out Distance);
 
-			Distance *= CENTIMETERS_PER_FOOT * 1.05; // leniency multiplier of 1.05, for imprecisions
+			double FinalDistance = 0;
+			TestPlane.Project(BottomLeft, out TestUv, out FinalDistance);
+
+			foreach (XYZ MidPoint in MidPointsToTest)
+			{
+				double Distance = 0;
+				TestPlane.Project(MidPoint, out TestUv, out Distance);
+
+				if (FinalDistance < Distance)
+				{
+					FinalDistance = Distance;
+				}
+			}
+
+			FinalDistance *= CENTIMETERS_PER_FOOT * 1.3;	// leniency multiplier of 15% in each direction
+															// seems to be needed in order to circumvent UE's decal projection imprecisions
+															// (issue would be visible via blending on big curvatured surfaces without the lenciency multiplier)
 
 			double DimensionZ = 2.0;
-			if (DimensionZ < Distance)
+			if (DimensionZ < FinalDistance)
 			{
-				DimensionZ = Distance;
+				DimensionZ = FinalDistance;
 			}
 
 			OutDecalDimensions = new XYZ(DecalQuad[0].Length * CENTIMETERS_PER_FOOT * 0.5, DecalQuad[1].Length * CENTIMETERS_PER_FOOT * 0.5, DimensionZ);
@@ -180,6 +198,11 @@ namespace DatasmithRevitExporter
 
 		public static bool IsElementDecal(Element InElement)
 		{
+			if (InElement == null)
+			{
+				return false;
+			}
+
 			ElementType ElementType = InElement.Document.GetElement(InElement.GetTypeId()) as ElementType;
 
 			if (ElementType != null && (ElementType.FamilyName == "Decal"))
