[Solidworks] Fixed animated rotatation export

diff --git a/Engine/Source/Programs/Enterprise/Datasmith/DatasmithSolidworksExporter/Private/AnimationExtractor.cs b/Engine/Source/Programs/Enterprise/Datasmith/DatasmithSolidworksExporter/Private/AnimationExtractor.cs
--- a/Engine/Source/Programs/Enterprise/Datasmith/DatasmithSolidworksExporter/Private/AnimationExtractor.cs
+++ b/Engine/Source/Programs/Enterprise/Datasmith/DatasmithSolidworksExporter/Private/AnimationExtractor.cs
@@ -87,7 +87,7 @@ namespace DatasmithSolidworks
 			int StepIndex = 0;
 			uint AllSteps = (uint)(Duration / Step);
 
-			MathTransform RootTransform = Root.GetTotalTransform(true);
+			MathTransform RootTransform = Root.GetTotalTransform(true) ?? Root.Transform2;
 
 			for (double CurTime = 0.0; CurTime < Duration; CurTime += Step)
 			{
@@ -122,20 +122,28 @@ namespace DatasmithSolidworks
 
                 if (ComponentTransform != null)
                 {
-					MathTransform ComponentWorldTransform = RootTransform.IMultiply(ComponentTransform);
+					MathTransform ComponentWorldTransform = RootTransform != null ? RootTransform.IMultiply(ComponentTransform) : ComponentTransform;
 					MathTransform ParentWorldTransform = RootTransform;
 
 					Component2 Parent = Component.GetParent();
 
 					if (Parent != null)
 					{
-						ParentWorldTransform = RootTransform.IMultiply(Parent.GetTotalTransform(true));
+						ParentWorldTransform = RootTransform != null ? RootTransform.IMultiply(Parent.GetTotalTransform(true)) : Parent.GetTotalTransform(true);
 					}
 
-					MathTransform ParentWorldTransformInverse = ParentWorldTransform.Inverse();
-					MathTransform ComponentLocalTransform = ComponentWorldTransform.IMultiply(ParentWorldTransformInverse);
+					MathTransform ComponentLocalTransform;
+					if (ParentWorldTransform != null)
+					{
+						MathTransform ParentInverse = ParentWorldTransform.Inverse();
+						ComponentLocalTransform = ComponentWorldTransform.IMultiply(ParentInverse);
+					}					
+					else
+					{
+						ComponentLocalTransform = ComponentWorldTransform;
+					}
 
-                    FAnimation.FChannel Channel = Anim.GetChannel(Component);
+					FAnimation.FChannel Channel = Anim.GetChannel(Component);
                     if (Channel == null)
                     {
                         Channel = Anim.NewChannel(Component);
diff --git a/Engine/Source/Programs/Enterprise/Datasmith/DatasmithSolidworksExporter/Private/DatasmithExporter.cs b/Engine/Source/Programs/Enterprise/Datasmith/DatasmithSolidworksExporter/Private/DatasmithExporter.cs
--- a/Engine/Source/Programs/Enterprise/Datasmith/DatasmithSolidworksExporter/Private/DatasmithExporter.cs
+++ b/Engine/Source/Programs/Enterprise/Datasmith/DatasmithSolidworksExporter/Private/DatasmithExporter.cs
@@ -878,16 +878,13 @@ namespace DatasmithSolidworks
 				{
 					FMatrix4 LocalMatrix = Keyframe.LocalMatrix;
 
-					// Get euler angles in degrees
-					float X = MathUtils.Rad2Deg * (float)Math.Atan2(LocalMatrix[6], LocalMatrix[10]);
-					float Y = MathUtils.Rad2Deg * (float)Math.Atan2(-LocalMatrix[2], Math.Sqrt(LocalMatrix[6] * LocalMatrix[6] + LocalMatrix[10] * LocalMatrix[10]));
-					float Z = MathUtils.Rad2Deg * (float)Math.Atan2(LocalMatrix[1], LocalMatrix[0]);
+					FVec3 Euler = MathUtils.ToEuler(LocalMatrix);
 
 					float Scale = LocalMatrix[15];
 
 					FVec3 Translation = new FVec3(LocalMatrix[12], LocalMatrix[13], LocalMatrix[14]);
 
-					Anim.AddFrame(EDatasmithFacadeAnimationTransformType.Rotation, Keyframe.Step, X, -Y, -Z);
+					Anim.AddFrame(EDatasmithFacadeAnimationTransformType.Rotation, Keyframe.Step, Euler.X, - Euler.Y, - Euler.Z);
 					Anim.AddFrame(EDatasmithFacadeAnimationTransformType.Scale, Keyframe.Step, Scale, Scale, Scale);
 					Anim.AddFrame(EDatasmithFacadeAnimationTransformType.Translation, Keyframe.Step, Translation.X, -Translation.Y, Translation.Z);
 				}
diff --git a/Engine/Source/Programs/Enterprise/Datasmith/DatasmithSolidworksExporter/Private/MathUtils.cs b/Engine/Source/Programs/Enterprise/Datasmith/DatasmithSolidworksExporter/Private/MathUtils.cs
--- a/Engine/Source/Programs/Enterprise/Datasmith/DatasmithSolidworksExporter/Private/MathUtils.cs
+++ b/Engine/Source/Programs/Enterprise/Datasmith/DatasmithSolidworksExporter/Private/MathUtils.cs
@@ -182,5 +182,24 @@ namespace DatasmithSolidworks
 			MathUtility MUtil = Addin.Instance.SolidworksApp.IGetMathUtility();
 			return (MathPoint)MUtil.CreatePoint(new[] { InX, InY, InZ }); ;
 		}
-	}
+
+		public static FVec3 ToEuler(FMatrix4 Mat)
+		{
+			// See UE::Math::TMatrix<T>::Rotator()
+
+			FVec3 XAxis = Mat.XBasis;
+			FVec3 YAxis = Mat.YBasis;
+			FVec3 ZAxis = Mat.ZBasis;
+
+			double Pitch = Math.Atan2(XAxis.Z, Math.Sqrt(XAxis.X*XAxis.X + XAxis.Y*XAxis.Y));
+			double Yaw = Math.Atan2(XAxis.Y, XAxis.X);
+
+			FVec3 SYAxis = new FVec3(-Math.Sin(Yaw), Math.Cos(Yaw), 0);
+
+			double Roll = Math.Atan2(FVec3.Dot(ZAxis, SYAxis), FVec3.Dot(YAxis, SYAxis));
+
+			// Y <-> Pitch,  Z <-> Yaw, X <-> Roll
+			return new FVec3(MathUtils.Rad2Deg * Roll, MathUtils.Rad2Deg * Pitch, MathUtils.Rad2Deg * Yaw);
+		}
+    }
 }
