Imported SolidWorks assembly animation very jittery

diff --git a/Engine/Source/Programs/Enterprise/Datasmith/DatasmithSolidworksExporter/Private/AnimationExtractor.cs b/Engine/Source/Programs/Enterprise/Datasmith/DatasmithSolidworksExporter/Private/AnimationExtractor.cs
--- a/Engine/Source/Programs/Enterprise/Datasmith/DatasmithSolidworksExporter/Private/AnimationExtractor.cs
+++ b/Engine/Source/Programs/Enterprise/Datasmith/DatasmithSolidworksExporter/Private/AnimationExtractor.cs
@@ -68,11 +68,16 @@ namespace DatasmithSolidworks
 				return null;
 			}
 
+			double InitialTime = Study.GetTime();
+
+			// Exporting Motion Study without first calling Stop on it sometimes(!) exports incorrect animation values. E.g. constant rotation.
+			// Even when no animation was playing(so Stop seemed redundant).
+			// Interestingly, even calling Study.IsPlaying helped with this. But 'Stop' looks nicer :)
+			Study.Stop();
+
 			FAnimation Anim = new FAnimation();
 			Anim.Name = Study.Name;
 
-			double InitialTime = Study.GetTime();
-
 			/*
 			 * swMotionStudyTypeAssembly			1 or 0x1 = Animation; D - cubed solver is used to do presentation animation only; no simulation is performed, so no results or plots are available; gravity, contact, springs, and forces cannot be used; mass and inertia values have no effect on the animation
 			 * swMotionStudyTypeCosmosMotion		4 or 0x4 = Motion Analysis; ADAMS(MSC.Software) solver is used to return accurate results; you must load the SOLIDWORKS Motion add -in with a SOLIDWORKS premium license to use this option
@@ -91,8 +96,8 @@ namespace DatasmithSolidworks
 
 			for (double CurTime = 0.0; CurTime < Duration; CurTime += Step)
 			{
-			    Study.SetTime(CurTime);
-			    ExtractAnimationData(StepIndex++, CurTime, Root, Anim, RootTransform);
+				Study.SetTime(CurTime);
+				ExtractAnimationData(StepIndex++, CurTime, Root, Anim, RootTransform);
 			}
 
 			Study.SetTime(Duration);
