PR #9679: REVIT-197037: Added Open in Twinmotion command. (Contributed by [REDACTED_PERSON])

diff --git a/Engine/Source/Programs/Enterprise/Datasmith/DatasmithRevitExporter/Private/DatasmithRevitApplication.cs b/Engine/Source/Programs/Enterprise/Datasmith/DatasmithRevitExporter/Private/DatasmithRevitApplication.cs
--- a/Engine/Source/Programs/Enterprise/Datasmith/DatasmithRevitExporter/Private/DatasmithRevitApplication.cs
+++ b/Engine/Source/Programs/Enterprise/Datasmith/DatasmithRevitExporter/Private/DatasmithRevitApplication.cs
@@ -203,6 +203,63 @@
 		{
 			Instance = this;
 
+			if (IsPreHandshakeRevitBuild(InApplication.ControlledApplication.VersionBuild))
+				InitializeRibbonUI(InApplication);
+
+			DocumentOpenedHandler = new EventHandler<DocumentOpenedEventArgs>(OnDocumentOpened);
+			InApplication.ControlledApplication.DocumentOpened += DocumentOpenedHandler;
+
+			DocumentCreatedHandler = new EventHandler<DocumentCreatedEventArgs>(OnDocumentCreated);
+			InApplication.ControlledApplication.DocumentCreated += DocumentCreatedHandler;
+
+			DocumentClosingHandler = new EventHandler<DocumentClosingEventArgs>(OnDocumentClosing);
+			InApplication.ControlledApplication.DocumentClosing += DocumentClosingHandler;
+
+			IdlingEventHandler = new EventHandler<IdlingEventArgs>(OnIdling);
+			InApplication.Idling += IdlingEventHandler;
+
+			ViewActivatedHandler = new EventHandler<ViewActivatedEventArgs>(OnViewActivated);
+			InApplication.ViewActivated += ViewActivatedHandler;
+
+			// Setup Direct Link
+
+			string RevitEngineDir = null;
+
+			try
+			{
+				using (RegistryKey Key = Registry.LocalMachine.OpenSubKey("Software\\Wow6432Node\\EpicGames\\Unreal Engine"))
+				{
+					RevitEngineDir = Key?.GetValue("RevitEngineDir") as string;
+				}
+			}
+			finally
+			{
+				if (RevitEngineDir == null)
+				{
+					// If we could not read the registry, fallback to hardcoded engine dir
+					RevitEngineDir = "C:\\ProgramData\\Epic\\Exporter\\RevitEngine\\";
+				}
+			}
+
+			bool bDirectLinkInitOk = FDatasmithFacadeDirectLink.Init(true, RevitEngineDir);
+
+			Debug.Assert(bDirectLinkInitOk);
+
+			// Register updater to react to view modification
+
+			ViewsUpdater = new View3DUpdater(InApplication.ControlledApplication.ActiveAddInId);
+			UpdaterRegistry.RegisterUpdater(ViewsUpdater);
+
+			ElementCategoryFilter Filter = new ElementCategoryFilter( BuiltInCategory.OST_Views);
+			UpdaterRegistry.AddTrigger(ViewsUpdater.GetUpdaterId(), Filter, Element.GetChangeTypeAny());
+			UpdaterRegistry.AddTrigger(ViewsUpdater.GetUpdaterId(), Filter, Element.GetChangeTypeElementAddition());
+			UpdaterRegistry.AddTrigger(ViewsUpdater.GetUpdaterId(), Filter, Element.GetChangeTypeElementDeletion());
+
+			return Result.Succeeded;
+		}
+
+		void InitializeRibbonUI(UIControlledApplication InApplication)
+		{
 			// Create a custom ribbon tab
 			string TabName = DatasmithRevitResources.Strings.DatasmithTabName;
 			InApplication.CreateRibbonTab(TabName);
@@ -288,57 +345,6 @@
 			ShowLogButton.Image = new BitmapImage(new Uri(DatasmithIconBase + "16.png"));
 			ShowLogButton.LargeImage = new BitmapImage(new Uri(DatasmithIconBase + "32.png"));
 			ShowLogButton.ToolTip = DatasmithRevitResources.Strings.ButtonMessagesHint;
-
-			DocumentOpenedHandler = new EventHandler<DocumentOpenedEventArgs>(OnDocumentOpened);
-			InApplication.ControlledApplication.DocumentOpened += DocumentOpenedHandler;
-
-			DocumentCreatedHandler = new EventHandler<DocumentCreatedEventArgs>(OnDocumentCreated);
-			InApplication.ControlledApplication.DocumentCreated += DocumentCreatedHandler;
-
-			DocumentClosingHandler = new EventHandler<DocumentClosingEventArgs>(OnDocumentClosing);
-			InApplication.ControlledApplication.DocumentClosing += DocumentClosingHandler;
-
-			IdlingEventHandler = new EventHandler<IdlingEventArgs>(OnIdling);
-			InApplication.Idling += IdlingEventHandler;
-
-			ViewActivatedHandler = new EventHandler<ViewActivatedEventArgs>(OnViewActivated);
-			InApplication.ViewActivated += ViewActivatedHandler;
-
-			// Setup Direct Link
-
-			string RevitEngineDir = null;
-
-			try
-			{
-				using (RegistryKey Key = Registry.LocalMachine.OpenSubKey("Software\\Wow6432Node\\EpicGames\\Unreal Engine"))
-				{
-					RevitEngineDir = Key?.GetValue("RevitEngineDir") as string;
-				}
-			}
-			finally
-			{
-				if (RevitEngineDir == null)
-				{
-					// If we could not read the registry, fallback to hardcoded engine dir
-					RevitEngineDir = "C:\\ProgramData\\Epic\\Exporter\\RevitEngine\\";
-				}
-			}
-
-			bool bDirectLinkInitOk = FDatasmithFacadeDirectLink.Init(true, RevitEngineDir);
-
-			Debug.Assert(bDirectLinkInitOk);
-
-			// Register updater to react to view modification
-
-			ViewsUpdater = new View3DUpdater(InApplication.ControlledApplication.ActiveAddInId);
-			UpdaterRegistry.RegisterUpdater(ViewsUpdater);
-
-			ElementCategoryFilter Filter = new ElementCategoryFilter( BuiltInCategory.OST_Views);
-			UpdaterRegistry.AddTrigger(ViewsUpdater.GetUpdaterId(), Filter, Element.GetChangeTypeAny());
-			UpdaterRegistry.AddTrigger(ViewsUpdater.GetUpdaterId(), Filter, Element.GetChangeTypeElementAddition());
-			UpdaterRegistry.AddTrigger(ViewsUpdater.GetUpdaterId(), Filter, Element.GetChangeTypeElementDeletion());
-
-			return Result.Succeeded;
 		}
 
 		void OnIdling(object Sender, IdlingEventArgs Args)
@@ -430,5 +436,10 @@
 				ExportMessagesDialog.Focus();
 			}
 		}
+
+		public static bool IsPreHandshakeRevitBuild(string VersionBuild)
+		{
+			return Version.Parse(VersionBuild) < Version.Parse("23.1");
+		} 
 	}
 }
diff --git a/Engine/Source/Programs/Enterprise/Datasmith/DatasmithRevitExporter/Private/DatasmithRevitCommand.cs b/Engine/Source/Programs/Enterprise/Datasmith/DatasmithRevitExporter/Private/DatasmithRevitCommand.cs
--- a/Engine/Source/Programs/Enterprise/Datasmith/DatasmithRevitExporter/Private/DatasmithRevitCommand.cs
+++ b/Engine/Source/Programs/Enterprise/Datasmith/DatasmithRevitExporter/Private/DatasmithRevitCommand.cs
@@ -11,14 +11,15 @@ using System.Windows.Forms;
 using Autodesk.Revit.Attributes;
 using Autodesk.Revit.DB;
 using Autodesk.Revit.UI;
+using Microsoft.Win32;
 
 namespace DatasmithRevitExporter
 {
-	public abstract class DatasmithRevitCommand : IExternalCommand
+	public abstract class DatasmithRevitCommandUtils
 	{
 		public const string DIALOG_CAPTION = "Export 3D View to Unreal Datasmith";
 
-		public Result Execute(
+		public static Result ExecuteFunc(
 			ExternalCommandData InCommandData,		// contains reference to Application and View
 			ref string			OutCommandMessage,  // error message to display in the failure dialog when the command returns "Failed"
 			ElementSet			OutElements         // set of problem elements to display in the failure dialog when the command returns "Failed"
@@ -49,17 +50,19 @@ namespace DatasmithRevitExporter
 				return Result.Cancelled;
 			}
 
-			return OnExecute(InCommandData, ref OutCommandMessage, OutElements);
+			return Result.Succeeded;
 		}
-
-		public abstract Result OnExecute(ExternalCommandData InCommandData, ref string OutCommandMessage, ElementSet OutElements);
 	}
 
 	[Transaction(TransactionMode.Manual)]
-	public class DatasmithSyncRevitCommand : DatasmithRevitCommand
+	public class DatasmithSyncRevitCommand : IExternalCommand
 	{
-		public override Result OnExecute(ExternalCommandData InCommandData, ref string OutCommandMessage, ElementSet OutElements)
+		public static Result ExecuteFunc(ExternalCommandData InCommandData, ref string OutCommandMessage, ElementSet OutElements)
 		{
+			Result Result = DatasmithRevitCommandUtils.ExecuteFunc(InCommandData, ref OutCommandMessage, OutElements);
+			if (Result != Result.Succeeded)
+				return Result;
+
 			UIDocument UIDoc = InCommandData.Application.ActiveUIDocument;
 			Document Doc = UIDoc.Document;
 			View3D ActiveView = FDocument.ActiveDocument?.ActiveDirectLinkInstance?.SyncView;
@@ -67,14 +70,14 @@ namespace DatasmithRevitExporter
 			if (ActiveView == null)
 			{
 				string Message = "You must select a 3D view to sync.";
-				MessageBox.Show(Message, DIALOG_CAPTION, MessageBoxButtons.OK, MessageBoxIcon.Warning);
+				MessageBox.Show(Message, DatasmithRevitCommandUtils.DIALOG_CAPTION, MessageBoxButtons.OK, MessageBoxIcon.Warning);
 				return Result.Cancelled;
 			}
 
 			if (ActiveView.IsTemplate || !ActiveView.CanBePrinted)
 			{
 				string Message = "The active 3D view cannot be exported.";
-				MessageBox.Show(Message, DIALOG_CAPTION, MessageBoxButtons.OK, MessageBoxIcon.Warning);
+				MessageBox.Show(Message, DatasmithRevitCommandUtils.DIALOG_CAPTION, MessageBoxButtons.OK, MessageBoxIcon.Warning);
 				return Result.Cancelled;
 			}
 
@@ -110,7 +113,7 @@ namespace DatasmithRevitExporter
 				catch (System.Exception exception)
 				{
 					OutCommandMessage = string.Format("Cannot export the 3D view:\n\n{0}\n\n{1}", exception.Message, exception.StackTrace);
-					MessageBox.Show(OutCommandMessage, DIALOG_CAPTION, MessageBoxButtons.OK, MessageBoxIcon.Error);
+					MessageBox.Show(OutCommandMessage, DatasmithRevitCommandUtils.DIALOG_CAPTION, MessageBoxButtons.OK, MessageBoxIcon.Error);
 					return Result.Failed;
 				}
 				finally
@@ -127,13 +130,22 @@ namespace DatasmithRevitExporter
 
 			return Result.Succeeded;
 		}
+
+		public Result Execute(ExternalCommandData InCommandData, ref string OutCommandMessage, ElementSet OutElements)
+		{
+			return ExecuteFunc(InCommandData, ref OutCommandMessage, OutElements);
+		}
 	}
 
 	[Transaction(TransactionMode.Manual)]
-	public class DatasmithAutoSyncRevitCommand : DatasmithRevitCommand
+	public class DatasmithAutoSyncRevitCommand : IExternalCommand
 	{
-		public override Result OnExecute(ExternalCommandData InCommandData, ref string OutCommandMessage, ElementSet OutElements)
+		public Result Execute(ExternalCommandData InCommandData, ref string OutCommandMessage, ElementSet OutElements)
 		{
+			Result Result = DatasmithRevitCommandUtils.ExecuteFunc(InCommandData, ref OutCommandMessage, OutElements);
+			if (Result != Result.Succeeded)
+				return Result;
+
 			FDirectLink.bAutoSync = !FDirectLink.bAutoSync;
 
 			DatasmithRevitApplication.Instance.SetAutoSyncButtonToggled(FDirectLink.bAutoSync);
@@ -144,7 +156,7 @@ namespace DatasmithRevitExporter
 
 	// Add-in external command Export to Unreal Datasmith. 
 	[Transaction(TransactionMode.Manual)]
-	public class DatasmithExportRevitCommand : DatasmithRevitCommand
+	public class DatasmithExportRevitCommand : IExternalCommand
 	{
 		class DocumentExportPathCache
 		{
@@ -156,8 +168,12 @@ namespace DatasmithRevitExporter
 		private static Dictionary<Document, DocumentExportPathCache> ExportPaths = new Dictionary<Document, DocumentExportPathCache>();
 	
 		// Implement the interface to execute the command.
-		public override Result OnExecute(ExternalCommandData InCommandData, ref string OutCommandMessage, ElementSet OutElements)
+		public Result Execute(ExternalCommandData InCommandData, ref string OutCommandMessage, ElementSet OutElements)
 		{
+			Result Result = DatasmithRevitCommandUtils.ExecuteFunc(InCommandData, ref OutCommandMessage, OutElements);
+			if (Result != Result.Succeeded)
+				return Result;
+
 			UIDocument UIDoc = InCommandData.Application.ActiveUIDocument;
 			Document Doc = UIDoc.Document;
 
@@ -166,7 +182,7 @@ namespace DatasmithRevitExporter
 			if (string.IsNullOrWhiteSpace(DocumentPath))
 			{
 				string message = "Your document must be saved on disk before exporting.";
-				MessageBox.Show(message, DIALOG_CAPTION, MessageBoxButtons.OK, MessageBoxIcon.Warning);
+				MessageBox.Show(message, DatasmithRevitCommandUtils.DIALOG_CAPTION, MessageBoxButtons.OK, MessageBoxIcon.Warning);
 				return Result.Cancelled;
 			}
 
@@ -186,14 +202,14 @@ namespace DatasmithRevitExporter
 				if (ActiveView == null)
 				{
 					string Message = "You must be in a 3D view to export.";
-					MessageBox.Show(Message, DIALOG_CAPTION, MessageBoxButtons.OK, MessageBoxIcon.Warning);
+					MessageBox.Show(Message, DatasmithRevitCommandUtils.DIALOG_CAPTION, MessageBoxButtons.OK, MessageBoxIcon.Warning);
 					return Result.Cancelled;
 				}
 
 				if (ActiveView.IsTemplate || !ActiveView.CanBePrinted)
 				{
 					string Message = "The active 3D view cannot be exported.";
-					MessageBox.Show(Message, DIALOG_CAPTION, MessageBoxButtons.OK, MessageBoxIcon.Warning);
+					MessageBox.Show(Message, DatasmithRevitCommandUtils.DIALOG_CAPTION, MessageBoxButtons.OK, MessageBoxIcon.Warning);
 					return Result.Cancelled;
 				}
 
@@ -230,7 +246,7 @@ namespace DatasmithRevitExporter
 
 				SaveFileDialog Dialog = new SaveFileDialog();
 
-				Dialog.Title            = DIALOG_CAPTION;
+				Dialog.Title            = DatasmithRevitCommandUtils.DIALOG_CAPTION;
 				Dialog.InitialDirectory = InitialDir;
 				Dialog.FileName         = FileName;
 				Dialog.DefaultExt       = "udatasmith";
@@ -251,7 +267,7 @@ namespace DatasmithRevitExporter
 				if (string.IsNullOrWhiteSpace(Dialog.FileName))
 				{
 					string message = "The given Unreal Datasmith file name is blank.";
-					MessageBox.Show(message, DIALOG_CAPTION, MessageBoxButtons.OK, MessageBoxIcon.Warning);
+					MessageBox.Show(message, DatasmithRevitCommandUtils.DIALOG_CAPTION, MessageBoxButtons.OK, MessageBoxIcon.Warning);
 					return Result.Cancelled;
 				}
 
@@ -264,9 +280,9 @@ namespace DatasmithRevitExporter
 				using (var FBD = new FolderBrowserDialog())
 				{
 					FBD.ShowNewFolderButton = true;
-					DialogResult result = FBD.ShowDialog();
+					DialogResult DlgResult = FBD.ShowDialog();
 
-					if (result != DialogResult.OK || string.IsNullOrWhiteSpace(FBD.SelectedPath))
+					if (DlgResult != DialogResult.OK || string.IsNullOrWhiteSpace(FBD.SelectedPath))
 					{
 						return Result.Cancelled;
 					}
@@ -325,7 +341,7 @@ namespace DatasmithRevitExporter
 				catch( System.Exception exception )
 				{
 					OutCommandMessage = string.Format("Cannot export the 3D view:\n\n{0}\n\n{1}", exception.Message, exception.StackTrace);
-					MessageBox.Show(OutCommandMessage, DIALOG_CAPTION, MessageBoxButtons.OK, MessageBoxIcon.Error);
+					MessageBox.Show(OutCommandMessage, DatasmithRevitCommandUtils.DIALOG_CAPTION, MessageBoxButtons.OK, MessageBoxIcon.Error);
 					return Result.Failed;
 				}
 				finally
@@ -441,4 +457,57 @@ namespace DatasmithRevitExporter
 			return Result.Succeeded;
 		}
 	}
+
+	[Transaction(TransactionMode.Manual)]
+	public class DatasmithOpenInTwinmotionCommand : IExternalCommand
+	{
+		private static string GetCompatibleVersionOfTwinmotionExecutablePath()
+		{
+			List<string> RegTwinmotion = new List<string> {
+				@"SOFTWARE\Microsoft\Windows\CurrentVersion\App Paths\Twinmotion2023.1.exe",
+				@"SOFTWARE\Microsoft\Windows\CurrentVersion\App Paths\Twinmotion2022.2.exe",
+				@"SOFTWARE\Microsoft\Windows\CurrentVersion\App Paths\Twinmotion2022.2-Revit.exe"
+			 };
+			
+			foreach (string Reg in RegTwinmotion)
+			{
+				try
+				{
+					var Key = Registry.CurrentUser.OpenSubKey(Reg);
+					var Exe = Key?.GetValue(null)?.ToString() ?? "";
+					if (!String.IsNullOrEmpty(Exe) && File.Exists(Exe))
+						return Exe;
+				}
+				catch (Exception Ex)
+				{
+					Trace.WriteLine(Ex.ToString()); // TODO_REVIEW (Dump journal comments.)
+				}
+			}
+			return null;
+		}
+		public Result Execute(ExternalCommandData InCommandData, ref string OutCommandMessage, ElementSet OutElements)
+		{
+			Result Result = DatasmithSyncRevitCommand.ExecuteFunc(InCommandData, ref OutCommandMessage, OutElements);
+			if (Result != Result.Succeeded)
+				return Result;
+			FDirectLink DirectLinkInstance = FDocument.ActiveDocument?.ActiveDirectLinkInstance;
+			if (DirectLinkInstance != null)
+			{
+				string SourceName = Regex.Replace($"{DirectLinkInstance.DatasmithScene.GetName()}", @"\s+", "_");
+				string ExecutablePath = GetCompatibleVersionOfTwinmotionExecutablePath();
+				if (ExecutablePath != null)
+				{
+					Process.Start(ExecutablePath, $@"-OpenProject=prompt -DirectLink.SourceName={SourceName}");
+				}
+			}
+			return Result.Succeeded;
+		}
+	}
+	public class DatasmithOpenInTwinmotionCommandAvailability : IExternalCommandAvailability
+	{
+		public bool IsCommandAvailable(UIApplication InUIApplication, CategorySet InCategorySet)
+		{
+			return !DatasmithRevitApplication.IsPreHandshakeRevitBuild(InUIApplication.Application.VersionBuild);
+		}
+	}
 }
