[Solidworks] Fixed MeshActor loosing its mesh on Sync when Part changes

- simpler fix for 5.1.1 HF, need to review this in 5.2

diff --git a/Engine/Source/Programs/Enterprise/Datasmith/DatasmithSolidworksExporter/Private/AssemblyDocument.cs b/Engine/Source/Programs/Enterprise/Datasmith/DatasmithSolidworksExporter/Private/AssemblyDocument.cs
--- a/Engine/Source/Programs/Enterprise/Datasmith/DatasmithSolidworksExporter/Private/AssemblyDocument.cs
+++ b/Engine/Source/Programs/Enterprise/Datasmith/DatasmithSolidworksExporter/Private/AssemblyDocument.cs
@@ -386,6 +386,12 @@ namespace DatasmithSolidworks
 					OutMeshesToExportMap.Add(InComponent, ActorExportInfo.Name);
 				}
 
+				if (ActorExportInfo.Type == EActorType.MeshActor )
+				{
+					// todo: deduplicate mesh name calculation
+					ActorExportInfo.MeshName = FDatasmithExporter.SanitizeName($"{ActorExportInfo.Name}_Mesh");
+				}
+
 				Exporter.ExportOrUpdateActor(ActorExportInfo);
 
 				SyncState.CleanComponents.Add(ComponentName);
