[3ds Max Exporter] CoronaLegacyMtl opacity converted to Datasmith

- texmap(on/off, amount support), scalar level or color(desaturated)

diff --git a/Engine/Source/Programs/Enterprise/Datasmith/DatasmithMaxExporter/Private/MaxMaterialsToUEPbr/DatasmithMaxCoronaMaterialsToUEPbr.cpp b/Engine/Source/Programs/Enterprise/Datasmith/DatasmithMaxExporter/Private/MaxMaterialsToUEPbr/DatasmithMaxCoronaMaterialsToUEPbr.cpp
--- a/Engine/Source/Programs/Enterprise/Datasmith/DatasmithMaxExporter/Private/MaxMaterialsToUEPbr/DatasmithMaxCoronaMaterialsToUEPbr.cpp
+++ b/Engine/Source/Programs/Enterprise/Datasmith/DatasmithMaxExporter/Private/MaxMaterialsToUEPbr/DatasmithMaxCoronaMaterialsToUEPbr.cpp
@@ -847,8 +847,22 @@ void FDatasmithMaxCoronaMaterialsToUEPbr::Convert( TSharedRef< IDatasmithScene >
 	// Opacity
 	IDatasmithMaterialExpression* OpacityExpression = nullptr;
 	{
-		TGuardValue< bool > SetIsMonoChannel( ConvertState.bIsMonoChannel, true );
-		OpacityExpression = ConvertTexmap( CoronaMaterialProperties.OpacityMap );
+		IDatasmithMaterialExpression* OpacityTexmapExpression = ConvertTexmap( CoronaMaterialProperties.OpacityMap );
+
+		if (OpacityTexmapExpression)
+		{
+			OpacityExpression = &Desaturate(Multiply(*OpacityTexmapExpression, Scalar(CoronaMaterialProperties.OpacityLevel)));
+		}
+		else
+		{
+			FLinearColor OpacityColor = CoronaMaterialProperties.Opacity.Value * CoronaMaterialProperties.OpacityLevel;
+
+			// Don't consider almost full opacity as transparency to prevent creating opacity expression(which would lead to making non-opaque shader)
+			if (!OpacityColor.Equals(FLinearColor::White)) 
+			{
+				OpacityExpression = &Desaturate(Color(OpacityColor));
+			}
+		}
 	}
 
 	// Refraction
@@ -1042,6 +1056,13 @@ void FDatasmithMaxCoronaMaterialsToUEPbr::Convert( TSharedRef< IDatasmithScene >
 			TransmittanceExpression->ConnectExpression( *ThinTranslucencyMaterialOutput->GetInput(0) );
 		}
 	}
+	else
+	{
+		if (OpacityExpression)
+		{
+			Connect(PbrMaterialElement->GetOpacity(), *OpacityExpression);
+		}
+	}
 
 	MaterialElement = PbrMaterialElement;
 }
