[3ds Max Exporter] Custom object properties changes synced

diff --git a/Engine/Source/Programs/Enterprise/Datasmith/DatasmithMaxExporter/Private/DatasmithMaxDirectLinkNotifications.cpp b/Engine/Source/Programs/Enterprise/Datasmith/DatasmithMaxExporter/Private/DatasmithMaxDirectLinkNotifications.cpp
--- a/Engine/Source/Programs/Enterprise/Datasmith/DatasmithMaxExporter/Private/DatasmithMaxDirectLinkNotifications.cpp
+++ b/Engine/Source/Programs/Enterprise/Datasmith/DatasmithMaxExporter/Private/DatasmithMaxDirectLinkNotifications.cpp
@@ -51,6 +51,17 @@ public:
 		}
 	}
 
+#if MAX_PRODUCT_YEAR_NUMBER < 2020  // Starting 2020 using ReferenceMaker to handle this notification
+	virtual void UserPropertiesChanged(NodeKeyTab& nodes) override
+	{
+		LogNodeEvent(L"UserPropertiesChanged", nodes);
+		for (int NodeIndex = 0; NodeIndex < nodes.Count(); ++NodeIndex)
+		{
+			SceneTracker.NodePropertiesChanged(NodeEventNamespace::GetNodeByKey(nodes[NodeIndex]));
+		}
+	}
+#endif
+
 	// Not used:
 
 	virtual void LayerChanged(NodeKeyTab& nodes) override
@@ -58,11 +69,6 @@ public:
 		LogNodeEvent(L"LayerChanged", nodes);
 	}
 
-	virtual void UserPropertiesChanged(NodeKeyTab& nodes) override
-	{
-		LogNodeEvent(L"UserPropertiesChanged", nodes);
-	}
-
 	virtual void HideChanged(NodeKeyTab& nodes) override
 	{
 		LogNodeEvent(L"HideChanged", nodes);
@@ -480,9 +486,11 @@ public:
 			SceneTracker.NodePropertiesChanged(&Node);
 		break;
 
-		// todo: appears in Max SDK since 2020 
-		// case REFMSG_NODE_USER_PROPERTY_CHANGED:
-		// break;
+#if MAX_PRODUCT_YEAR_NUMBER >= 2020
+		case REFMSG_NODE_USER_PROPERTY_CHANGED:
+			SceneTracker.NodePropertiesChanged(&Node);
+		break;
+#endif
 
 		// SceneTracker.NodeGeometryChanged(nodes[NodeIndex]);
 		default: 
