Fixed missing calls to FEngineLoop::AppExit

Skipping AppExit can cause crashes on exit because systems are destroyed in the wrong order.

#rnx

diff --git a/Engine/Source/Programs/ChaosVisualDebugger/Private/ChaosVisualDebuggerMain.cpp b/Engine/Source/Programs/ChaosVisualDebugger/Private/ChaosVisualDebuggerMain.cpp
--- a/Engine/Source/Programs/ChaosVisualDebugger/Private/ChaosVisualDebuggerMain.cpp
+++ b/Engine/Source/Programs/ChaosVisualDebugger/Private/ChaosVisualDebuggerMain.cpp
@@ -133,7 +133,9 @@ int32 ChaosVisualDebuggerMain(const TCHAR* CommandLine)
 
 	FCoreDelegates::OnExit.Broadcast();
 	FSlateApplication::Shutdown();
+	FEngineLoop::AppPreExit();
 	FModuleManager::Get().UnloadModulesAtShutdown();
+	FEngineLoop::AppExit();
 
 	return 0;
 } 
\ No newline at end of file
diff --git a/Engine/Source/Programs/Enterprise/Datasmith/DatasmithCADWorker/Private/DatasmithCADWorker.cpp b/Engine/Source/Programs/Enterprise/Datasmith/DatasmithCADWorker/Private/DatasmithCADWorker.cpp
--- a/Engine/Source/Programs/Enterprise/Datasmith/DatasmithCADWorker/Private/DatasmithCADWorker.cpp
+++ b/Engine/Source/Programs/Enterprise/Datasmith/DatasmithCADWorker/Private/DatasmithCADWorker.cpp
@@ -114,7 +114,9 @@ INT32_MAIN_INT32_ARGC_TCHAR_ARGV()
 		Result = EXIT_FAILURE;
 	}
 
-	FCoreDelegates::OnExit.Broadcast();
+	FEngineLoop::AppPreExit();
+	FModuleManager::Get().UnloadModulesAtShutdown();
+	FEngineLoop::AppExit();
 
 	FPlatformMisc::RequestExit(true);
 
diff --git a/Engine/Source/Programs/HeadlessChaosPerf/Private/HeadlessChaosPerfApp.cpp b/Engine/Source/Programs/HeadlessChaosPerf/Private/HeadlessChaosPerfApp.cpp
--- a/Engine/Source/Programs/HeadlessChaosPerf/Private/HeadlessChaosPerfApp.cpp
+++ b/Engine/Source/Programs/HeadlessChaosPerf/Private/HeadlessChaosPerfApp.cpp
@@ -19,8 +19,9 @@ INT32_MAIN_INT32_ARGC_TCHAR_ARGV()
 	ChaosPerf::FPerfTestRegistry::Get().RunAll();
 
 
-	FCoreDelegates::OnExit.Broadcast();
+	FEngineLoop::AppPreExit();
 	FModuleManager::Get().UnloadModulesAtShutdown();
+	FEngineLoop::AppExit();
 
 	FPlatformMisc::RequestExit(false);
 
diff --git a/Engine/Source/Programs/UnrealVersionSelector/Private/UnrealVersionSelector.cpp b/Engine/Source/Programs/UnrealVersionSelector/Private/UnrealVersionSelector.cpp
--- a/Engine/Source/Programs/UnrealVersionSelector/Private/UnrealVersionSelector.cpp
+++ b/Engine/Source/Programs/UnrealVersionSelector/Private/UnrealVersionSelector.cpp
@@ -447,6 +447,7 @@ int Main(const TArray<FString>& Arguments)
 
 		FEngineLoop::AppPreExit();
 		FModuleManager::Get().UnloadModulesAtShutdown();
+		FEngineLoop::AppExit();
 
 		return Result;
 	}
