[DatasmithSDKSample]
Fix winding issue in the cloth scene

diff --git a/Engine/Source/Programs/Enterprise/Datasmith/DatasmithSDKSample/Scenes/Cloth.cpp b/Engine/Source/Programs/Enterprise/Datasmith/DatasmithSDKSample/Scenes/Cloth.cpp
--- a/Engine/Source/Programs/Enterprise/Datasmith/DatasmithSDKSample/Scenes/Cloth.cpp
+++ b/Engine/Source/Programs/Enterprise/Datasmith/DatasmithSDKSample/Scenes/Cloth.cpp
@@ -68,18 +68,18 @@ TSharedPtr<IDatasmithScene> FClothScene::Export()
 
 		FDatasmithCloth Cloth;
 
-		FDatasmithClothPattern& Pattern = Cloth.Patterns.AddDefaulted_GetRef();
-
 		// Make Rectangle
+		auto MakeRectangle = [](FDatasmithCloth& Cloth, FVector3f Offset)
 		{
-			int32 h = 2;
-			int32 w = 3;
+			FDatasmithClothPattern& Pattern = Cloth.Patterns.AddDefaulted_GetRef();
+			int32 h = 7;
+			int32 w = 7;
 			for (int32 a = 0; a < h; ++a)
 			{
 				for (int32 b = 0; b < w; ++b)
 				{
 					Pattern.SimPosition.Add(FVector2f(100.f*a, 100.f*b));
-					Pattern.SimRestPosition.Add(FVector3f(100.f*a, 100.f*b, 0.f));
+					Pattern.SimRestPosition.Add(Offset + FVector3f(100.f*a, 100.f*b, 0.f));
 
 					if (a < h-1 && b < w-1)
 					{
@@ -91,18 +91,23 @@ TSharedPtr<IDatasmithScene> FClothScene::Export()
 						Pattern.SimTriangleIndices.Add((a+1)*w+b); // C
 
 						Pattern.SimTriangleIndices.Add(a*w+b+1); // B
-						Pattern.SimTriangleIndices.Add((a+1)*w+b); // C
 						Pattern.SimTriangleIndices.Add((a+1)*w+b+1); // D
+						Pattern.SimTriangleIndices.Add((a+1)*w+b); // C
 					}
 				}
 			}
-		}
 
-		// #ue_ds_sdk_todo Proper data API
-		// eg. SetVertices, SetVertexParameter...
-		FParameterData& Param = Pattern.Parameters.AddDefaulted_GetRef();
-		Param.Name = TEXT("Test");
-		Param.Data.Set<TArray<double>>({1,2,3});
+// 			// #ue_ds_sdk_todo Proper data API
+// 			// eg. SetVertices, SetVertexParameter...
+// 			FParameterData& Param = Pattern.Parameters.AddDefaulted_GetRef();
+// // 			Param.Name = TEXT("Test");
+// 			Param.Data.Set<TArray<double>>({1,2,3});
+		};
+
+// 		MakeRectangle(Cloth, FVector3f{});
+		MakeRectangle(Cloth, FVector3f{-30.f, 0.f, 90.f});
+		MakeRectangle(Cloth, FVector3f{-30.f, 0.f, 190.f});
+		MakeRectangle(Cloth, FVector3f{-30.f, 0.f, 290.f});
 
 		FDatasmithMeshExporter MeshExporter;
 		FString ClothFilePath = FString(DatasmithSceneExporter.GetAssetsOutputPath()) / ClothElement->GetName();
