[3ds Max Exporter] Fixed Very long material name causes Importer to crash

- names exceeding FName max length are shortened

diff --git a/Engine/Source/Programs/Enterprise/Datasmith/DatasmithMaxExporter/Private/DatasmithMaxDirectLinkMaterial.cpp b/Engine/Source/Programs/Enterprise/Datasmith/DatasmithMaxExporter/Private/DatasmithMaxDirectLinkMaterial.cpp
--- a/Engine/Source/Programs/Enterprise/Datasmith/DatasmithMaxExporter/Private/DatasmithMaxDirectLinkMaterial.cpp
+++ b/Engine/Source/Programs/Enterprise/Datasmith/DatasmithMaxExporter/Private/DatasmithMaxDirectLinkMaterial.cpp
@@ -15,6 +15,7 @@
 #include "DatasmithExportOptions.h"
 
 #include "Misc/Paths.h"
+#include "Misc/SecureHash.h"
 
 
 #include "Logging/LogMacros.h"
@@ -61,14 +62,29 @@ void FMaterialsCollectionTracker::ReleaseActualMaterial(FMaterialTracker& Materi
 	}
 }
 
+
 const TCHAR* FMaterialsCollectionTracker::GetMaterialName(Mtl* Material)
 {
 	if (FString* NamePtr = UsedMaterialToDatasmithMaterialName.Find(Material))
 	{
 		return **NamePtr;
 	}
-	return *UsedMaterialToDatasmithMaterialName.Add(Material, 
-		MaterialNameProvider.GenerateUniqueName(FDatasmithUtils::SanitizeObjectName(Material->GetName().data())));
+
+	// Using material's Name as Max doesn't have another persistent id for materials(like INode has with its GetHandle)
+	FString SourceName = Material->GetName().data();
+	FString SanitizedName = FDatasmithUtils::SanitizeObjectName(SourceName);
+
+	// Limit Datasmith name size to FName length max
+	if (SanitizedName.Len() > (NAME_SIZE-1))
+	{
+		// In case we have such long material names make sure that their shortened versions don't match when their source names are different
+		// We could rely GenerateUniqueName but since we can build a stable name better do it because we can.
+		// GenerateUniqueName is the last resort when materials have identical names
+		FString NameHash = FMD5::HashBytes(reinterpret_cast<const uint8*>(*SourceName), SourceName.Len()*sizeof(TCHAR));
+		SanitizedName = SanitizedName.Left(NAME_SIZE-1-NameHash.Len()) + NameHash;
+	}
+
+	return *UsedMaterialToDatasmithMaterialName.Add(Material, MaterialNameProvider.GenerateUniqueName(SanitizedName));
 }
 
 bool FMaterialsCollectionTracker::IsMaterialUsed(Mtl* Material)
