[Revit Exporter] Fixed sync of removed family instances

- Fix for 'HideByCategory' is just a side effect. This change fixes other circumstances where instance removal is synced

- code: simplified/commented Optimize function

diff --git a/Engine/Source/Programs/Enterprise/Datasmith/DatasmithRevitExporter/Private/DatasmithRevitDirectLink.cs b/Engine/Source/Programs/Enterprise/Datasmith/DatasmithRevitExporter/Private/DatasmithRevitDirectLink.cs
--- a/Engine/Source/Programs/Enterprise/Datasmith/DatasmithRevitExporter/Private/DatasmithRevitDirectLink.cs
+++ b/Engine/Source/Programs/Enterprise/Datasmith/DatasmithRevitExporter/Private/DatasmithRevitDirectLink.cs
@@ -125,6 +125,16 @@ namespace DatasmithRevitExporter
 					FDocumentData.FBaseElementData ElementData = CachedElements[ElemId];
 					CachedElements.Remove(ElemId);
 					ElementData.Parent?.ChildElements.Remove(ElementData);
+
+					// Remove all owned children
+					foreach (FDocumentData.FBaseElementData ChildElementData in ElementData.ChildElements)
+					{
+						if (ChildElementData.bOwnedByParent && (ChildElementData.ElementActor != null))
+						{
+							ElementData.ElementActor.RemoveChild(ChildElementData.ElementActor);
+						}
+					}
+
 					DatasmithScene.RemoveActor(ElementData.ElementActor, FDatasmithFacadeScene.EActorRemovalRule.KeepChildrenAndKeepRelativeTransform);
 					ExportedActorsMap.Remove(ElementData.ElementActor.GetName());
 				}
diff --git a/Engine/Source/Programs/Enterprise/Datasmith/DatasmithRevitExporter/Private/DatasmithRevitDocument.cs b/Engine/Source/Programs/Enterprise/Datasmith/DatasmithRevitExporter/Private/DatasmithRevitDocument.cs
--- a/Engine/Source/Programs/Enterprise/Datasmith/DatasmithRevitExporter/Private/DatasmithRevitDocument.cs
+++ b/Engine/Source/Programs/Enterprise/Datasmith/DatasmithRevitExporter/Private/DatasmithRevitDocument.cs
@@ -45,7 +45,10 @@ namespace DatasmithRevitExporter
 
 				NewActiveDocument.SettingsChangedHandler = new EventHandler((object Sender, EventArgs Args) =>
 				{
-					NewActiveDocument.ActiveDirectLinkInstance.bSettingsDirty = true;
+					if (NewActiveDocument.ActiveDirectLinkInstance != null)
+					{
+						NewActiveDocument.ActiveDirectLinkInstance.bSettingsDirty = true;
+					}
 				});
 
 				FSettingsManager.SettingsUpdated += NewActiveDocument.SettingsChangedHandler;
diff --git a/Engine/Source/Programs/Enterprise/Datasmith/DatasmithRevitExporter/Private/DatasmithRevitDocumentData.cs b/Engine/Source/Programs/Enterprise/Datasmith/DatasmithRevitExporter/Private/DatasmithRevitDocumentData.cs
--- a/Engine/Source/Programs/Enterprise/Datasmith/DatasmithRevitExporter/Private/DatasmithRevitDocumentData.cs
+++ b/Engine/Source/Programs/Enterprise/Datasmith/DatasmithRevitExporter/Private/DatasmithRevitDocumentData.cs
@@ -299,6 +299,8 @@ namespace DatasmithRevitExporter
 
 			public FBaseElementData			Parent = null;
 
+			public bool bOwnedByParent = false; // Lifetime is controlled by parent Element
+
 			public FBaseElementData(
 				ElementType InElementType, FDocumentData InDocumentData
 			)
@@ -314,20 +316,6 @@ namespace DatasmithRevitExporter
 				DocumentData = InDocumentData;
 			}
 
-			public bool IsSimpleActor()
-			{
-				FDatasmithFacadeActorMesh MeshActor = ElementActor as FDatasmithFacadeActorMesh;
-
-				if (MeshActor != null && MeshActor.GetMeshName().Length == 0)
-				{
-					ElementActor = new FDatasmithFacadeActor(MeshActor.GetName());
-					CopyActorData(MeshActor, ElementActor);
-					return true;
-				}
-
-				return !(ElementActor is FDatasmithFacadeActorMesh || ElementActor is FDatasmithFacadeActorLight || ElementActor is FDatasmithFacadeActorCamera);
-			}
-
 			void CopyActorData(FDatasmithFacadeActor InFromActor, FDatasmithFacadeActor InToActor)
 			{
 				InToActor.SetLabel(InFromActor.GetLabel());
@@ -358,9 +346,16 @@ namespace DatasmithRevitExporter
 				ElementMetaData?.SetAssociatedElement(InToActor);
 			}
 
+			// Return element when this 'ElementData' is associated with an Element
+			// todo: Probably worth separating FBaseElementData into different entity, calling it now "Element" data is confusing as it's also used to hold for non-Element datasmith actors
+			protected virtual Element GetElement()
+			{
+				return null;
+			}
+
 			public void AddToScene(FDatasmithFacadeScene InScene, FBaseElementData InParent, bool bInSkipChildren, bool bInForceAdd = false)
 			{
-				Element ThisElement = (this as FElementData)?.CurrentElement;
+				Element ThisElement = GetElement();
 
 				if (!bInSkipChildren)
 				{
@@ -421,25 +416,24 @@ namespace DatasmithRevitExporter
 				bIsModified = false;
 			}
 
+			// Returns true if element's actor is simple(plain actor without descendants)
 			public bool Optimize()
 			{
-				bool bIsSimpleActor = IsSimpleActor();
-
-				List<FBaseElementData> ChildrenToRemove = new List<FBaseElementData>();
-
-				foreach (FBaseElementData CurrentChild in ChildElements)
+				// Replace MeshActor without geometry by a dummy Actor
+				if (ElementActor is FDatasmithFacadeActorMesh MeshActor && MeshActor.GetMeshName().Length == 0)
 				{
-					if (CurrentChild.Optimize())
-					{
-						ChildrenToRemove.Add(CurrentChild);
-					}
+					ElementActor = new FDatasmithFacadeActor(MeshActor.GetName());
+					CopyActorData(MeshActor, ElementActor);
 				}
 
-				foreach (FBaseElementData Child in ChildrenToRemove)
+				// Optimize and remove children whose actors are simple
+				List<FBaseElementData> SimpleChildren = ChildElements.Where(Child => Child.Optimize()).ToList(); // Build a list of elements to remove after enumeration
+				foreach (FBaseElementData Child in SimpleChildren)
 				{
 					ChildElements.Remove(Child);
 				}
 
+				bool bIsSimpleActor = !(ElementActor is FDatasmithFacadeActorMesh || ElementActor is FDatasmithFacadeActorLight || ElementActor is FDatasmithFacadeActorCamera);
 				return bIsSimpleActor && (ChildElements.Count == 0) && bOptimizeHierarchy;
 			}
 
@@ -480,6 +474,11 @@ namespace DatasmithRevitExporter
 				InitializeElement(InWorldTransform, this);
 			}
 
+			protected override Element GetElement() 
+			{
+				return CurrentElement;
+			}
+
 			public void InitializePivotPlacement(ref Transform InOutWorldTransform)
 			{
 				// If element has location, use it as a transform in order to have better pivot placement.
@@ -718,6 +717,7 @@ namespace DatasmithRevitExporter
 			)
 			{
 				FBaseElementData InstanceData = new FBaseElementData(InInstanceType, DocumentData);
+				InstanceData.bOwnedByParent = true;
 
 				FamilyInstance CurrentFamilyInstance = CurrentElement as FamilyInstance;
 				if (CurrentFamilyInstance != null && CurrentFamilyInstance.HasModifiedGeometry())
@@ -787,7 +787,7 @@ namespace DatasmithRevitExporter
 				FDatasmithRevitLight.SetLightProperties(InLightAsset, CurrentElement, LightActor);
 
 				// Add the light actor to the Datasmith actor hierarchy.
-				AddChildActor(LightActor, LightMetaData, false);
+				AddChildActor(LightActor, LightMetaData, false, true);
 			}
 
 			public bool AddRPCActor(
@@ -943,7 +943,7 @@ namespace DatasmithRevitExporter
 				}
 
 				// Add the RPC mesh actor to the Datasmith actor hierarchy.
-				AddChildActor(FacadeActor, ElementMetaData, false);
+				AddChildActor(FacadeActor, ElementMetaData, false, true);
 
 				return OutDatasmithMesh != null;
 			}
@@ -961,7 +961,8 @@ namespace DatasmithRevitExporter
 			public void AddChildActor(
 				FDatasmithFacadeActor ChildActor,
 				FDatasmithFacadeMetaData MetaData,
-				bool bOptimizeHierarchy
+				bool bOptimizeHierarchy,
+				bool bOwned // Make its lifetime controlled by this ElementData
 			)
 			{
 				FBaseElementData ElementData = new FBaseElementData(ChildActor, MetaData, DocumentData);
@@ -971,6 +972,7 @@ namespace DatasmithRevitExporter
 
 				Parent.ChildElements.Add(ElementData);
 				ElementData.Parent = Parent;
+				ElementData.bOwnedByParent = bOwned;
 			}
 
 			public void InitializeElement(
