CAD conversion operation crashes with "out of memory":
One problem is a (undefined) memory leak in the CAD conversion process.
The workaround to fix it, is to restart CADWorker as soon as the leak is too important.
A first implementation is done is 5.1.1 withour header modification: the state ETaskState::Unknown is used to inform that the worker need to be restart.

diff --git a/Engine/Plugins/Enterprise/DatasmithCADImporter/Source/DatasmithDispatcher/Private/DatasmithDispatcher.cpp b/Engine/Plugins/Enterprise/DatasmithCADImporter/Source/DatasmithDispatcher/Private/DatasmithDispatcher.cpp
--- a/Engine/Plugins/Enterprise/DatasmithCADImporter/Source/DatasmithDispatcher/Private/DatasmithDispatcher.cpp
+++ b/Engine/Plugins/Enterprise/DatasmithCADImporter/Source/DatasmithDispatcher/Private/DatasmithDispatcher.cpp
@@ -93,16 +93,31 @@ void FDatasmithDispatcher::SetTaskState(int32 TaskIndex, ETaskState TaskState)
 		Task.State = TaskState;
 		FileDescriptor = Task.FileDescription;
 
-		if (TaskState == ETaskState::ProcessOk
-			|| TaskState == ETaskState::ProcessFailed
-			|| TaskState == ETaskState::FileNotFound)
+		switch (TaskState)
 		{
+		case ETaskState::Unknown:
+		{
+			Task.State = ETaskState::ProcessOk;
 			CompletedTaskCount++;
+			break;
 		}
 
-		if (TaskState == ETaskState::UnTreated)
+		case ETaskState::ProcessOk:
+		case ETaskState::ProcessFailed:
+		case ETaskState::FileNotFound:
+		{
+			CompletedTaskCount++;
+			break;
+		}
+
+		case ETaskState::UnTreated:
 		{
 			NextTaskIndex = TaskIndex;
+			break;
+		}
+
+		default:
+			break;
 		}
 	}
 
diff --git a/Engine/Plugins/Enterprise/DatasmithCADImporter/Source/DatasmithDispatcher/Private/DatasmithWorkerHandler.cpp b/Engine/Plugins/Enterprise/DatasmithCADImporter/Source/DatasmithDispatcher/Private/DatasmithWorkerHandler.cpp
--- a/Engine/Plugins/Enterprise/DatasmithCADImporter/Source/DatasmithDispatcher/Private/DatasmithWorkerHandler.cpp
+++ b/Engine/Plugins/Enterprise/DatasmithCADImporter/Source/DatasmithDispatcher/Private/DatasmithWorkerHandler.cpp
@@ -214,9 +214,9 @@ void FDatasmithWorkerHandler::RunInternal()
 					ProcessCommand(*Command);
 
 					bool bProcessingOver = CurrentTask.IsSet() == false;
-					if (bProcessingOver)
+					if (bProcessingOver && bShouldTerminate)
 					{
-						WorkerState = bShouldTerminate ? EWorkerState::Closing : EWorkerState::Idle;
+						WorkerState = EWorkerState::Closing;
 					}
 				}
 				else
@@ -342,6 +342,27 @@ void FDatasmithWorkerHandler::ProcessCommand(FCompletedTaskCommand& CompletedTas
 	Dispatcher.LinkCTFileToUnrealCacheFile(CurrentTask->FileDescription, CompletedTaskCommand.SceneGraphFileName, CompletedTaskCommand.GeomFileName);
 	Dispatcher.LogWarningMessages(CompletedTaskCommand.WarningMessages);
 	CurrentTask.Reset();
+
+	if (CompletedTaskCommand.ProcessResult == ETaskState::Unknown)
+	{
+		// try to close the process gracefully
+		if (WorkerHandle.IsValid())
+		{
+			FPlatformProcess::TerminateProc(WorkerHandle, true);
+		}
+
+		// Process commands still in input queue
+		CommandIO.Disconnect(0);
+		while (TSharedPtr<ICommand> Command = CommandIO.GetNextCommand(0))
+		{
+			ProcessCommand(*Command);
+		}
+		WorkerState = EWorkerState::Uninitialized;
+	}
+	else
+	{
+		WorkerState = EWorkerState::Idle;
+	}
 }
 
 const TCHAR* FDatasmithWorkerHandler::EWorkerErrorStateAsString(EWorkerErrorState Error)
diff --git a/Engine/Source/Programs/Enterprise/Datasmith/DatasmithCADWorker/Private/DatasmithCADWorkerImpl.cpp b/Engine/Source/Programs/Enterprise/Datasmith/DatasmithCADWorker/Private/DatasmithCADWorkerImpl.cpp
--- a/Engine/Source/Programs/Enterprise/Datasmith/DatasmithCADWorker/Private/DatasmithCADWorkerImpl.cpp
+++ b/Engine/Source/Programs/Enterprise/Datasmith/DatasmithCADWorker/Private/DatasmithCADWorkerImpl.cpp
@@ -25,7 +25,7 @@ FDatasmithCADWorkerImpl::FDatasmithCADWorkerImpl(int32 InServerPID, int32 InServ
 	, EnginePluginsPath(InEnginePluginsPath)
 	, CachePath(InCachePath)
 	, PingStartCycle(0)
-	, bRequestExit(false)
+	, bRequestRestart(false)
 {
 }
 
@@ -145,6 +145,9 @@ uint64 DefineMaximumAllowedDuration(const CADLibrary::FFileDescriptor& FileDescr
 		break;
 	}
 
+	constexpr int64 OneKiloBit = 1024;
+	UE_LOG(LogDatasmithCADWorker, Verbose, TEXT("    - File size %lld KB"), FileStatData.FileSize / OneKiloBit);
+
 	uint64 MaximumDuration = ((double)FileStatData.FileSize) * MaxTimePerMb * SafetyCoeficient;
 	return FMath::Max(MaximumDuration, (uint64)30);
 }
@@ -163,15 +166,18 @@ void FDatasmithCADWorkerImpl::ProcessCommand(const FRunTaskCommand& RunTaskComma
 	Checkers.Emplace(UE::Tasks::Launch(TEXT("MemoryChecker"), [&]() { CheckMemory(); }));
 
 	CADLibrary::FCADFileReader FileReader(ImportParameters, FileToProcess, EnginePluginsPath, CachePath);
-	ETaskState ProcessResult = FileReader.ProcessFile();
+	CompletedTask.ProcessResult = FileReader.ProcessFile();
 
 	bProcessIsRunning = false;
 	UE::Tasks::Wait(Checkers);
 
-	CompletedTask.ProcessResult = ProcessResult;
-
 	if (CompletedTask.ProcessResult == ETaskState::ProcessOk)
 	{
+		if (bRequestRestart)
+		{
+			CompletedTask.ProcessResult = ETaskState::Unknown;
+		}
+			
 		const CADLibrary::FCADFileData& CADFileData = FileReader.GetCADFileData();
 		CompletedTask.ExternalReferences = CADFileData.GetExternalRefSet();
 		CompletedTask.SceneGraphFileName = CADFileData.GetSceneGraphFileName();
@@ -186,18 +192,22 @@ void FDatasmithCADWorkerImpl::ProcessCommand(const FRunTaskCommand& RunTaskComma
 			UE_LOG(LogDatasmithCADWorker, Verbose, TEXT("     - Body %s"), *BodyFileName);
 		}
 	}
+	else if(CompletedTask.ProcessResult == ETaskState::FileNotFound)
+	{
+		UE_LOG(LogDatasmithCADWorker, Verbose, TEXT("=> File not found %s"), *FileToProcess.GetFileName());
+	}
 	else
 	{
 		UE_LOG(LogDatasmithCADWorker, Verbose, TEXT("=> Process %s %s failed"), *FileToProcess.GetFileName(), *FileToProcess.GetConfiguration());
 	}
-	CommandIO.SendCommand(CompletedTask, Config::SendCommandTimeout_s);
-
 	UE_LOG(LogDatasmithCADWorker, Verbose, TEXT("End of Process %s %s"), *FileToProcess.GetFileName(), *FileToProcess.GetConfiguration());
-
-	if (bRequestExit)
+	
+	if(bRequestRestart)
 	{
-		FPlatformMisc::RequestExit(true);
+		GLog->Flush();
 	}
+
+	CommandIO.SendCommand(CompletedTask, Config::SendCommandTimeout_s);
 }
 
 void FDatasmithCADWorkerImpl::CheckDuration(const CADLibrary::FFileDescriptor& FileToProcess, const int64 MaxDuration)
@@ -249,7 +259,8 @@ void FDatasmithCADWorkerImpl::CheckMemory()
 	UE_LOG(LogDatasmithCADWorker, Verbose, TEXT("    - Max Ram used %llu MB"), MaxMemoryUsed / OneMegaBit);
 	if(EndMemoryUsed > GigaBit)
 	{
-		UE_LOG(LogDatasmithCADWorker, Verbose, TEXT("    - End Ram used %llu MB. Exceeded the limit to start a new process. CADWorker restart is requested"), EndMemoryUsed / OneMegaBit);
-		bRequestExit = true;
+		UE_LOG(LogDatasmithCADWorker, Verbose, TEXT("    - Ram used (%llu MB) after cleanup exceeds limit to start new process. CADWorker restart is requested"), EndMemoryUsed / OneMegaBit);
+		bRequestRestart = true;
 	}
+
 }
diff --git a/Engine/Source/Programs/Enterprise/Datasmith/DatasmithCADWorker/Private/DatasmithCADWorkerImpl.h b/Engine/Source/Programs/Enterprise/Datasmith/DatasmithCADWorker/Private/DatasmithCADWorkerImpl.h
--- a/Engine/Source/Programs/Enterprise/Datasmith/DatasmithCADWorker/Private/DatasmithCADWorkerImpl.h
+++ b/Engine/Source/Programs/Enterprise/Datasmith/DatasmithCADWorker/Private/DatasmithCADWorkerImpl.h
@@ -46,5 +46,5 @@ private:
 	uint64 PingStartCycle;
 
 	std::atomic<bool> bProcessIsRunning;
-	std::atomic<bool> bRequestExit;
+	std::atomic<bool> bRequestRestart;
 };
