[Solidworks] Fixed material export for Part files

- also added a debug logging(disabled by default, use DatasmithSolidworksDebugOutput synbol to enable)

diff --git a/Engine/Source/Programs/Enterprise/Datasmith/DatasmithSolidworksExporter/Private/Addin.cs b/Engine/Source/Programs/Enterprise/Datasmith/DatasmithSolidworksExporter/Private/Addin.cs
--- a/Engine/Source/Programs/Enterprise/Datasmith/DatasmithSolidworksExporter/Private/Addin.cs
+++ b/Engine/Source/Programs/Enterprise/Datasmith/DatasmithSolidworksExporter/Private/Addin.cs
@@ -6,10 +6,15 @@ using SolidWorks.Interop.swpublished;
 using SolidWorksTools;
 using SolidWorksTools.File;
 using System;
+using System.Collections.Concurrent;
 using System.Collections.Generic;
+using System.Diagnostics;
+using System.IO;
 using System.Linq;
 using System.Reflection;
 using System.Runtime.InteropServices;
+using System.Threading;
+using Environment = System.Environment;
 
 namespace DatasmithSolidworks
 {
@@ -34,13 +39,19 @@ namespace DatasmithSolidworks
 		public ISldWorks SolidworksApp { get; private set; } = null;
 		public FDocument CurrentDocument { get; private set; } = null;
 
+		// DebugLog is enabled with DatasmithSolidworksDebugOutput conditional compilation symbol
+		private ConcurrentQueue<string> DebugLog;
+		private Thread LogWriterThread;
+
 		public static Addin Instance { get; private set; } = null;
 
+
 		public Addin()
 		{
 			if (Instance == null)
 			{
 				Instance = this;
+				StartLogWriterTread();
 			}
 		}
 
@@ -81,8 +92,7 @@ namespace DatasmithSolidworks
 			object[] ObjDocuments = (object[])SolidworksApp.GetDocuments();
 			for (int Index = 0; Index < NumOpenDocs; Index++)
 			{
-				ModelDoc2 Doc = ObjDocuments[Index] as ModelDoc2;
-				if (Doc != null)
+				if (ObjDocuments[Index] is ModelDoc2 Doc)
 				{
 					CurrentOpenDocumentsSet.Add(GetDocumentId(Doc));
 				}
@@ -524,5 +534,48 @@ namespace DatasmithSolidworks
 		}
 
 		#endregion
+
+		#region Logging
+
+		[Conditional("DatasmithSolidworksDebugOutput")]
+		private void StartLogWriterTread()
+		{
+			DebugLog = new ConcurrentQueue<string>();
+
+			LogWriterThread = new Thread(() =>
+			{
+				LogWriterProc();
+			});
+
+			LogWriterThread.Start();
+		}
+
+		[Conditional("DatasmithSolidworksDebugOutput")]
+		private void LogWriterProc()
+		{
+			string LogPath = Path.Combine(Environment.GetFolderPath(Environment.SpecialFolder.LocalApplicationData),
+				"UnrealDatasmithExporter/Saved/Logs/UnrealDatasmithSolidworksExporterDebug.log");
+
+			StreamWriter LogFile = new StreamWriter(LogPath);
+
+			while (true)
+			{
+				while (DebugLog.TryDequeue(out string Message))
+				{
+					LogFile.WriteLine(Message);
+				}
+				LogFile.Flush();
+
+				Thread.Sleep(10);
+			}
+		}
+
+		[Conditional("DatasmithSolidworksDebugOutput")]
+		public void LogDebug(string Message)  
+		{
+			DebugLog.Enqueue(Message);
+		}
+
+		#endregion
 	}
 }
diff --git a/Engine/Source/Programs/Enterprise/Datasmith/DatasmithSolidworksExporter/Private/AssemblyDocument.cs b/Engine/Source/Programs/Enterprise/Datasmith/DatasmithSolidworksExporter/Private/AssemblyDocument.cs
--- a/Engine/Source/Programs/Enterprise/Datasmith/DatasmithSolidworksExporter/Private/AssemblyDocument.cs
+++ b/Engine/Source/Programs/Enterprise/Datasmith/DatasmithSolidworksExporter/Private/AssemblyDocument.cs
@@ -49,7 +49,7 @@ namespace DatasmithSolidworks
 			SwAsmDoc = InSwDoc;
 		}
 
-		public override void PreExport(FMeshes Meshes)
+		public override void PreExport(FMeshes Meshes, bool bConfigurations)
 		{
 		}
 
diff --git a/Engine/Source/Programs/Enterprise/Datasmith/DatasmithSolidworksExporter/Private/DatasmithExporter.cs b/Engine/Source/Programs/Enterprise/Datasmith/DatasmithSolidworksExporter/Private/DatasmithExporter.cs
--- a/Engine/Source/Programs/Enterprise/Datasmith/DatasmithSolidworksExporter/Private/DatasmithExporter.cs
+++ b/Engine/Source/Programs/Enterprise/Datasmith/DatasmithSolidworksExporter/Private/DatasmithExporter.cs
@@ -4,6 +4,7 @@ using SolidWorks.Interop.sldworks;
 using System;
 using System.Collections.Concurrent;
 using System.Collections.Generic;
+using System.Diagnostics;
 using System.IO;
 using System.Linq;
 using System.Runtime.InteropServices;
@@ -357,17 +358,21 @@ namespace DatasmithSolidworks
 
 		public bool ExportMesh(string InMeshName, FMeshData InData, FActorName InUpdateMeshActor, out Tuple<FDatasmithFacadeMeshElement, FDatasmithFacadeMesh> OutMeshPair)
 		{
+			Addin.Instance.LogDebug($"ExportMesh(MeshName: {InMeshName}, MeshActor: {InUpdateMeshActor.GetString()})");
+
 			InMeshName = SanitizeName(InMeshName);  //Compute mesh name early(it might be needed to remove old mesh)
 
 			OutMeshPair = null;
 
 			if (InData.Vertices == null || InData.Normals == null || InData.TexCoords == null || InData.Triangles == null)
 			{
+				Addin.Instance.LogDebug($"  skipping(some attributes are null) - Vertices:{InData.Vertices == null}, Normals:{InData.Normals == null}, TexCoords: {InData.TexCoords == null}, Triangles:{InData.Triangles}");
 				return false;
 			}
 
 			if (InData.Vertices.Length == 0 || InData.Normals.Length == 0 || InData.TexCoords.Length == 0 || InData.Triangles.Length == 0)
 			{
+				Addin.Instance.LogDebug($"  skipping - Vertices: {InData.Vertices.Length}, Normals: {InData.Normals.Length}, TexCoords: {InData.TexCoords.Length} Triangles: {InData.Triangles.Length}");
 				return false;
 			}
 
@@ -410,11 +415,15 @@ namespace DatasmithSolidworks
 				{
 					if (!MeshAddedMaterials.Contains(Triangle.MaterialID))
 					{
+						Addin.Instance.LogDebug($"  set material if it's exported: {Triangle.MaterialID}");
+
 						FDatasmithFacadeMaterialInstance Material = null;
 						ExportedMaterialsMap.TryGetValue(Triangle.MaterialID, out Material);
 
 						if (Material != null)
 						{
+							Addin.Instance.LogDebug($" SetMaterial({Material.GetName()}, SlotId: {Triangle.MaterialID})");
+
 							MeshAddedMaterials.Add(Triangle.MaterialID);
 							MeshElement.SetMaterial(Material.GetName(), Triangle.MaterialID);
 							MatID = Triangle.MaterialID;
@@ -440,6 +449,7 @@ namespace DatasmithSolidworks
 				if (ExportedActorsMap.TryGetValue(InUpdateMeshActor, out ExportedActorInfo) && ExportedActorInfo.Item1 == EActorType.MeshActor)
 				{
 					FDatasmithFacadeActorMesh MeshActor = ExportedActorInfo.Item2 as FDatasmithFacadeActorMesh;
+					Addin.Instance.LogDebug($"  SetMesh on MeshActor: {MeshActor.GetName()}: {MeshActor.GetLabel()}  )");
 					MeshActor.SetMesh(InMeshName);
 				}
 			}
@@ -745,9 +755,12 @@ namespace DatasmithSolidworks
 					}
 				}
 			});
+
+			Addin.Instance.LogDebug($"AddMaterials");
 			// Adding stuff to a datasmith scene cannot be multithreaded!
 			foreach (FDatasmithFacadeMaterialInstance Mat in CreatedMaterials)
 			{
+				Addin.Instance.LogDebug($"  AddMaterial: {Mat.GetName()}: {Mat.GetLabel()}");
 				DatasmithScene.AddMaterial(Mat);
 			}
 			foreach (FDatasmithFacadeTexture Texture in CreatedTextures)
@@ -758,14 +771,19 @@ namespace DatasmithSolidworks
 
 		private bool CreateAndCacheMaterial(FMaterial InMaterial, out List<FDatasmithFacadeTexture> OutCreatedTextures, out FDatasmithFacadeMaterialInstance OutCreatedMaterial)
 		{
+			Addin.Instance.LogDebug($"CreateAndCacheMaterial({InMaterial.Name}: {InMaterial.ID})");
+
 			OutCreatedTextures = null;
 			OutCreatedMaterial = null;
 
 			if (ExportedMaterialsMap.ContainsKey(InMaterial.ID))
 			{
+				Addin.Instance.LogDebug($"  skipping, already present in map");
 				return false;
 			}
 
+			Addin.Instance.LogDebug($"  making, ShaderName: '{InMaterial.ShaderName}'");
+
 			FMaterial.EMaterialType Type = FMaterial.GetMaterialType(InMaterial.ShaderName);
 
 			float Roughness = (float)InMaterial.Roughness;
diff --git a/Engine/Source/Programs/Enterprise/Datasmith/DatasmithSolidworksExporter/Private/Document.cs b/Engine/Source/Programs/Enterprise/Datasmith/DatasmithSolidworksExporter/Private/Document.cs
--- a/Engine/Source/Programs/Enterprise/Datasmith/DatasmithSolidworksExporter/Private/Document.cs
+++ b/Engine/Source/Programs/Enterprise/Datasmith/DatasmithSolidworksExporter/Private/Document.cs
@@ -170,6 +170,8 @@ namespace DatasmithSolidworks
 				Stopwatch Watch = Stopwatch.StartNew();
 #endif
 				FMeshes Meshes = new FMeshes();
+				
+				PreExport(Meshes, false);
 				ExportToDatasmithScene(Meshes);
 
 				ExportLights();
@@ -203,7 +205,7 @@ namespace DatasmithSolidworks
 
 				FMeshes Meshes = new FMeshes();
 
-				PreExport(Meshes);
+				PreExport(Meshes, true);
 
 				List<FConfigurationData> Configs = new FConfigurationExporter(Meshes).ExportConfigurations(this);
 				bHasConfigurations = (Configs != null) && (Configs.Count != 0);
@@ -239,7 +241,7 @@ namespace DatasmithSolidworks
 
 		public abstract bool HasMaterialUpdates();
 		public abstract void ExportToDatasmithScene(FMeshes Meshes);
-		public abstract void PreExport(FMeshes Meshes);  // Called before configurations are parsed to prepare meshes needed to identify if they create different configurations
+		public abstract void PreExport(FMeshes Meshes, bool bConfigurations);  // Called before configurations are parsed to prepare meshes needed to identify if they create different configurations
 
 		public virtual void Init()
 		{
diff --git a/Engine/Source/Programs/Enterprise/Datasmith/DatasmithSolidworksExporter/Private/PartDocument.cs b/Engine/Source/Programs/Enterprise/Datasmith/DatasmithSolidworksExporter/Private/PartDocument.cs
--- a/Engine/Source/Programs/Enterprise/Datasmith/DatasmithSolidworksExporter/Private/PartDocument.cs
+++ b/Engine/Source/Programs/Enterprise/Datasmith/DatasmithSolidworksExporter/Private/PartDocument.cs
@@ -79,8 +79,17 @@ namespace DatasmithSolidworks
 			}
 		}
 
-		public override void PreExport(FMeshes Meshes)
+		public override void PreExport(FMeshes Meshes, bool bConfigurations)
 		{
+			if (!bConfigurations)
+			{
+				// Not exporting configurations - just load materials here
+				ExportedPartMaterials = FObjectMaterials.LoadPartMaterials(this, SwPartDoc, swDisplayStateOpts_e.swThisDisplayState, null);
+				return;
+			}
+
+			// Load mesh data for each configuration - need to have information about meshes to understand what to export for configurations.
+			// (Also load materials - mesh processing later needs them)
 			ConfigurationManager ConfigManager = (SwPartDoc as ModelDoc2).ConfigurationManager;
 			IConfiguration OriginalConfiguration = ConfigManager.ActiveConfiguration;
 
@@ -88,6 +97,9 @@ namespace DatasmithSolidworks
 			{
 				SwDoc.ShowConfiguration2(ConfigurationName);
 
+				// Load materials for each configuration
+				ExportedPartMaterials = FObjectMaterials.LoadPartMaterials(this, SwPartDoc, swDisplayStateOpts_e.swThisDisplayState, null);
+
 				// todo: Solidworks api docs says that GetRootComponent3 is null for Part but our current 
 				//   parsing code relies on it's existence(and it exists in all cases). Probably better to refactor this without Component use for parts
 				//   This is possible, we don't need Component, just code would be be less 'generic'   
@@ -114,8 +126,6 @@ namespace DatasmithSolidworks
 
 			SetExportStatus($"{PartName} Materials");
 
-			ExportedPartMaterials = FObjectMaterials.LoadPartMaterials(this, SwPartDoc, swDisplayStateOpts_e.swThisDisplayState, null);
-
 			Exporter.ExportMaterials(ExportedMaterialsMap);
 
 			SetExportStatus($"{PartName} Meshes");
