[Revit] Stable AND unique actor names

- Actor name for a Revit element is now composed from the 'path' through whole hierarchy of Revit elements(including linked documents). This makes every element that occurs multiple times in the scene(e.g. in Linked docs) to have different actor name

- This makes it unique(path to element occurrence is unique).

- This 'path' is stable because of not using any 'unstable' elements(like temporary geometry symbols created by Revit during export - only elements UniqueIds)

- additional: debug output for scene parsing(compile-time enabled using DatasmithRevitDebugOutput symbol)

diff --git a/Engine/Source/Programs/Enterprise/Datasmith/DatasmithRevitExporter/Private/DatasmithRevitDirectLink.cs b/Engine/Source/Programs/Enterprise/Datasmith/DatasmithRevitExporter/Private/DatasmithRevitDirectLink.cs
--- a/Engine/Source/Programs/Enterprise/Datasmith/DatasmithRevitExporter/Private/DatasmithRevitDirectLink.cs
+++ b/Engine/Source/Programs/Enterprise/Datasmith/DatasmithRevitExporter/Private/DatasmithRevitDirectLink.cs
@@ -171,6 +171,7 @@ namespace DatasmithRevitExporter
 		private FDatasmithFacadeDirectLink								DatasmithDirectLink;
 		private string													SceneName;
 
+		public Dictionary<string, FDocumentData> ExportedDocuments 		= new Dictionary<string, FDocumentData>();
 		private Dictionary<string, int>									ExportedActorNames = new Dictionary<string, int>();
 
 		// The number of times this document was synced (sent to receiver)
@@ -232,6 +233,11 @@ namespace DatasmithRevitExporter
 
 			Debug.Assert(ActiveInstance != null);
 
+			if (ActiveInstance == null)
+			{
+				return;
+			}
+
 			// Handle modified elements
 			foreach (ElementId ElemId in InArgs.GetModifiedElementIds())
 			{
diff --git a/Engine/Source/Programs/Enterprise/Datasmith/DatasmithRevitExporter/Private/DatasmithRevitDocumentData.cs b/Engine/Source/Programs/Enterprise/Datasmith/DatasmithRevitExporter/Private/DatasmithRevitDocumentData.cs
--- a/Engine/Source/Programs/Enterprise/Datasmith/DatasmithRevitExporter/Private/DatasmithRevitDocumentData.cs
+++ b/Engine/Source/Programs/Enterprise/Datasmith/DatasmithRevitExporter/Private/DatasmithRevitDocumentData.cs
@@ -457,7 +457,7 @@ namespace DatasmithRevitExporter
 			public Element		CurrentElement = null;
 			public Transform	MeshPointsTransform = null;
 
-			private Stack<FBaseElementData> InstanceDataStack = new Stack<FBaseElementData>();
+			public Stack<FBaseElementData> InstanceDataStack = new Stack<FBaseElementData>();
 
 			public FElementData(
 				Element InElement,
@@ -467,11 +467,6 @@ namespace DatasmithRevitExporter
 				: base(InElement.Document.GetElement(InElement.GetTypeId()) as ElementType, InDocumentData)
 			{
 				CurrentElement = InElement;
-
-				InitializePivotPlacement(ref InWorldTransform);
-
-				// Create a new Datasmith mesh actor.
-				InitializeElement(InWorldTransform, this);
 			}
 
 			protected override Element GetElement() 
@@ -1088,26 +1083,10 @@ namespace DatasmithRevitExporter
 
 			private string GetActorName(bool bEnsureUnique)
 			{
-				string ActorName;
-
-				if (InstanceDataStack.Count == 0)
-				{
-					ActorName = $"{DocumentData.DocumentId}:{CurrentElement.UniqueId}";
-				}
-				else
-				{
-					ActorName = GenerateUniqueInstanceName();
-				}
-
-				if (bEnsureUnique && DocumentData.DirectLink != null)
-				{
-					ActorName = DocumentData.DirectLink.EnsureUniqueActorName(ActorName);
-				}
-
-				return ActorName;
+				return DocumentData.GetActorName(this);
 			}
 
-			private string GenerateUniqueInstanceName()
+			public string GenerateUniqueInstanceNameSuffix()
 			{
 				// GenerateUniqueInstanceName is being called when generating a name for instance.
 				// After the call, the intance is added as a child to its parent.
@@ -1115,7 +1094,7 @@ namespace DatasmithRevitExporter
 
 				// To add uniqueness to the generated name, we construct a string with child counts from
 				// current parent instance, up to the root:
-				// Elem->Instance->Instance->Instace can produce something like: "1:5:3" for example.
+				// Elem->Instance->Instance->Instance can produce something like: "1:5:3" for example.
 				// However, this is not enough because elsewhere we might encounter the same sequence in terms of child counts,
 				// but adding the CurrentElement unique id ensures we get unique name string in the end.
 
@@ -1130,8 +1109,7 @@ namespace DatasmithRevitExporter
 				// Add child count for the root element (parent of all instances)
 				ChildCounts.AppendFormat(":{0}", ChildElements.Count);
 
-				FBaseElementData Instance = InstanceDataStack.Peek();
-				return $"{DocumentData.DocumentId}:{CurrentElement.UniqueId}:{Instance.BaseElementType.UniqueId}{ChildCounts.ToString()}";
+				return $"I:{ChildCounts}";
 			}
 
 			private string GetMeshName()
@@ -1140,20 +1118,16 @@ namespace DatasmithRevitExporter
 				{
 					return $"{DocumentData.DocumentId}:{CurrentElement.UniqueId}";
 				}
-				else
-				{
-					FBaseElementData Instance = InstanceDataStack.Peek();
 
-					if (!Instance.bAllowMeshInstancing)
-					{
-						return GenerateUniqueInstanceName();
-					}
-					else
-					{
-						// Generate instanced mesh name
-						return $"{DocumentData.DocumentId}:{Instance.BaseElementType.UniqueId}";
-					}
+				FBaseElementData Instance = InstanceDataStack.Peek();
+
+				if (Instance.bAllowMeshInstancing)
+				{
+					// Generate instanced mesh name using Instance geometry UniqueId
+					return $"{DocumentData.DocumentId}:{Instance.BaseElementType.UniqueId}";
 				}
+
+				return GetActorName(true);  // Use unique element's actor name to base mesh name on
 			}
 
 			private string GetActorLabel()
@@ -1250,6 +1224,26 @@ namespace DatasmithRevitExporter
 			}
 		}
 
+		private string GetActorName(FElementData InElementData)
+		{
+			// GetActorName should be called only for the current processed element
+			// so InElementData serves internally the only purpose to validate the call
+			Debug.Assert(InElementData == ElementDataStack.Peek());
+
+			string ActorName = Context.GetActorName(this);
+			if (InElementData.InstanceDataStack.Count == 0)
+			{
+				return ActorName;
+			}
+
+			
+			// Instance actor name (when InstanceDataStack is not empty), using current element's encountered instance count as instance's identification
+			// (to to OnElementBegin could follow multiple OnInstanceBegin/End calls, meaning an element can have multiple instances nested right under it)
+			// And each instance should have a separate unique name name identified by element path from root plus instances index(calling GetActorName for next
+			// instance in the same element will have InstanceDataStack.Count increased)
+			return ActorName + InElementData.GenerateUniqueInstanceNameSuffix();
+		}
+
 		public Dictionary<string, Tuple<FDatasmithFacadeMeshElement, Task<bool>>>
 														MeshMap = new Dictionary<string, Tuple<FDatasmithFacadeMeshElement, Task<bool>>>();
 		public Dictionary<ElementId, FBaseElementData>	ActorMap = new Dictionary<ElementId, FDocumentData.FBaseElementData>();
@@ -1271,6 +1265,8 @@ namespace DatasmithRevitExporter
 		private XYZ										ProjectSurveyPoint = null;
 		private XYZ										ProjectBasePoint = null;
 
+		public FDatasmithRevitExportContext				Context = null;
+
 		public string									DocumentId { get; private set; } = "";
 
 		public bool										bSkipMetadataExport { get; private set; } = false;
@@ -1285,11 +1281,10 @@ namespace DatasmithRevitExporter
 			FSettings InSettings,
 			ref List<string> InMessageList,
 			FDirectLink InDirectLink,
-			string InLinkedDocumentId
+			string InDocumentId
 		)
 		{
-			InsertionPoint = InSettings?.InsertionPoint ?? FSettings.EInsertionPoint.Default;
-
+			
 			CurrentSettings = InSettings;
 			DirectLink = InDirectLink;
 			CurrentDocument = InDocument;
@@ -1297,6 +1292,31 @@ namespace DatasmithRevitExporter
 			// With DirectLink, we delay export of metadata for a faster initial export.
 			bSkipMetadataExport = (DirectLink != null);
 
+
+			DocumentId = InDocumentId;
+		}
+
+		public void Reset(FDatasmithRevitExportContext InContext)
+		{
+			Context = InContext; // Make currently processing document aware of the current export context (currently context is recreated on each Sync)
+			MeshMap = new Dictionary<string, Tuple<FDatasmithFacadeMeshElement, Task<bool>>>();
+			ActorMap = new Dictionary<ElementId, FBaseElementData>();
+			MaterialDataMap = null;
+			NewMaterialsMap = new Dictionary<string, FMaterialData>();
+			DecalElementsMap = new Dictionary<ElementId, FElementData>();
+			DecalMaterialsMap = new Dictionary<ElementId, FDecalMaterial>();
+
+			ElementDataStack = new Stack<FElementData>();
+			CurrentMaterialName = null;
+
+			InsertionPoint = FSettings.EInsertionPoint.Default;
+
+			ProjectSurveyPoint = null;
+			ProjectBasePoint = null;
+
+			SectionBoxOutline = null;
+			SectionBox = null;
+
 			if (DirectLink != null)
 			{
 				MaterialDataMap = DirectLink.MaterialDataMap;
@@ -1306,6 +1326,8 @@ namespace DatasmithRevitExporter
 				MaterialDataMap = new Dictionary<string, FMaterialData>();
 			}
 
+			InsertionPoint = CurrentSettings?.InsertionPoint ?? FSettings.EInsertionPoint.Default;
+
 			// Cache document section boxes
 			if (CurrentDocument.ActiveView != null)
 			{
@@ -1337,11 +1359,6 @@ namespace DatasmithRevitExporter
 					}
 				}
 			}
-
-			if (InLinkedDocumentId != null)
-			{
-				DocumentId = InLinkedDocumentId;
-			}
 		}
 
 		public Element GetElement(
@@ -1443,8 +1460,6 @@ namespace DatasmithRevitExporter
 
 						ElementData.BaseElementType = InElement.Document.GetElement(InElement.GetTypeId()) as ElementType;
 						ElementData.CurrentElement = InElement;
-						ElementData.InitializePivotPlacement(ref InWorldTransform);
-						ElementData.InitializeElement(InWorldTransform, ElementData);
 						ElementData.MeshMaterialsMap.Clear();
 
 						if (ElementData.bIsDecalElement)
@@ -1470,9 +1485,19 @@ namespace DatasmithRevitExporter
 						DecalElementsMap.Add(InElement.Id, ElementData);
 					}
 				}
+
+				ElementDataStack.Push(ElementData);
+
+				// Initialize element after pushing it on the stack(to unify this with other calls to element's methods)
+				// GetActorName depends on this 
+				ElementData.InitializePivotPlacement(ref InWorldTransform);
+				ElementData.InitializeElement(InWorldTransform, ElementData);
+			}
+			else
+			{
+				ElementDataStack.Push(ElementData);
 			}
 
-			ElementDataStack.Push(ElementData);
 			ElementDataStack.Peek().ElementActor.AddTag("IsElement");
 
 			return true;
@@ -2617,5 +2642,14 @@ namespace DatasmithRevitExporter
 			// Set the world transform of the Datasmith actor.
 			IOActor.SetWorldTransform(worldMatrix);
 		}
+
+		/// <summary>
+		/// Combine whole element stack in the document hierarchy to build unique path to the element in the document
+		/// </summary>
+		public string GetElementStackName()
+		{
+			return string.Join(", ", ElementDataStack.Select(Data => $"{Data.CurrentElement.UniqueId}"));
+		}
+
 	}
 }
diff --git a/Engine/Source/Programs/Enterprise/Datasmith/DatasmithRevitExporter/Private/DatasmithRevitExportContext.cs b/Engine/Source/Programs/Enterprise/Datasmith/DatasmithRevitExporter/Private/DatasmithRevitExportContext.cs
--- a/Engine/Source/Programs/Enterprise/Datasmith/DatasmithRevitExporter/Private/DatasmithRevitExportContext.cs
+++ b/Engine/Source/Programs/Enterprise/Datasmith/DatasmithRevitExporter/Private/DatasmithRevitExportContext.cs
@@ -4,6 +4,7 @@ using System;
 using System.Collections.Generic;
 using System.Diagnostics;
 using System.IO;
+using System.Linq;
 using System.Threading.Tasks;
 using Autodesk.Revit.ApplicationServices;
 using Autodesk.Revit.DB;
@@ -48,7 +49,7 @@ namespace DatasmithRevitExporter
 		private Stack<FDocumentData> DocumentDataStack = new Stack<FDocumentData>();
 
 		// Cache exported document data for future reference (including the parenting (linked documents))
-		private List<FDocumentData> ExportedDocuments = new List<FDocumentData>();
+		private Dictionary<string, FDocumentData> ExportedDocuments = null;
 
 		// List of extra search paths for Revit texture files.
 		private IList<string> ExtraTexturePaths = new List<string>();
@@ -158,10 +159,12 @@ namespace DatasmithRevitExporter
 			{
 				DirectLink.OnBeginExport();
 				DatasmithScene = DirectLink.DatasmithScene;
+				ExportedDocuments = DirectLink.ExportedDocuments;
 			}
 			else
 			{
 				DatasmithScene = new FDatasmithFacadeScene(HOST_NAME, VENDOR_NAME, PRODUCT_NAME, ProductVersion);
+				ExportedDocuments = new Dictionary<string, FDocumentData>();
 			}
 
 			DatasmithScene.PreExport();
@@ -220,11 +223,97 @@ namespace DatasmithRevitExporter
 			WorldTransformStack.Pop();
 		}
 
+		private int DebugIndent = 0;
+
+		/// <summary>
+		/// Log debug for plugin developer - no runtime toggle(enabled by adding DatasmithRevitDebugOutput conditional compilation symbol)
+		/// </summary>
+		[Conditional("DatasmithRevitDebugOutput")]
+		public void LogDebug(string Message)  
+		{
+			MessageList.Add(String.Concat(Enumerable.Repeat("    ", DebugIndent)) + Message);
+		}
+
+		[Conditional("DatasmithRevitDebugOutput")]
+		private void LogDebugElementBegin(ElementId InElementId)
+		{
+			LogDebug($"OnElementBegin({InElementId})");
+			Element CurrentElement = GetElement(InElementId);
+			if (CurrentElement != null)
+			{
+				LogDebug($" # '{CurrentElement.Name}'({CurrentElement.GetType()}), UID: {CurrentElement.UniqueId}");
+				if (CurrentElement is FamilyInstance Instance)
+				{
+					LogDebug($" # Instance: '{Instance.Name}'");
+					FamilySymbol Symbol = Instance.Symbol;
+					LogDebug($" #  FamilySymbol: '{Symbol.Name}',  Id: {Symbol.Id}, UID: {Symbol.UniqueId}");
+					Family SymbolFamily = Symbol.Family;
+					LogDebug($" #   Family: '{SymbolFamily.Name}',  Id: {SymbolFamily.Id}, UID {SymbolFamily.UniqueId} ");
+				}
+			}
+			else
+			{
+				LogDebug(" # <Unknown>");
+			}
+
+			++DebugIndent;
+		}
+
+		[Conditional("DatasmithRevitDebugOutput")]
+		private void LogDebugElementEnd(ElementId InElementId)
+		{
+			--DebugIndent;
+
+			LogDebug($"OnElementEnd({InElementId})");
+		}
+
+		[Conditional("DatasmithRevitDebugOutput")]
+		private void LogDebugInstanceBegin(InstanceNode InInstanceNode)
+		{
+			LogDebug($"OnInstanceBegin({InInstanceNode.NodeName})");
+
+#if REVIT_API_2023
+			SymbolGeometryId GeometryId = InInstanceNode.GetSymbolGeometryId();
+			LogDebug($" #  GeometryUniqueId: '{GeometryId.AsUniqueIdentifier()}'");
+			Element CurrentFamilySymbol = GetElement(GeometryId.SymbolId);
+#else
+			Element CurrentFamilySymbol = GetElement(InInstanceNode.GetSymbolId());
+#endif
+
+			if (CurrentFamilySymbol != null)
+			{
+				LogDebug($" #  Symbol: '{CurrentFamilySymbol.GetType()}', '{CurrentFamilySymbol.Name}',  Id: {CurrentFamilySymbol.Id}, UID: {CurrentFamilySymbol.UniqueId}");
+
+				switch (CurrentFamilySymbol)
+				{
+					case FamilySymbol Symbol:
+					{
+						Family SymbolFamily = Symbol.Family;
+						LogDebug(
+							$" #   Family: '{SymbolFamily.Name}',  Id: {SymbolFamily.Id}, UID {SymbolFamily.UniqueId} ");
+						break;
+					}
+				}
+			}
+
+			++DebugIndent;
+		}
+
+		[Conditional("DatasmithRevitDebugOutput")]
+		private void LogDebugInstanceEnd(InstanceNode InInstanceNode)
+		{
+			--DebugIndent;
+
+			LogDebug($"OnInstanceEnd({InInstanceNode.NodeName})");
+		}
+
 		// OnElementBegin marks the beginning of an element to be exported.
 		public RenderNodeAction OnElementBegin(
 			ElementId InElementId // exported element ID
 		)
 		{
+			LogDebugElementBegin(InElementId);
+
 			CurrentElementSkipped = true;
 
 			Element CurrentElement = GetElement(InElementId);
@@ -251,8 +340,11 @@ namespace DatasmithRevitExporter
 			ElementId InElementId // exported element ID
 		)
 		{
+			LogDebugElementEnd(InElementId);
+
 			if (GetElement(InElementId) != null)
 			{
+				Debug.Assert(!GetElement(InElementId).IsTransient); // OnElementBegin checks for IsTransient
 				// Forget the current element being exported.
 				PopElement("Element End");
 			}
@@ -263,6 +355,8 @@ namespace DatasmithRevitExporter
 			InstanceNode InInstanceNode // family instance output node
 		)
 		{
+			LogDebugInstanceBegin(InInstanceNode);
+
 #if REVIT_API_2023
 			Element CurrentFamilySymbol = GetElement(InInstanceNode.GetSymbolGeometryId().SymbolId);
 #else
@@ -298,6 +392,8 @@ namespace DatasmithRevitExporter
 			InstanceNode InInstanceNode // family instance output node
 		)
 		{
+			LogDebugInstanceEnd(InInstanceNode);
+
 #if REVIT_API_2023
 			Element CurrentFamilySymbol = GetElement(InInstanceNode.GetSymbolGeometryId().SymbolId);
 #else
@@ -620,18 +716,30 @@ namespace DatasmithRevitExporter
 				DirectLink.OnBeginLinkedDocument(LinkedDocElement);
 			}
 
-			// Check if we have cache for this document.
-			FDocumentData DocumentData = new FDocumentData(InRevitDocument, DocumentSettings, ref MessageList, DirectLink, InLinkedDocumentId);
+			bool bIsTopDocument = InLinkedDocumentId == null;
+			string DocumentId = InLinkedDocumentId ?? "";
+
+			// Check if we have cache for this document so we don't recreate it each Sync
+			// Also ElementData objects are referencing their owning FDocumentData 
+			if (!ExportedDocuments.TryGetValue(DocumentId, out var DocumentData))  
+			{
+				DocumentData = new FDocumentData(InRevitDocument, DocumentSettings, ref MessageList,
+					DirectLink, DocumentId);
+				ExportedDocuments.Add(DocumentId, DocumentData);
+			}
+
+			// Reset DocumentData for new update
+			// todo: Current implementation expects DocumentData to be clean but what is possible is to track created assets and to reuse them when their element changes
+			DocumentData.Reset(this);
 
 			DocumentDataStack.Push(DocumentData);
 
-			if (InLinkedDocumentId == null)
+			if (bIsTopDocument)
 			{
 				// Top level document
 				DocumentDataStack.Peek().AddLocationActors(WorldTransformStack.Peek());
 			}
 
-			ExportedDocuments.Add(DocumentData);
 		}
 
 		private FDocumentData PopDocument()
@@ -837,5 +945,17 @@ namespace DatasmithRevitExporter
 				DirectLink?.CacheElement(RevitDocument, InView3D, new FDocumentData.FBaseElementData(CameraActor, null, DocumentDataStack.Peek()));
 			}
 		}
+
+		public string GetActorName(FDocumentData InDocumentData)
+		{
+			// Expecting this to be called only for current document on the stack
+			Debug.Assert(DocumentDataStack.Peek() == InDocumentData);
+
+			// Use Element "path" in the whole linked documents hierarchy to build unique name for any kind of instance element(e.g. instance of the same family of the same linked file)
+			string ActorName = string.Join(":", 
+				                   DocumentDataStack.Select(DocumentData => $"{DocumentData.DocumentId}({DocumentData.GetElementStackName()})"));
+			LogDebug($"GetActorName:'{ActorName}'");
+			return ActorName;
+		}
 	}
 }
