[3ds Max Exporter] Fixed Forest Pack meshes are labelled incorrectly in exported Datasmith Scene

diff --git a/Engine/Source/Programs/Enterprise/Datasmith/DatasmithMaxExporter/Private/DatasmithMaxExporterWithDirectLink.cpp b/Engine/Source/Programs/Enterprise/Datasmith/DatasmithMaxExporter/Private/DatasmithMaxExporterWithDirectLink.cpp
--- a/Engine/Source/Programs/Enterprise/Datasmith/DatasmithMaxExporter/Private/DatasmithMaxExporterWithDirectLink.cpp
+++ b/Engine/Source/Programs/Enterprise/Datasmith/DatasmithMaxExporter/Private/DatasmithMaxExporterWithDirectLink.cpp
@@ -2980,11 +2980,8 @@ public:
 
 	virtual void SetupDatasmithHISMForNode(FNodeTracker& NodeTracker, FMeshConverterSource& MeshSource, Mtl* Material, int32 MeshIndex, const TArray<Matrix3>& Transforms) override
 	{
-		FString MeshName = FString::FromInt(NodeTracker.Node->GetHandle()) + TEXT("_") + FString::FromInt(MeshIndex);
-
-		// note: when export Mesh goes to other place due to parallellizing it's result would be unknown here so MeshIndex handling will change(i.e. increment for any mesh)
-
-		MeshSource.MeshName = MeshName; // todo: !!!
+		MeshSource.MeshName = FString::FromInt(NodeTracker.Node->GetHandle()) + TEXT("_") + FString::FromInt(MeshIndex);
+		// note: when export Mesh goes to other place due to parallelizing it's result would be unknown here so MeshIndex handling will change(i.e. increment for any mesh)
 
 		FMeshConverted MeshConvertedDummy; // todo: possible reuse of previous converted mesh
 		FMeshes::AddMesh(*this, MeshSource, MeshConvertedDummy, [&](bool bHasConverted, FMeshConverted& MeshConverted)
@@ -2998,7 +2995,7 @@ public:
 				RegisterNodeForMaterial(NodeTracker, Material);
 				MaterialsCollectionTracker.SetMaterialsForMeshElement(MeshConverted.DatasmithMeshElement, Material, MeshConverted.SupportedChannels);
 
-				FString MeshLabel = NodeTrackersNames.GetNodeName(NodeTracker) + (TEXT("_") + FString::FromInt(MeshIndex));
+				FString MeshLabel = FString(MeshSource.Node->GetName()) + (TEXT("_") + FString::FromInt(MeshIndex));
 				MeshConverted.DatasmithMeshElement->SetLabel(*MeshLabel);
 
 				FDatasmithConverter Converter;
@@ -3007,7 +3004,7 @@ public:
 				TSharedPtr< IDatasmithActorElement > InversedHISMActor;
 				// todo: ExportHierarchicalInstanceStaticMeshActor CustomMeshNode only used for Material - can be simplified, Material anyway is dealt with outside too
 				TSharedRef<IDatasmithActorElement> HismActorElement = ExportHierarchicalInstanceStaticMeshActor(NodeTracker.Node, MeshSource.Node, *MeshLabel, MeshConverted.SupportedChannels,
-					Material, &Transforms, *MeshName, Converter.UnitToCentimeter, EStaticMeshExportMode::Default, InversedHISMActor);
+					Material, &Transforms, *MeshSource.MeshName, Converter.UnitToCentimeter, EStaticMeshExportMode::Default, InversedHISMActor);
 				NodeTracker.GetConverted().DatasmithActorElement->AddChild(HismActorElement, EDatasmithActorAttachmentRule::KeepWorldTransform);
 				if (InversedHISMActor)
 				{
