[Solidworks] fixed animation euler angles translation to Unreal

- Inverted X and Z rotations as conversion of handedness affected these(Y rotation not affected because that coordinate is negated itself)

diff --git a/Engine/Source/Programs/Enterprise/Datasmith/DatasmithSolidworksExporter/Private/DatasmithExporter.cs b/Engine/Source/Programs/Enterprise/Datasmith/DatasmithSolidworksExporter/Private/DatasmithExporter.cs
--- a/Engine/Source/Programs/Enterprise/Datasmith/DatasmithSolidworksExporter/Private/DatasmithExporter.cs
+++ b/Engine/Source/Programs/Enterprise/Datasmith/DatasmithSolidworksExporter/Private/DatasmithExporter.cs
@@ -902,7 +902,7 @@ namespace DatasmithSolidworks
 
 					FVec3 Translation = new FVec3(LocalMatrix[12], LocalMatrix[13], LocalMatrix[14]);
 
-					Anim.AddFrame(EDatasmithFacadeAnimationTransformType.Rotation, Keyframe.Step, Euler.X, - Euler.Y, - Euler.Z);
+					Anim.AddFrame(EDatasmithFacadeAnimationTransformType.Rotation, Keyframe.Step, -Euler.X, Euler.Y, -Euler.Z);
 					Anim.AddFrame(EDatasmithFacadeAnimationTransformType.Scale, Keyframe.Step, Scale, Scale, Scale);
 					Anim.AddFrame(EDatasmithFacadeAnimationTransformType.Translation, Keyframe.Step, Translation.X, -Translation.Y, Translation.Z);
 				}
