[Solidworks] Reworked plugin to avoid modyfing suppression state of components

- all configuration's meshes and materials collected and exported when the configuration is activated

- configurations are activated only once now(resulting ~x2 export speedup as configuration switch is among slowest operation)

- no component is being unsuppressed(and then suppressed) means another time saver

- components exported only when they are unsuppressed in any configuration, meaning no export of unused components.

- Elevation_test is from 78 seconds export down to 11(measured in Debug but anyways switching state configs/suppression was the main time waste)

diff --git a/Engine/Source/Programs/Enterprise/Datasmith/DatasmithSolidworksExporter/Private/AssemblyDocument.cs b/Engine/Source/Programs/Enterprise/Datasmith/DatasmithSolidworksExporter/Private/AssemblyDocument.cs
--- a/Engine/Source/Programs/Enterprise/Datasmith/DatasmithSolidworksExporter/Private/AssemblyDocument.cs
+++ b/Engine/Source/Programs/Enterprise/Datasmith/DatasmithSolidworksExporter/Private/AssemblyDocument.cs
@@ -5,6 +5,7 @@ using SolidWorks.Interop.swconst;
 using System;
 using System.Collections.Concurrent;
 using System.Collections.Generic;
+using System.Diagnostics;
 using System.Linq;
 using System.Threading.Tasks;
 
@@ -53,16 +54,50 @@ namespace DatasmithSolidworks
 		{
 		}
 
-		public override void ExportToDatasmithScene(FMeshes meshes)
+		// Export to datasmith scene extracted configurations data
+		public override void ExportToDatasmithScene(FConfigurationExporter ConfigurationExporter)
 		{
 			FSyncState OldSyncState = SyncState;
 
 			if (bFileExportInProgress)
 			{
+				// todo: separate DL from file export more explicitly
 				SyncState = new FSyncState();
 			}
 
 			SetExportStatus("Actors");
+
+			ProcessComponentsPendingDelete();
+
+			Configuration CurrentConfig = SwDoc.GetActiveConfiguration() as Configuration;
+
+			// Configurations combined tree should have single child(root component of each config is the same)
+			Debug.Assert(ConfigurationExporter.CombinedTree.Children.Count == 1);
+			ExportComponentRecursive(CurrentConfig.Name, ConfigurationExporter, ConfigurationExporter.CombinedTree.Children[0], null);
+			SyncState.DirtyComponents.Clear();
+
+			SyncState = OldSyncState;
+
+			// todo: check animation export for DL
+			// Export animations (only allow when exporting to file)
+			// if (bFileExportInProgress)
+			{
+				Component2 Root = CurrentConfig.GetRootComponent3(true);
+				List<FAnimation> Animations = FAnimationExtractor.ExtractAnimations(SwAsmDoc, Root);
+				if (Animations != null)
+				{
+					SetExportStatus($"Animations");
+
+					foreach (FAnimation Anim in Animations)
+					{
+						Exporter.ExportAnimation(Anim);
+					}
+				}
+			}
+		}
+
+		private void ProcessComponentsPendingDelete()
+		{
 			foreach (FComponentName CompName in SyncState.ComponentsToDelete)
 			{
 				FActorName ActorName = CompName.GetActorName();
@@ -79,6 +114,44 @@ namespace DatasmithSolidworks
 				SyncState.ComponentsTransformsMap.Remove(CompName);
 				Exporter.RemoveActor(ActorName);
 			}
+			SyncState.ComponentsToDelete.Clear();
+		}
+
+		public override ConcurrentDictionary<FComponentName, FObjectMaterials> LoadDocumentMaterials(HashSet<FComponentName> ComponentNamesToExportSet)
+		{
+			return FObjectMaterials.LoadAssemblyMaterials(this, ComponentNamesToExportSet, swDisplayStateOpts_e.swThisDisplayState, null);
+		}
+
+		public override void AddComponentMaterials(FComponentName ComponentName, FObjectMaterials Materials)
+		{
+			Addin.Instance.LogDebug($"AddComponentMaterials: {ComponentName} - {Materials}");
+			SyncState.ComponentsMaterialsMap[ComponentName] = Materials;
+		}
+
+		public override void AddMeshForComponent(FComponentName ComponentName, string MeshName)
+		{
+			SyncState.ComponentNameToMeshNameMap[ComponentName] = MeshName;
+		}
+
+		public override FObjectMaterials GetComponentMaterials(Component2 Comp)
+		{
+			FObjectMaterials ComponentMaterials = null;
+			SyncState.ComponentsMaterialsMap?.TryGetValue(new FComponentName(Comp), out ComponentMaterials);
+			return ComponentMaterials;
+		}
+
+		// Export to datasmith scene current configuration only
+		public override void ExportToDatasmithScene(FMeshes meshes)
+		{
+			FSyncState OldSyncState = SyncState;
+
+			if (bFileExportInProgress)
+			{
+				SyncState = new FSyncState();
+			}
+
+			SetExportStatus("Actors");
+			ProcessComponentsPendingDelete();
 
 			Configuration CurrentConfig = SwDoc.GetActiveConfiguration() as Configuration;
 
@@ -89,10 +162,7 @@ namespace DatasmithSolidworks
 			// actor hierarchy has been exported
 			Dictionary<Component2, FActorName> MeshesToExportMap = new Dictionary<Component2, FActorName>();
 		
-			// Store original component suppression state to revert it after export
-			// Component needs to be fully resolved to have its data loaded and even suppressed components are needed for configurations export(in case they are visible in other configs)
-			List<Tuple<Component2, int>> ComponentsSuppression = new List<Tuple<Component2, int>>();
-			ExportComponentRecursive(Root, null, MeshesToExportMap, ComponentsSuppression);
+			ExportComponentRecursive(Root, null, MeshesToExportMap);
 
 			// Export materials
 			SetExportStatus($"Component Materials");
@@ -148,8 +218,6 @@ namespace DatasmithSolidworks
 				string MeshName = Exporter.AddMesh(KVP.Value);
 				SyncState.ComponentNameToMeshNameMap[new FComponentName(KVP.Key)] = MeshName;
 			}
-
-			SyncState.ComponentsToDelete.Clear();
 			SyncState.DirtyComponents.Clear();
 
 			SyncState = OldSyncState;
@@ -168,11 +236,6 @@ namespace DatasmithSolidworks
 					}
 				}
 			}
-
-			foreach(var ComponentSuppression in ComponentsSuppression)
-			{
-				ComponentSuppression.Item1.SetSuppression2(ComponentSuppression.Item2);
-			}
 		}
 
 		public override bool HasMaterialUpdates()
@@ -276,19 +339,22 @@ namespace DatasmithSolidworks
 			}
 			else
 			{
-				DatasmithTransform = new float[]
-				{
-					1.0f, 0.0f, 0.0f, 0.0f,
-					0.0f, 1.0f, 0.0f, 0.0f,
-					0.0f, 0.0f, 1.0f, 0.0f,
-					0.0f, 0.0f, 0.0f, 1.0f
-				};
+				DatasmithTransform = MathUtils.TransformIdentity();
 			}
 
 			return DatasmithTransform;
 		}
 
-		private void ExportComponentRecursive(Component2 InComponent, Component2 InParent, Dictionary<Component2, FActorName> OutMeshesToExportMap, List<Tuple<Component2, int>> OutComponentSuppression)
+		/// <summary>
+		/// Get component transform in specified configuration
+		/// </summary>
+		private float[] GetComponentDatasmithTransform(FConfigurationTree.FComponentTreeNode InNode,
+			FConfigurationTree.FComponentConfig ComponentConfig)
+		{
+			return ComponentConfig?.Transform ?? InNode.CommonConfig.Transform ?? MathUtils.TransformIdentity();
+		}
+
+		private void ExportComponentRecursive(Component2 InComponent, Component2 InParent, Dictionary<Component2, FActorName> OutMeshesToExportMap)
 		{
 			bool bHasDirtyTransform = false;
 
@@ -323,12 +389,6 @@ namespace DatasmithSolidworks
 
 				// When exporting configurations we need to load all components to make their data available for export
 				// Suppressed or lightweight components don't have this
-				int Suppression = InComponent.GetSuppression();
-				if (bHasConfigurations && (Suppression !=  (int)swComponentSuppressionState_e.swComponentFullyResolved))
-				{
-					InComponent.SetSuppression2((int)swComponentSuppressionState_e.swComponentFullyResolved);
-					OutComponentSuppression.Add(new Tuple<Component2, int>(InComponent, Suppression));
-				}
 
 				if (!InComponent.IsSuppressed())
 				{
@@ -351,7 +411,8 @@ namespace DatasmithSolidworks
 
 				// ComponentDoc is null if component is suppressed or lightweight
 				ModelDoc2 ComponentDoc = InComponent.GetModelDoc2();
-				if (ComponentDoc is PartDoc) 
+				
+				if (ComponentDoc is PartDoc) // Also possible to test document type with swDocumentTypes_e.swDocPART
 				{
 					bool bFirstExport = !SyncState.ExportedComponentsMap.ContainsKey(ComponentName);
 					bNeedsGeometryExport = bFirstExport;
@@ -407,11 +468,121 @@ namespace DatasmithSolidworks
 				foreach (object Obj in Children)
 				{
 					Component2 Child = (Component2)Obj;
-					ExportComponentRecursive(Child, InComponent, OutMeshesToExportMap, OutComponentSuppression);
+					ExportComponentRecursive(Child, InComponent, OutMeshesToExportMap);
 				}
 			}
 		}
 
+		private void ExportComponentRecursive(string ActiveConfigName,
+			FConfigurationExporter ConfigurationExporter,
+			FConfigurationTree.FComponentTreeNode InNode, FConfigurationTree.FComponentTreeNode InParent)
+		{
+
+			FConfigurationTree.FComponentConfig ActiveComponentConfig = InNode.Configurations?.Find(Config => Config.ConfigName == ActiveConfigName);
+
+			if (NeedExportComponent(InNode, ActiveComponentConfig))
+			{
+				SetExportStatus(InNode.ComponentName.GetString());
+
+				FDatasmithActorExportInfo ActorExportInfo = new FDatasmithActorExportInfo();
+
+				FActorName ActorName = InNode.ActorName;
+
+				ActorExportInfo.Label = InNode.ComponentName.GetLabel();
+				ActorExportInfo.Name = ActorName;
+
+				if (InParent != null)
+				{
+					ActorExportInfo.ParentName = InParent.GetActorName();
+				}
+				
+				ActorExportInfo.bVisible = true;
+				ActorExportInfo.Type = Exporter.GetExportedActorType(ActorName) ?? EActorType.SimpleActor;
+				ActorExportInfo.Transform = GetComponentDatasmithTransform(InNode, ActiveComponentConfig);
+
+				SyncState.ComponentsTransformsMap[InNode.ComponentName] = ActorExportInfo.Transform;
+
+				ActorExportInfo.bVisible = ActiveComponentConfig?.bVisible ?? InNode.CommonConfig.bVisible;
+
+				if (InNode.IsPartComponent())
+				{
+					// This component has associated part document -- treat is as a mesh actor
+					ActorExportInfo.Type = EActorType.MeshActor;
+				}
+
+				if (ActorExportInfo.Type == EActorType.MeshActor )
+				{
+					// todo: deduplicate mesh name calculation
+					ActorExportInfo.MeshName = FDatasmithExporter.SanitizeName($"{ActorExportInfo.Name}_Mesh");
+				}
+
+				Exporter.ExportOrUpdateActor(ActorExportInfo);
+
+				SyncState.CleanComponents.Add(InNode.ComponentName);
+			}
+
+			// AddExportedComponent(InNode);
+
+			// Export component children
+			foreach (FConfigurationTree.FComponentTreeNode Child in InNode.EnumChildren())
+			{
+				ExportComponentRecursive(ActiveConfigName, ConfigurationExporter, Child, InNode);
+			}
+		}
+
+		public override void AddExportedComponent(FConfigurationTree.FComponentTreeNode InNode)
+		{
+			SyncState.ExportedComponentsMap[InNode.ComponentName] = InNode.Component;
+		}
+
+		public override bool NeedExportComponent(FConfigurationTree.FComponentTreeNode InComponent, FConfigurationTree.FComponentConfig ActiveComponentConfig)
+		{
+			bool bHasDirtyTransform = false;
+			if (SyncState.ComponentsTransformsMap.ContainsKey(InComponent.ComponentName))
+			{
+				float[] ComponentTm = GetComponentDatasmithTransform(InComponent, ActiveComponentConfig);
+				bHasDirtyTransform =
+					!MathUtils.TransformsAreEqual(SyncState.ComponentsTransformsMap[InComponent.ComponentName], ComponentTm);
+			}
+
+			bool bNeedExportComponent = bHasDirtyTransform || !SyncState.CleanComponents.Contains(InComponent.ComponentName);
+			return bNeedExportComponent;
+		}
+
+		public override bool NeedGeometryExport(FConfigurationTree.FComponentTreeNode InNode, FConfigurationTree.FComponentConfig ActiveComponentConfig)
+		{
+			bool bNeedsGeometryExport;
+			bool bFirstExport = !SyncState.ExportedComponentsMap.ContainsKey(InNode.ComponentName);
+			bNeedsGeometryExport = bFirstExport;
+
+			if (!bFirstExport && SyncState.DirtyComponents.ContainsKey(InNode.ComponentName))
+			{
+				uint DirtyState = SyncState.DirtyComponents[InNode.ComponentName];
+				bNeedsGeometryExport =
+					((DirtyState & (1u << (int)EComponentDirtyState.Material)) != 0) ||
+					((DirtyState & (1u << (int)EComponentDirtyState.Geometry)) != 0) ||
+					((DirtyState & (1u << (int)EComponentDirtyState.Delete)) != 0);
+			}
+
+			return bNeedsGeometryExport;
+		}
+
+		public override void AddPartDocument(FConfigurationTree.FComponentTreeNode InNode)
+		{
+			//TODO this will be null for new part, think of more solid solution
+			if (!SyncState.PartsMap.ContainsKey(InNode.PartPath))
+			{
+				ModelDoc2 ComponentDoc = InNode.Component.GetModelDoc2();
+				// New part
+				int PartDocId = Addin.Instance.GetDocumentId(ComponentDoc as ModelDoc2);
+				SyncState.PartsMap[InNode.PartPath] = new FPartDocument(PartDocId, ComponentDoc as PartDoc, Exporter, this,
+					InNode.ComponentName);
+				SyncState.PartsMap[InNode.PartPath].Init();
+			}
+
+			SyncState.ComponentToPartMap[InNode.ComponentName] = InNode.PartPath;
+		}
+
 		public void SetComponentDirty(FComponentName InComponent, EComponentDirtyState InState)
 		{
 			if (SyncState.CleanComponents.Contains(InComponent))
diff --git a/Engine/Source/Programs/Enterprise/Datasmith/DatasmithSolidworksExporter/Private/ConfigurationExporter.cs b/Engine/Source/Programs/Enterprise/Datasmith/DatasmithSolidworksExporter/Private/ConfigurationExporter.cs
--- a/Engine/Source/Programs/Enterprise/Datasmith/DatasmithSolidworksExporter/Private/ConfigurationExporter.cs
+++ b/Engine/Source/Programs/Enterprise/Datasmith/DatasmithSolidworksExporter/Private/ConfigurationExporter.cs
@@ -15,7 +15,7 @@ namespace DatasmithSolidworks
 	{
 		public string Name;
 		public bool bIsDisplayStateConfiguration = false;
-		public Dictionary<FComponentName, float[]> ComponentTransform = new Dictionary<FComponentName, float[]>();
+		public Dictionary<FComponentName, float[]> ComponentTransform = new Dictionary<FComponentName, float[]>();  // Relative(to parent) transform
 		public Dictionary<FComponentName, bool> ComponentVisibility = new Dictionary<FComponentName, bool>();
 		public Dictionary<FComponentName, FObjectMaterials> ComponentMaterials = new Dictionary<FComponentName, FObjectMaterials>();
 
@@ -187,24 +187,28 @@ namespace DatasmithSolidworks
 	[ComVisible(false)]
 	public class FConfigurationExporter
 	{
-		private readonly FMeshes Meshes;
+		public readonly FMeshes Meshes;
+		public readonly string[] ConfigurationNames;
 
-		public FConfigurationExporter(FMeshes Meshes)
+		public FConfigurationTree.FComponentTreeNode CombinedTree;
+
+		public FConfigurationExporter(FMeshes InMeshes, string[] InConfigurationNames)
 		{
-			this.Meshes = Meshes;
+			Meshes = InMeshes;
+			ConfigurationNames = InConfigurationNames;
 		}
 
 		public List<FConfigurationData> ExportConfigurations(FDocument InDoc)
 		{
-			string[] CfgNames = InDoc.SwDoc?.GetConfigurationNames();
-
-			if (CfgNames == null || CfgNames.Length <= 1)
+			// todo: may refactor to union behavior for 'no configuration'(i.e. active) and one-to-many configurations export
+			// Not sure how CfgNames might be technically null though
+			if (ConfigurationNames == null) // || CfgNames.Length <= 1)
 			{
 				return null;
 			}
 
-			FConfigurationTree.FComponentTreeNode CombinedTree = new FConfigurationTree.FComponentTreeNode();
-			CombinedTree.ComponentName = FComponentName.FromCustomString("CombinedTree");
+			CombinedTree = new FConfigurationTree.FComponentTreeNode(null);
+			CombinedTree.ComponentInfo.ComponentName = FComponentName.FromCustomString("CombinedTree");
 
 			// Ensure recursion will not stop on the root node (this may happen if it is not explicitly marked as visible)
 			CombinedTree.bVisibilitySame = true;
@@ -215,7 +219,7 @@ namespace DatasmithSolidworks
 
 			List<string> ExportedConfigurationNames = new List<string>();
 
-			foreach (string CfgName in CfgNames)
+			foreach (string CfgName in ConfigurationNames)
 			{
 				IConfiguration swConfiguration = InDoc.SwDoc.GetConfigurationByName(CfgName) as IConfiguration;
 
@@ -233,14 +237,16 @@ namespace DatasmithSolidworks
 				}
 
 				FConfigurationTree.FComponentTreeNode ConfigNode = new FConfigurationTree.FComponentTreeNode();
+				ConfigNode.ComponentInfo.ComponentName = FComponentName.FromCustomString(ConfigName);
+
 				ConfigNode.Children = new List<FConfigurationTree.FComponentTreeNode>();
-				ConfigNode.ComponentName = FComponentName.FromCustomString(ConfigName);
 
+				Dictionary<Component2, FActorName> MeshesToExportMap = new Dictionary<Component2, FActorName>();
 
 				// Build the tree and get default materials (which aren't affected by any configuration)
 				// Use GetRootComponent3() with Resolve = true to ensure suppressed components will be loaded
 				// todo: docs says that Part document SW returns null. Not the case. But probably better to add a guard
-				CollectComponentsRecursive(InDoc, swConfiguration.GetRootComponent3(true), ConfigNode);
+				CollectComponentsRecursive(InDoc, swConfiguration.GetRootComponent3(true), ConfigNode, MeshesToExportMap);
 
 				ExportedConfigurationNames.Add(ConfigName);
 
@@ -260,6 +266,61 @@ namespace DatasmithSolidworks
 					}
 				}
 
+				// Export materials
+				InDoc.SetExportStatus($"Component Materials");
+				HashSet<FComponentName> ComponentNamesToExportSet = new HashSet<FComponentName>();
+				foreach (var KVP in MeshesToExportMap)
+				{
+					FComponentName ComponentName = new FComponentName(KVP.Key);
+					if (!ComponentNamesToExportSet.Contains(ComponentName))
+					{
+						ComponentNamesToExportSet.Add(ComponentName);
+					}
+				}
+
+				ConcurrentDictionary<FComponentName, FObjectMaterials> ModifiedComponentsMaterials = InDoc.LoadDocumentMaterials(ComponentNamesToExportSet);
+
+				if (ModifiedComponentsMaterials != null)
+				{
+					foreach (var MatKVP in ModifiedComponentsMaterials)
+					{
+						InDoc.AddComponentMaterials(MatKVP.Key, MatKVP.Value);
+					}
+				}
+				
+				// Export materials before exporting meshes(currently mesh export code is tied to FDatasmithExporter.ExportedMaterialsMap
+				// todo: separate material export and mesh export(see comment above)?
+				InDoc.Exporter.ExportMaterials(InDoc.ExportedMaterialsMap);
+
+				// Export meshes
+				InDoc.SetExportStatus($"Component Meshes");
+				ConcurrentDictionary<Component2, Tuple<FDatasmithFacadeMeshElement, FDatasmithFacadeMesh>> CreatedMeshes = new ConcurrentDictionary<Component2, Tuple<FDatasmithFacadeMeshElement, FDatasmithFacadeMesh>>();
+
+				foreach (KeyValuePair<Component2, FActorName> KVP in MeshesToExportMap)
+				{
+					Component2 Comp = KVP.Key;
+
+					ConcurrentBag<FBody> Bodies = FBody.FetchBodies(Comp);
+					FMeshData MeshData = FStripGeometry.CreateMeshData(Bodies, InDoc.GetComponentMaterials(Comp));
+
+					if (MeshData != null)
+					{
+						Tuple<FDatasmithFacadeMeshElement, FDatasmithFacadeMesh> NewMesh = null;
+						if (InDoc.Exporter.ExportMesh($"{KVP.Value}_Mesh", MeshData, KVP.Value, out NewMesh))
+						{
+							CreatedMeshes[Comp] = NewMesh;
+						}
+					}
+				}
+
+				foreach (var KVP in CreatedMeshes)
+				{
+					FComponentName ComponentName = new FComponentName(KVP.Key);
+					string MeshName = InDoc.Exporter.AddMesh(KVP.Value);
+
+					InDoc.AddMeshForComponent(ComponentName, MeshName);
+				}
+
 				// Combine separate scene trees into the single one with configuration-specific data
 				FConfigurationTree.Merge(CombinedTree, ConfigNode, CfgName);
 			}
@@ -327,10 +388,8 @@ namespace DatasmithSolidworks
 			swDisplayStateOpts_e Option = InDisplayState != null ? swDisplayStateOpts_e.swSpecifyDisplayState : swDisplayStateOpts_e.swThisDisplayState;
 			string[] DisplayStates = InDisplayState != null ? new string[] { InDisplayState } : null;
 
-			if (InDoc is FAssemblyDocument)
+			if (InDoc is FAssemblyDocument AsmDoc)
 			{
-				FAssemblyDocument AsmDoc = InDoc as FAssemblyDocument;
-
 				HashSet<FComponentName> InComponentsSet = new HashSet<FComponentName>();
 
 				foreach (FComponentName CompName in AsmDoc.SyncState.ExportedComponentsMap.Keys)
@@ -384,23 +443,27 @@ namespace DatasmithSolidworks
 				TargetConfig.Materials = Materials;
 			}
 
-			if (InComponentTree.Children != null)
+			foreach (FConfigurationTree.FComponentTreeNode Child in InComponentTree.EnumChildren())
 			{
-				foreach (FConfigurationTree.FComponentTreeNode Child in InComponentTree.Children)
-				{
-					SetComponentTreeMaterials(Child, InComponentMaterialsMap, InConfigurationName, bIsDisplayState);
-				}
+				SetComponentTreeMaterials(Child, InComponentMaterialsMap, InConfigurationName, bIsDisplayState);
 			}
 		}
 
-		private static void CollectComponentsRecursive(FDocument InDoc, Component2 InComponent, FConfigurationTree.FComponentTreeNode InParentNode)
+		private static void CollectComponentsRecursive(FDocument InDoc, Component2 InComponent,
+			FConfigurationTree.FComponentTreeNode InParentNode, Dictionary<Component2, FActorName> OutMeshesToExportMap)
 		{
-			FConfigurationTree.FComponentTreeNode NewNode = new FConfigurationTree.FComponentTreeNode();
+			FConfigurationTree.FComponentTreeNode NewNode = new FConfigurationTree.FComponentTreeNode(InComponent);
 			InParentNode.Children.Add(NewNode);
 
 			// Basic properties
-			NewNode.ComponentName = new FComponentName(InComponent);
-			NewNode.ComponentID = InComponent.GetID();
+			NewNode.ComponentInfo.ComponentName = new FComponentName(InComponent);
+			NewNode.ComponentInfo.ComponentId = InComponent.GetID();
+
+			// ComponentDoc is null if component is suppressed or lightweight
+			ModelDoc2 ModelDoc = (ModelDoc2)InComponent.GetModelDoc2();
+			NewNode.ComponentInfo.PartPath = (ModelDoc is PartDoc) ? ModelDoc.GetPathName() : null;  // Identify whether the component is a Part component
+
+
 			NewNode.CommonConfig.bVisible = InComponent.Visible != (int)swComponentVisibilityState_e.swComponentHidden;
 			NewNode.CommonConfig.bSuppressed = InComponent.IsSuppressed();
 
@@ -437,14 +500,25 @@ namespace DatasmithSolidworks
 				foreach (object ObjChild in Children)
 				{
 					Component2 Child = (Component2)ObjChild;
-					CollectComponentsRecursive(InDoc, Child, NewNode);
+					CollectComponentsRecursive(InDoc, Child, NewNode, OutMeshesToExportMap);
 				}
 
-				NewNode.Children.Sort(delegate (FConfigurationTree.FComponentTreeNode InA, FConfigurationTree.FComponentTreeNode InB)
+				// todo: sorting ensures that component order is stable in export(e.g. in variant property bindings order)
+				// but probably makes sense to to it at the datasmith export, explicitly
+				// especially that Children here begs for a dictionary, not list(child looked up by name form Children list)
+				NewNode.Children.Sort((InA, InB) => InA.ComponentId - InB.ComponentId);
+			}
+
+			if (InDoc.NeedExportComponent(NewNode, NewNode.CommonConfig))
+			{
+				if (NewNode.IsPartComponent() && InDoc.NeedGeometryExport(NewNode, NewNode.CommonConfig))
 				{
-					return InA.ComponentID - InB.ComponentID;
-				});
+					// Collect meshes to export when it's a part component
+					InDoc.AddPartDocument(NewNode);
+					OutMeshesToExportMap.Add(NewNode.Component, NewNode.ActorName);
+				}
 			}
+			InDoc.AddExportedComponent(NewNode);
 		}
 
 		public bool IsSameMesh(FComponentName ComponentName, string ConfigNameA, string ConfigNameB)
diff --git a/Engine/Source/Programs/Enterprise/Datasmith/DatasmithSolidworksExporter/Private/ConfigurationTree.cs b/Engine/Source/Programs/Enterprise/Datasmith/DatasmithSolidworksExporter/Private/ConfigurationTree.cs
--- a/Engine/Source/Programs/Enterprise/Datasmith/DatasmithSolidworksExporter/Private/ConfigurationTree.cs
+++ b/Engine/Source/Programs/Enterprise/Datasmith/DatasmithSolidworksExporter/Private/ConfigurationTree.cs
@@ -40,10 +40,26 @@ namespace DatasmithSolidworks
 		};
 
 		// A node contains all the configuration data for a single component
+		[ComVisible(false)]
 		public class FComponentTreeNode
 		{
-			public FComponentName ComponentName;
-			public int ComponentID;
+			public readonly Component2 Component;
+
+			[ComVisible(false)]
+			public struct FComponentInfo
+			{
+				public FComponentName ComponentName;
+				public int ComponentId;
+				public string PartPath;
+			}
+
+			public FComponentInfo ComponentInfo = new FComponentInfo();
+
+			public int ComponentId => ComponentInfo.ComponentId;
+
+			public FComponentName ComponentName => ComponentInfo.ComponentName;
+			public FActorName ActorName => ComponentInfo.ComponentName.GetActorName();
+			public string PartPath => ComponentInfo.PartPath;
 
 			// Common configuration data
 			public FComponentConfig CommonConfig = new FComponentConfig();
@@ -61,6 +77,24 @@ namespace DatasmithSolidworks
 
 			public List<FComponentTreeNode> Children;
 
+			public IEnumerable<FComponentTreeNode> EnumChildren()
+			{
+				if (Children == null)
+				{
+					yield break;
+				}
+
+				foreach (FComponentTreeNode Child in Children)
+				{
+					yield return Child;
+				}
+			}
+
+			public FComponentTreeNode(Component2 InComponent=null)
+			{
+				Component = InComponent;
+			}
+
 			public FComponentConfig AddConfiguration(string InConfigurationName, bool bIsDisplayState)
 			{
 				List<FComponentConfig> TargetList = null;
@@ -117,35 +151,38 @@ namespace DatasmithSolidworks
 			}
 
 			// Add all parameters except those which could be configuration-specific
-			public void AddParametersFrom(FComponentTreeNode InInput)
-			{
-				ComponentName = InInput.ComponentName;
-				ComponentID = InInput.ComponentID;
-			}
-		};
 
-		static public void Merge(FComponentTreeNode OutCombined, FComponentTreeNode InTree, string InConfigurationName)
-		{
-			if (InTree.Children == null)
+			public FActorName GetActorName()
 			{
-				return;
+				return ComponentName.GetActorName();
 			}
 
-			if (OutCombined.Children == null)
+			public bool IsPartComponent()
 			{
-				OutCombined.Children = new List<FComponentTreeNode>();
+				return ComponentInfo.PartPath != null;
 			}
+		};
 
-			foreach (FComponentTreeNode Child in InTree.Children)
+		static public void Merge(FComponentTreeNode OutCombined, FComponentTreeNode InTree, string InConfigurationName)
+		{
+			foreach (FComponentTreeNode Child in InTree.EnumChildren())
 			{
+				if (OutCombined.Children == null)
+				{
+					OutCombined.Children = new List<FComponentTreeNode>();
+				}
+
 				// Find the same node in 'combined', merge parameters
 				FComponentTreeNode CombinedChild = OutCombined.Children.FirstOrDefault(X => X.ComponentName == Child.ComponentName);
 
 				if (CombinedChild == null)
 				{
-					// The node doesn't exist yet, so add it
-					CombinedChild = new FComponentTreeNode();
-					CombinedChild.AddParametersFrom(Child);
+					// The node doesn't exist yet, so add it. copying its info from the input node
+					CombinedChild = new FComponentTreeNode(Child.Component)
+					{
+						ComponentInfo = Child.ComponentInfo
+					};
+
 					OutCombined.Children.Add(CombinedChild);
 
 					// Copy common materials, which should be "default" for the component. Do not propagate these
@@ -153,6 +190,17 @@ namespace DatasmithSolidworks
 					OutCombined.CommonConfig.Materials = Child.CommonConfig.Materials;
 				}
 
+				if (CombinedChild.PartPath == null)
+				{
+					CombinedChild.ComponentInfo.PartPath = Child.PartPath;
+				}
+				else
+				{
+					// Either child didn't have PartPath resolved - when it's document is not loaded in this configuration(suppressed) or it's not a part
+					// Same component is not expected to have different PartPath in different configurations, seems like by SW design
+					Debug.Assert(Child.PartPath == null || CombinedChild.PartPath == Child.PartPath);
+				}
+
 				// Make a NodeConfig and copy parameter values from 'tree' node
 				FComponentConfig NodeConfig = CombinedChild.AddConfiguration(InConfigurationName, false);
 				NodeConfig.CopyFrom(Child.CommonConfig);
@@ -188,14 +236,15 @@ namespace DatasmithSolidworks
 			if (InNode.Configurations != null && InNode.Configurations.Count > 0)
 			{
 				// Check transform
-				float[] Transform = InNode.Configurations[0].RelativeTransform;
+				float[] Transform = InNode.Configurations[0].Transform;
+				float[] RelativeTransform = InNode.Configurations[0].RelativeTransform;
 				bool bAllTransformsAreSame = true;
-				if (Transform != null)
+				if (RelativeTransform != null)
 				{
 					// There could be components without a transform, so we're checking if for null
 					for (int Idx = 1; Idx < InNode.Configurations.Count; Idx++)
 					{
-						if (!InNode.Configurations[Idx].RelativeTransform.SequenceEqual(Transform))
+						if (!InNode.Configurations[Idx].RelativeTransform.SequenceEqual(RelativeTransform))
 						{
 							bAllTransformsAreSame = false;
 							break;
@@ -203,10 +252,12 @@ namespace DatasmithSolidworks
 					}
 					if (bAllTransformsAreSame)
 					{
-						InNode.CommonConfig.RelativeTransform = Transform;
+						InNode.CommonConfig.Transform = Transform;
+						InNode.CommonConfig.RelativeTransform = RelativeTransform;
 						foreach (FComponentConfig Config in InNode.Configurations)
 						{
 							Config.RelativeTransform = null;
+							Config.Transform = null;
 						}
 					}
 				}
@@ -297,12 +348,9 @@ namespace DatasmithSolidworks
 			}
 
 			// Recurse to children
-			if (InNode.Children != null)
+			foreach (FComponentTreeNode Child in InNode.EnumChildren())
 			{
-				foreach (FComponentTreeNode Child in InNode.Children)
-				{
-					Compress(Child, ConfigurationExporter);
-				}
+				Compress(Child, ConfigurationExporter);
 			}
 		}
 
@@ -363,12 +411,9 @@ namespace DatasmithSolidworks
 			}
 
 			// Recurse to children
-			if (InNode.Children != null)
+			foreach (FComponentTreeNode Child in InNode.EnumChildren())
 			{
-				foreach (FComponentTreeNode Child in InNode.Children)
-				{
-					FillConfigurationData(ConfigurationExporter, Child, InConfigurationName, OutData, bIsDisplayState);
-				}
+				FillConfigurationData(ConfigurationExporter, Child, InConfigurationName, OutData, bIsDisplayState);
 			}
 		}
 	}
diff --git a/Engine/Source/Programs/Enterprise/Datasmith/DatasmithSolidworksExporter/Private/DatasmithExporter.cs b/Engine/Source/Programs/Enterprise/Datasmith/DatasmithSolidworksExporter/Private/DatasmithExporter.cs
--- a/Engine/Source/Programs/Enterprise/Datasmith/DatasmithSolidworksExporter/Private/DatasmithExporter.cs
+++ b/Engine/Source/Programs/Enterprise/Datasmith/DatasmithSolidworksExporter/Private/DatasmithExporter.cs
@@ -406,11 +406,14 @@ namespace DatasmithSolidworks
 
 			HashSet<int> MeshAddedMaterials = new HashSet<int>();
 
+			HashSet<int> MeshMaterialIds = new HashSet<int>();
+
 			for (int TriIndex = 0; TriIndex < InData.Triangles.Length; TriIndex++)
 			{
 				FTriangle Triangle = InData.Triangles[TriIndex];
 				int MatID = 0;
 
+				MeshMaterialIds.Add(Triangle.MaterialID);
 				if (Triangle.MaterialID >= 1)
 				{
 					if (!MeshAddedMaterials.Contains(Triangle.MaterialID))
@@ -439,6 +442,12 @@ namespace DatasmithSolidworks
 				Mesh.SetFaceUV(TriIndex, 0, Triangle[0], Triangle[1], Triangle[2]);
 			}
 
+			Addin.Instance.LogDebug($"  material ids assigned to faces: {string.Join(", ", MeshMaterialIds)})");
+			if (MeshAddedMaterials.Count > 0)
+			{
+				Addin.Instance.LogDebug($"  assigned materials: {string.Join(", ", MeshAddedMaterials)})");
+			}
+
 			OutMeshPair = new Tuple<FDatasmithFacadeMeshElement, FDatasmithFacadeMesh>(MeshElement, Mesh);
 
 			DatasmithScene.ExportDatasmithMesh(MeshElement, Mesh);
@@ -739,6 +748,8 @@ namespace DatasmithSolidworks
 
 		public void ExportMaterials(ConcurrentDictionary<int, FMaterial> InMaterialsMap)
 		{
+			Addin.Instance.LogDebug($"ExportMaterials: \n  {string.Join("  \n", InMaterialsMap.Select(KVP => $"{KVP.Key}: {KVP.Value}" ))}");
+
 			ConcurrentBag<FDatasmithFacadeTexture> CreatedTextures = new ConcurrentBag<FDatasmithFacadeTexture>();
 			ConcurrentBag<FDatasmithFacadeMaterialInstance> CreatedMaterials = new ConcurrentBag<FDatasmithFacadeMaterialInstance>();
 			Parallel.ForEach(InMaterialsMap, MatKVP =>
@@ -756,7 +767,7 @@ namespace DatasmithSolidworks
 				}
 			});
 
-			Addin.Instance.LogDebug($"AddMaterials");
+			Addin.Instance.LogDebug($"  AddMaterials");
 			// Adding stuff to a datasmith scene cannot be multithreaded!
 			foreach (FDatasmithFacadeMaterialInstance Mat in CreatedMaterials)
 			{
diff --git a/Engine/Source/Programs/Enterprise/Datasmith/DatasmithSolidworksExporter/Private/Document.cs b/Engine/Source/Programs/Enterprise/Datasmith/DatasmithSolidworksExporter/Private/Document.cs
--- a/Engine/Source/Programs/Enterprise/Datasmith/DatasmithSolidworksExporter/Private/Document.cs
+++ b/Engine/Source/Programs/Enterprise/Datasmith/DatasmithSolidworksExporter/Private/Document.cs
@@ -7,6 +7,7 @@ using System.Collections.Generic;
 using System.Diagnostics;
 using System.IO;
 using System.Threading;
+using static DatasmithSolidworks.FAssemblyDocument;
 
 namespace DatasmithSolidworks
 {
@@ -22,7 +23,7 @@ namespace DatasmithSolidworks
 		public bool bHasConfigurations = false;
 		
 		protected FDatasmithFacadeScene DatasmithScene = null;
-		protected FDatasmithExporter Exporter = null;
+		public FDatasmithExporter Exporter = null;
 		protected int DocId = -1;
 		protected bool bDirectLinkSyncInProgress = false;
 		protected bool bFileExportInProgress = false;
@@ -88,7 +89,7 @@ namespace DatasmithSolidworks
 			bDocumentIsDirty = bInDirty;
 		}
 
-		protected void SetExportStatus(string InMessage)
+		public void SetExportStatus(string InMessage)
 		{
 			if (bDirectLinkSyncInProgress)
 			{
@@ -186,10 +187,8 @@ namespace DatasmithSolidworks
 			else
 			{
 				// Export to file
-
 				bFileExportInProgress = true;
 
-
 				string OutDir = Path.GetDirectoryName(DatasmithFileExportPath);
 				string CleanFileName = Path.GetFileNameWithoutExtension(DatasmithFileExportPath);
 
@@ -207,10 +206,14 @@ namespace DatasmithSolidworks
 
 				PreExport(Meshes, true);
 
-				List<FConfigurationData> Configs = new FConfigurationExporter(Meshes).ExportConfigurations(this);
+				string[] ConfigurationNames = SwDoc?.GetConfigurationNames();
+				FConfigurationExporter ConfigurationExporter = new FConfigurationExporter(Meshes, ConfigurationNames);
+
+				List<FConfigurationData> Configs = ConfigurationExporter.ExportConfigurations(this);
+
 				bHasConfigurations = (Configs != null) && (Configs.Count != 0);
 
-				ExportToDatasmithScene(Meshes);
+				ExportToDatasmithScene(ConfigurationExporter);
 			
 				ExportLights();
 				ExportConfigurations(Configs);
@@ -239,6 +242,12 @@ namespace DatasmithSolidworks
 			}
 		}
 
+		// todo: make abstract
+		public virtual void ExportToDatasmithScene(FConfigurationExporter ConfigurationExporter)
+		{
+			throw new NotImplementedException();
+		}
+
 		public abstract bool HasMaterialUpdates();
 		public abstract void ExportToDatasmithScene(FMeshes Meshes);
 		public abstract void PreExport(FMeshes Meshes, bool bConfigurations);  // Called before configurations are parsed to prepare meshes needed to identify if they create different configurations
@@ -314,5 +323,30 @@ namespace DatasmithSolidworks
 				OnDirectLinkSync();
 			}
 		}
+
+		public virtual bool NeedExportComponent(FConfigurationTree.FComponentTreeNode InComponent,
+			FConfigurationTree.FComponentConfig ActiveComponentConfig)
+		{
+			return false;
+		}
+
+		/// <summary>
+		/// Whether this component node needs geometry exported/reexported. Currenly doesn't use configuration
+		/// todo: DL configurations
+		/// </summary>
+		public virtual bool NeedGeometryExport(FConfigurationTree.FComponentTreeNode InNode, FConfigurationTree.FComponentConfig ActiveComponentConfig)
+		{
+			return false;
+		}
+
+		public virtual void AddPartDocument(FConfigurationTree.FComponentTreeNode InNode){}
+
+		public virtual void AddExportedComponent(FConfigurationTree.FComponentTreeNode InNode) {}
+
+		public abstract ConcurrentDictionary<FComponentName, FObjectMaterials> LoadDocumentMaterials(HashSet<FComponentName> ComponentNamesToExportSet);
+		public abstract void AddComponentMaterials(FComponentName ComponentName, FObjectMaterials Materials);
+		public abstract FObjectMaterials GetComponentMaterials(Component2 Comp);
+
+		public abstract void AddMeshForComponent(FComponentName ComponentName, string MeshName);
 	}
 }
\ No newline at end of file
diff --git a/Engine/Source/Programs/Enterprise/Datasmith/DatasmithSolidworksExporter/Private/MathUtils.cs b/Engine/Source/Programs/Enterprise/Datasmith/DatasmithSolidworksExporter/Private/MathUtils.cs
--- a/Engine/Source/Programs/Enterprise/Datasmith/DatasmithSolidworksExporter/Private/MathUtils.cs
+++ b/Engine/Source/Programs/Enterprise/Datasmith/DatasmithSolidworksExporter/Private/MathUtils.cs
@@ -78,6 +78,17 @@ namespace DatasmithSolidworks
 			return Result;
 		}
 
+		public static float[] TransformIdentity()
+		{
+			return new[]
+			{
+				1.0f, 0.0f, 0.0f, 0.0f,
+				0.0f, 1.0f, 0.0f, 0.0f,
+				0.0f, 0.0f, 1.0f, 0.0f,
+				0.0f, 0.0f, 0.0f, 1.0f
+			};
+		}
+
 		public static float[] LookAt(FVec3 InDirection, FVec3 InFromPoint, float InGeomScale)
 		{
 			float[] Ret = new float[16];
diff --git a/Engine/Source/Programs/Enterprise/Datasmith/DatasmithSolidworksExporter/Private/ObjectMaterials.cs b/Engine/Source/Programs/Enterprise/Datasmith/DatasmithSolidworksExporter/Private/ObjectMaterials.cs
--- a/Engine/Source/Programs/Enterprise/Datasmith/DatasmithSolidworksExporter/Private/ObjectMaterials.cs
+++ b/Engine/Source/Programs/Enterprise/Datasmith/DatasmithSolidworksExporter/Private/ObjectMaterials.cs
@@ -76,6 +76,11 @@ namespace DatasmithSolidworks
 			return GetMaterial(PartMaterialID);
 		}
 
+		public override string ToString()
+		{
+			return $"FObjectMaterials:ComponentMaterialID='{ComponentMaterialID}', PartMaterialID='{PartMaterialID}'";
+		}
+
 		public FMaterial GetMaterial(Face2 InFace)
 		{
 			if (ComponentMaterialID != -1) // Highest priority
@@ -231,6 +236,8 @@ namespace DatasmithSolidworks
 
 		public static FObjectMaterials LoadComponentMaterials(FDocument InComponentOwner, Component2 InComponent, swDisplayStateOpts_e InDisplayState, string[] InDisplayStateNames)
 		{
+			Addin.Instance.LogDebug($"LoadComponentMaterials: {InComponentOwner.GetPathName()}, {InDisplayState}, {InDisplayStateNames}");
+
 			ModelDoc2 ComponentDoc = InComponent.GetModelDoc2() as ModelDoc2;
 			if (ComponentDoc == null || !(ComponentDoc is PartDoc))
 			{
@@ -239,6 +246,8 @@ namespace DatasmithSolidworks
 			}
 
 			int NumMaterials = InComponent.GetRenderMaterialsCount2((int)InDisplayState, InDisplayStateNames);
+			Addin.Instance.LogDebug($"  GetRenderMaterialsCount2: {NumMaterials}");
+
 			if (NumMaterials == 0)
 			{
 				return null;
@@ -293,9 +302,14 @@ namespace DatasmithSolidworks
 
 		public static ConcurrentDictionary<FComponentName, FObjectMaterials> LoadAssemblyMaterials(FAssemblyDocument InAsmDoc, HashSet<FComponentName> InComponentsSet, swDisplayStateOpts_e InDisplayState, string[] InDisplayStateNames)
 		{
+			Addin.Instance.LogDebug($"LoadAssemblyMaterials: {InAsmDoc.GetPathName()}, {InDisplayState}, {InDisplayStateNames}");
+			Addin.Instance.LogDebug($"  for components: {string.Join(", ", InComponentsSet)}");
+
 			IModelDocExtension Ext = InAsmDoc.SwDoc.Extension;
 			int NumMaterials = Ext.GetRenderMaterialsCount2((int)InDisplayState, InDisplayStateNames);
 
+			Addin.Instance.LogDebug($"  NumMaterials={NumMaterials}");
+
 			ConcurrentDictionary<FComponentName, FObjectMaterials> DocMaterials = new ConcurrentDictionary<FComponentName, FObjectMaterials>();
 
 			if (NumMaterials > 0)
@@ -332,8 +346,8 @@ namespace DatasmithSolidworks
 							FObjectMaterials PartMaterials = LoadPartMaterials(InAsmDoc, Doc as PartDoc, InDisplayState, InDisplayStateNames);
 							if (PartMaterials != null)
 							{
-								// xxx: why trying to present PahtName as component name?
-								// xxx: how it can be a PartDoc in Assembly function
+								// todo: why trying to present PathName as component name?
+								//  how it can be a PartDoc in Assembly function
 								DocMaterials[FComponentName.FromCustomString((Doc as ModelDoc2).GetPathName())] = PartMaterials;
 							}
 							continue;
@@ -343,7 +357,9 @@ namespace DatasmithSolidworks
 			}
 
 			// Check for materials that are not per component (but instead per face, feature, part etc.)
-			Parallel.ForEach(InComponentsSet, CompName =>
+			// todo: consider replacing Parallel.ForEach with simple loop
+			Parallel.ForEach(InComponentsSet,
+				CompName =>
 			{
 				Component2 Comp = null;
 
diff --git a/Engine/Source/Programs/Enterprise/Datasmith/DatasmithSolidworksExporter/Private/PartDocument.cs b/Engine/Source/Programs/Enterprise/Datasmith/DatasmithSolidworksExporter/Private/PartDocument.cs
--- a/Engine/Source/Programs/Enterprise/Datasmith/DatasmithSolidworksExporter/Private/PartDocument.cs
+++ b/Engine/Source/Programs/Enterprise/Datasmith/DatasmithSolidworksExporter/Private/PartDocument.cs
@@ -4,6 +4,7 @@ using SolidWorks.Interop.sldworks;
 using SolidWorks.Interop.swconst;
 using System;
 using System.Collections.Concurrent;
+using System.Collections.Generic;
 using System.IO;
 
 namespace DatasmithSolidworks
@@ -120,6 +121,11 @@ namespace DatasmithSolidworks
 			}
 		}
 
+		public override void ExportToDatasmithScene(FConfigurationExporter ConfigurationExporter)
+		{
+			ExportToDatasmithScene(ConfigurationExporter.Meshes);
+		}
+
 		public override void ExportToDatasmithScene(FMeshes Meshes)
 		{
 			string PartName = Path.GetFileNameWithoutExtension(PathName);
@@ -146,7 +152,7 @@ namespace DatasmithSolidworks
 						continue;
 					}
 					// todo: not make extra copy - use existing(i.e. instead of bHasConfigurations) 
-					FActorName ActorName = new FConfigurationExporter(Meshes).GetMeshActorName(RootComponentName, ConfigurationName);
+					FActorName ActorName = new FConfigurationExporter(Meshes, null).GetMeshActorName(RootComponentName, ConfigurationName);
 
 					bool bHasMesh = Exporter.ExportMesh($"{ActorName}_Mesh", MeshData, ActorName,
 						out Tuple<FDatasmithFacadeMeshElement, FDatasmithFacadeMesh> NewMesh);
@@ -269,6 +275,24 @@ namespace DatasmithSolidworks
 			SwPartDoc.FileSavePostNotify -= new DPartDocEvents_FileSavePostNotifyEventHandler(OnPartFileSavePostNotify);
 		}
 
+		public override FObjectMaterials GetComponentMaterials(Component2 Comp)
+		{
+			return null;
+		}
+
+		public override void AddMeshForComponent(FComponentName ComponentName, string MeshName)
+		{
+		}
+
+		public override void AddComponentMaterials(FComponentName ComponentName, FObjectMaterials Materials)
+		{
+		}
+
+		public override ConcurrentDictionary<FComponentName, FObjectMaterials> LoadDocumentMaterials(HashSet<FComponentName> ComponentNamesToExportSet)
+		{
+			return null;
+		}
+
 		int OnPartDocumentDestroyNotify2(int DestroyType)
 		{
 			Addin.Instance.CloseDocument(DocId);
