Added Mac ARM support for ArchiCAD 26 exporter
Eliminated harmless warnings during compilation

diff --git a/Engine/Source/Programs/Enterprise/Datasmith/DatasmithARCHICADExporter/Build/Config/UE_AC_ARM.xcconfig b/Engine/Source/Programs/Enterprise/Datasmith/DatasmithARCHICADExporter/Build/Config/UE_AC_ARM.xcconfig
new file mode 100644
--- /dev/null
+++ b/Engine/Source/Programs/Enterprise/Datasmith/DatasmithARCHICADExporter/Build/Config/UE_AC_ARM.xcconfig
@@ -0,0 +1,75 @@
+// Copyright Epic Games, Inc. All Rights Reserved.
+
+#include? "SDKsRoot.xcconfig"
+
+AC_SDK_PATH = $(UESDKRoot)/HostMac/Mac/Archicad/$(AC_VERSION)
+
+AC_INCLUDES_PATH = "$(AC_SDK_PATH)/Support/Inc" "$(AC_SDK_PATH)/Support/Modules/GSRoot/STL/imp" "$(AC_SDK_PATH)/Support/Modules/GSRoot" "$(AC_SDK_PATH)/Support/Modules/InputOutput"/** "$(AC_SDK_PATH)/Support/Modules/GSModelDevLib"/** "$(AC_SDK_PATH)/Support/Modules/GSUtils"/** "$(AC_SDK_PATH)/Support/Modules/GSModeler"/** "$(AC_SDK_PATH)/Support/Modules/Model3D"/** "$(AC_SDK_PATH)/Support/Modules/MPIInterfaceDevLib"/** "$(AC_SDK_PATH)/Support/Modules/Geometry"/** "$(AC_SDK_PATH)/Support/Modules/DGLib"/** "$(AC_SDK_PATH)/Support/Modules/GDL"/** "$(AC_SDK_PATH)/Support/Modules/GXImage"/** "$(AC_SDK_PATH)/Support/Modules/GX"/** "$(AC_SDK_PATH)/Support/Modules/DGGraphix"/** "$(AC_SDK_PATH)/Support/Modules/Graphix"/** "$(AC_SDK_PATH)/Support/Modules/TextEngine"/** "$(AC_SDK_PATH)/Support/Modules/GXImageBase"/** "$(AC_SDK_PATH)/Support/Modules/PointCloudManager" "$(AC_SDK_PATH)/Support/Modules/PointCloud" "$(AC_SDK_PATH)/Support/Modules/Pattern" "$(AC_SDK_PATH)/Support/Modules/Brep" "$(AC_SDK_PATH)/Support/Modules/VBAttributes" "$(AC_SDK_PATH)/Support/Modules/GSModeler2D" "$(AC_SDK_PATH)/Support/Modules/ObjectDatabase" "$(AC_SDK_PATH)/Support/Modules/UDLib" "$(AC_SDK_PATH)/Support/Modules/RS" "$(UE_Engine)/Binaries/Mac/DatasmithSDK"
+
+UE_Runtime = $(UE_Engine)/Source/Runtime
+UE_Core = $(UE_Runtime)/Core/Public
+UE_Project = $(UE_Runtime)/Projects/Public
+UE_ArchiCADExporter = $(UE_Engine)/Source/Programs/Enterprise/Datasmith/DatasmithARCHICADExporter
+
+UE_INCLUDES_PATH = "$(UE_Core)" "$(UE_Core)/Internationalization" "$(UE_Core)/Async" "$(UE_Core)/Concurrency" "$(UE_Core)/Containers" "$(UE_Core)/Delegates" "$(UE_Core)/GenericPlatform" "$(UE_Core)/HAL" "$(UE_Core)/Logging" "$(UE_Core)/Math" "$(UE_Core)/Misc" "$(UE_Core)/Modules" "$(UE_Core)/Modules\Boilerplate" "$(UE_Core)/ProfilingDebugging" "$(UE_Core)/Serialization" "$(UE_Core)/Serialization\Csv" "$(UE_Core)/Stats" "$(UE_Core)/Templates" "$(UE_Core)/UObject" "$(UE_Project)/Public" "$(UE_Project)/Interfaces" "$(UE_Runtime)/TraceLog/Public"  "$(UE_Runtime)/Datasmith/DatasmithCore/Public" "$(UE_Runtime)/Datasmith/DirectLink/Public" "$(UE_Runtime)/Launch/Resources" "$(UE_Engine)/Source/Developer/Datasmith/DatasmithExporter/Public" "$(UE_Engine)/Source/Developer/Datasmith/DatasmithExporterUI/Public" "$(UE_Engine)/Source/Runtime/Messaging/Public" "$(UE_Engine)/Source/Runtime/CoreUObject/Public/" "$(UE_ArchiCADExporter)/DatasmithValidator/Public"
+
+USER_HEADER_SEARCH_PATHS = "." $(AC_INCLUDES_PATH) $(UE_INCLUDES_PATH)
+
+AUTO_SYNC = 1
+
+VMajor = 1
+VMinor = 0
+VRevision = 0
+VBetaOrRevision = b // b for beta or . for revision
+VAddOn = $(VMajor).$(VMinor)$(VBetaOrRevision)$(VRevision)
+
+CurrentYear = 2022
+
+FRAMEWORK_SEARCH_PATHS = $(FRAMEWORK_SEARCH_PATHS) $(AC_SDK_PATH)/Support/Frameworks
+
+LIBRARY_SEARCH_PATHS = "Dylibs" "$(AC_SDK_PATH)/Support/Lib/Mactel"
+
+REZ_SEARCH_PATHS = RO "$(AC_SDK_PATH)/Support/Inc" "$(AC_SDK_PATH)/Support/Modules/DGLib"
+
+AC_FRAMEWORKS = -framework DG -framework GSModeler -framework GSRoot -framework GSUtils -framework GSZLib -framework GX -framework GXImage -framework GXImageBase -framework InputOutput -framework UD -framework VBAttributes -framework VectorImage -framework ObjectDatabase -framework Geometry
+
+DATASMITH_DYLIB[config=Release] = DatasmithUE4ArchiCAD.dylib
+DATASMITH_DYLIB[config=Debug] = DatasmithUE4ArchiCAD-Mac-Debug.dylib
+
+OTHER_LDFLAGS = -lACAP_STAT $(AC_FRAMEWORKS) ./Dylibs/$DATASMITH_DYLIB
+
+AC_LANGUAGE_FOLDER = /Extensions
+
+AC_APPLICATION_NAME = "ARCHICAD"
+AC_ADDON_NAME = "Datasmith Direct Link"
+AC_ADDON_DEBUG_SUFFIXE =
+
+PlugInName = "DatasmithARCHICAD$(AC_VERSION)Exporter"
+
+PRODUCT_NAME = DatasmithARCHICAD$(AC_VERSION)Exporter
+
+PRODUCT_BUNDLE_IDENTIFIER = com.epicgames.DatasmithARCHICAD$(AC_VERSION)Exporter
+
+GRC_CFG_DEFINES[config=Release] = -DNDEBUG
+GRC_CFG_DEFINES[config=Debug] = -DDEBUG
+
+GRC_DEFINES = $GRC_CFG_DEFINES -DPREPROCESS_GRC -DCHECK_UPDATE -DGS_64BIT -Dmacintosh -DINT__APP -DAUTO_SYNC=$(AUTO_SYNC)
+
+GRC_LIB_OUTPUT = $(PROJECT_TEMP_DIR)/GRC_OUT$(AC_VERSION)/
+GRC_OUTPUT = $(PROJECT_TEMP_DIR)/GRC_OUT$(RES_NAME)$(AC_VERSION)/
+
+ADDON_PREPROCESSOR_DEFINITIONS = ADDON_VERSION=$(VAddOn) VMAJOR=$(VMajor) VMINOR=$(VMinor) VREVISION=$(VRevision) AUTO_SYNC=$(AUTO_SYNC)
+
+AC_PREPROCESSOR_DEFINITIONS = AC_VERSION=$(AC_VERSION) AC_SDK=$(AC_SDK_BUILD_NUMBER) macintosh=1 ACExtension
+
+UE_PREPROCESSOR_DEFINITIONS = UE_BUILD_DEVELOPMENT=1 UE_BUILD_MINIMAL=1 WITH_EDITOR=0 WITH_EDITORONLY_DATA=0 WITH_SERVER_CODE=1 WITH_ENGINE=0 WITH_UNREAL_DEVELOPER_TOOLS=0 WITH_PLUGIN_SUPPORT=0 IS_MONOLITHIC=1 IS_PROGRAM=1 PLATFORM_MAC=1 PLATFORM_APPLE=1 UE_BUILD_DEVELOPMENT_WITH_DEBUGGAME=0 UBT_COMPILED_PLATFORM=Mac CORE_API=DLLIMPORT DATASMITHEXPORTER_API=DLLIMPORT DATASMITHCORE_API=DLLIMPORT DIRECTLINK_API=DLLIMPORT DATASMITHEXPORTERUI_API=DLLIMPORT DATASMITHVALIDATOR_API=DLLIMPORT COREUOBJECT_API=DLLIMPORT
+
+//:configuration = Debug
+GCC_PREPROCESSOR_DEFINITIONS = $(ADDON_PREPROCESSOR_DEFINITIONS) $(AC_PREPROCESSOR_DEFINITIONS) $(UE_PREPROCESSOR_DEFINITIONS)
+
+ARCHS = arm64 x86_64
+//ARCHS_STANDARD=arm64 x86_64
+//ARCHS_STANDARD_32_64_BIT=arm64 x86_64
+//ARCHS_STANDARD_64_BIT=arm64 x86_64
+//ARCHS_STANDARD_INCLUDING_64_BIT=arm64 x86_64
+
