[Solidworks] Assembly Part configurations

- Part component included into assembly with different own configurations creates different meshes for corresponding variants

- Reworked plugin to Separate DirectLink session from Export session

diff --git a/Engine/Source/Programs/Enterprise/Datasmith/DatasmithSolidworksExporter/Private/Addin.cs b/Engine/Source/Programs/Enterprise/Datasmith/DatasmithSolidworksExporter/Private/Addin.cs
--- a/Engine/Source/Programs/Enterprise/Datasmith/DatasmithSolidworksExporter/Private/Addin.cs
+++ b/Engine/Source/Programs/Enterprise/Datasmith/DatasmithSolidworksExporter/Private/Addin.cs
@@ -309,8 +309,6 @@ namespace DatasmithSolidworks
 					default: throw new Exception("Unsupported document type");
 				}
 
-				CurrentDocument.Init();
-
 				OpenDocuments.Add(DocId, CurrentDocument);
 			}
 
@@ -333,13 +331,7 @@ namespace DatasmithSolidworks
 		{
 			if (CurrentDocument != null)
 			{
-				CurrentDocument.bDirectLinkAutoSync = !CurrentDocument.bDirectLinkAutoSync;
-
-				if (CurrentDocument.bDirectLinkAutoSync && CurrentDocument.DirectLinkSyncCount == 0)
-				{
-					// Run first sync
-					CurrentDocument.OnDirectLinkSync();
-				}
+				CurrentDocument.ToggleDirectLinkAutoSync();
 			}
 		}
 
diff --git a/Engine/Source/Programs/Enterprise/Datasmith/DatasmithSolidworksExporter/Private/AssemblyDocument.cs b/Engine/Source/Programs/Enterprise/Datasmith/DatasmithSolidworksExporter/Private/AssemblyDocument.cs
--- a/Engine/Source/Programs/Enterprise/Datasmith/DatasmithSolidworksExporter/Private/AssemblyDocument.cs
+++ b/Engine/Source/Programs/Enterprise/Datasmith/DatasmithSolidworksExporter/Private/AssemblyDocument.cs
@@ -11,7 +11,8 @@ using System.Threading.Tasks;
 
 namespace DatasmithSolidworks
 {
-	public class FAssemblyDocument : FDocument
+
+	public class FAssemblyDocumentTracker : FDocumentTracker
 	{
 		public enum EComponentDirtyState
 		{
@@ -25,9 +26,14 @@ namespace DatasmithSolidworks
 		public class FSyncState
 		{
 			public Dictionary<string, FPartDocument> PartsMap = new Dictionary<string, FPartDocument>();
-			public ConcurrentDictionary<FComponentName, FObjectMaterials> ComponentsMaterialsMap = new ConcurrentDictionary<FComponentName, FObjectMaterials>();
+
+			public ConcurrentDictionary<FComponentName, FObjectMaterials> ComponentsMaterialsMap =
+				new ConcurrentDictionary<FComponentName, FObjectMaterials>();
+
 			public Dictionary<FComponentName, string> ComponentToPartMap = new Dictionary<FComponentName, string>();
-			public Dictionary<FComponentName, float[]> ComponentsTransformsMap = new Dictionary<FComponentName, float[]>();
+
+			public Dictionary<FComponentName, float[]> ComponentsTransformsMap =
+				new Dictionary<FComponentName, float[]>();
 
 			/// All updated(not dirty) components. Updated component will be removed from dirty on export end
 			public HashSet<FComponentName> CleanComponents = new HashSet<FComponentName>();
@@ -36,20 +42,35 @@ namespace DatasmithSolidworks
 			public Dictionary<FComponentName, uint> DirtyComponents = new Dictionary<FComponentName, uint>();
 
 			public HashSet<FComponentName> ComponentsToDelete = new HashSet<FComponentName>();
-			public Dictionary<FComponentName, Component2> ExportedComponentsMap = new Dictionary<FComponentName, Component2>();
+
+			public Dictionary<FComponentName, Component2> ExportedComponentsMap =
+				new Dictionary<FComponentName, Component2>();
+
 			/** Stores which mesh was exported for each component*/
-			public Dictionary<FComponentName, string> ComponentNameToMeshNameMap = new Dictionary<FComponentName, string>();
+			public Dictionary<FComponentName, string> ComponentNameToMeshNameMap =
+				new Dictionary<FComponentName, string>();
 		}
 
 		public AssemblyDoc SwAsmDoc { get; private set; } = null;
 
 		public FSyncState SyncState { get; private set; } = new FSyncState();
 
-		public FAssemblyDocument(int InDocId, AssemblyDoc InSwDoc, FDatasmithExporter InExporter) : base(InDocId, InSwDoc as ModelDoc2, InExporter)
+		public FAssemblyDocumentTracker(FAssemblyDocument InDoc, AssemblyDoc InSwDoc, FDatasmithExporter InExporter) : base(InDoc, InExporter)
 		{
 			SwAsmDoc = InSwDoc;
 		}
 
+		public override void Destroy()
+		{
+			// Release all contained parts
+			foreach (FPartDocument PartDocument in SyncState.PartsMap.Values)
+			{
+				PartDocument.Destroy();
+			}
+
+			base.Destroy();
+		}
+
 		public override void PreExport(FMeshes Meshes, bool bConfigurations)
 		{
 		}
@@ -57,14 +78,6 @@ namespace DatasmithSolidworks
 		// Export to datasmith scene extracted configurations data
 		public override void ExportToDatasmithScene(FConfigurationExporter ConfigurationExporter)
 		{
-			FSyncState OldSyncState = SyncState;
-
-			if (bFileExportInProgress)
-			{
-				// todo: separate DL from file export more explicitly
-				SyncState = new FSyncState();
-			}
-
 			SetExportStatus("Actors");
 
 			ProcessComponentsPendingDelete();
@@ -73,14 +86,12 @@ namespace DatasmithSolidworks
 
 			// Configurations combined tree should have single child(root component of each config is the same)
 			Debug.Assert(ConfigurationExporter.CombinedTree.Children.Count == 1);
-			ExportComponentRecursive(CurrentConfig.Name, ConfigurationExporter, ConfigurationExporter.CombinedTree.Children[0], null);
+			ExportComponentRecursive(CurrentConfig.Name, ConfigurationExporter,
+				ConfigurationExporter.CombinedTree.Children[0], null);
 			SyncState.DirtyComponents.Clear();
 
-			SyncState = OldSyncState;
-
 			// todo: check animation export for DL
 			// Export animations (only allow when exporting to file)
-			// if (bFileExportInProgress)
 			{
 				Component2 Root = CurrentConfig.GetRootComponent3(true);
 				List<FAnimation> Animations = FAnimationExtractor.ExtractAnimations(SwAsmDoc, Root);
@@ -114,12 +125,15 @@ namespace DatasmithSolidworks
 				SyncState.ComponentsTransformsMap.Remove(CompName);
 				Exporter.RemoveActor(ActorName);
 			}
+
 			SyncState.ComponentsToDelete.Clear();
 		}
 
-		public override ConcurrentDictionary<FComponentName, FObjectMaterials> LoadDocumentMaterials(HashSet<FComponentName> ComponentNamesToExportSet)
+		public override ConcurrentDictionary<FComponentName, FObjectMaterials> LoadDocumentMaterials(
+			HashSet<FComponentName> ComponentNamesToExportSet)
 		{
-			return FObjectMaterials.LoadAssemblyMaterials(this, ComponentNamesToExportSet, swDisplayStateOpts_e.swThisDisplayState, null);
+			return FObjectMaterials.LoadAssemblyMaterials(this, ComponentNamesToExportSet,
+				swDisplayStateOpts_e.swThisDisplayState, null);
 		}
 
 		public override void AddComponentMaterials(FComponentName ComponentName, FObjectMaterials Materials)
@@ -141,15 +155,8 @@ namespace DatasmithSolidworks
 		}
 
 		// Export to datasmith scene current configuration only
-		public override void ExportToDatasmithScene(FMeshes meshes)
+		public override void ExportToDatasmithScene(FMeshes Meshes)
 		{
-			FSyncState OldSyncState = SyncState;
-
-			if (bFileExportInProgress)
-			{
-				SyncState = new FSyncState();
-			}
-
 			SetExportStatus("Actors");
 			ProcessComponentsPendingDelete();
 
@@ -160,24 +167,25 @@ namespace DatasmithSolidworks
 
 			// Track components that need their mesh exported: we want to do that in parallel after the 
 			// actor hierarchy has been exported
-			Dictionary<Component2, FActorName> MeshesToExportMap = new Dictionary<Component2, FActorName>();
-		
-			ExportComponentRecursive(Root, null, MeshesToExportMap);
+			FMeshes.FConfiguration MeshesConfiguration = Meshes.GetMeshesConfiguration(ConfigurationName: CurrentConfig.Name);
+
+			ExportComponentRecursive(Root, null, MeshesConfiguration);
 
 			// Export materials
 			SetExportStatus($"Component Materials");
 
 			HashSet<FComponentName> ComponentNamesToExportSet = new HashSet<FComponentName>();
-			foreach (var KVP in MeshesToExportMap)
+			foreach (FComponentName ComponentName in MeshesConfiguration.EnumerateComponentNames())
 			{
-				FComponentName ComponentName = new FComponentName(KVP.Key);
 				if (!ComponentNamesToExportSet.Contains(ComponentName))
 				{
 					ComponentNamesToExportSet.Add(ComponentName);
 				}
 			}
 
-			ConcurrentDictionary<FComponentName, FObjectMaterials> ModifiedComponentsMaterials = FObjectMaterials.LoadAssemblyMaterials(this, ComponentNamesToExportSet, swDisplayStateOpts_e.swThisDisplayState, null);
+			ConcurrentDictionary<FComponentName, FObjectMaterials> ModifiedComponentsMaterials =
+				FObjectMaterials.LoadAssemblyMaterials(this, ComponentNamesToExportSet,
+					swDisplayStateOpts_e.swThisDisplayState, null);
 
 			if (ModifiedComponentsMaterials != null)
 			{
@@ -191,39 +199,12 @@ namespace DatasmithSolidworks
 
 			// Export meshes
 			SetExportStatus($"Component Meshes");
-			ConcurrentDictionary<Component2, Tuple<FDatasmithFacadeMeshElement, FDatasmithFacadeMesh>> CreatedMeshes = new ConcurrentDictionary<Component2, Tuple<FDatasmithFacadeMeshElement, FDatasmithFacadeMesh>>();
-			
-			Parallel.ForEach(MeshesToExportMap, KVP =>
-			{
-				Component2 Comp = KVP.Key;
-
-				FObjectMaterials ComponentMaterials = null;
-				SyncState.ComponentsMaterialsMap?.TryGetValue(new FComponentName(Comp), out ComponentMaterials);
-
-				ConcurrentBag<FBody> Bodies = FBody.FetchBodies(Comp);
-				FMeshData MeshData = FStripGeometry.CreateMeshData(Bodies, ComponentMaterials);
-
-				if (MeshData != null)
-				{
-					Tuple<FDatasmithFacadeMeshElement, FDatasmithFacadeMesh> NewMesh = null;
-					if (Exporter.ExportMesh($"{KVP.Value}_Mesh", MeshData, KVP.Value, out NewMesh))
-					{
-						CreatedMeshes[Comp] = NewMesh;
-					}
-				}
-			});
+			ProcessConfigurationMeshes(MeshesConfiguration, null);
 
-			foreach (var KVP in CreatedMeshes)
-			{
-				string MeshName = Exporter.AddMesh(KVP.Value);
-				SyncState.ComponentNameToMeshNameMap[new FComponentName(KVP.Key)] = MeshName;
-			}
 			SyncState.DirtyComponents.Clear();
 
-			SyncState = OldSyncState;
-
 			// Export animations (only allow when exporting to file)
-		//	if (bFileExportInProgress)
+			//	if (bFileExportInProgress)
 			{
 				List<FAnimation> Animations = FAnimationExtractor.ExtractAnimations(SwAsmDoc, Root);
 				if (Animations != null)
@@ -238,6 +219,14 @@ namespace DatasmithSolidworks
 			}
 		}
 
+		public override FMeshData ExtractComponentMeshData(Component2 Comp)
+		{
+			FObjectMaterials ComponentMaterials = GetComponentMaterials(Comp);
+			ConcurrentBag<FBody> Bodies = FBody.FetchBodies(Comp);
+			FMeshData MeshData = FStripGeometry.CreateMeshData(Bodies, ComponentMaterials);
+			return MeshData;
+		}
+
 		public override bool HasMaterialUpdates()
 		{
 			// Dig into part level materials (they wont be read by LoadDocumentMaterials)
@@ -245,7 +234,9 @@ namespace DatasmithSolidworks
 			AllExportedComponents.UnionWith(SyncState.CleanComponents);
 			AllExportedComponents.UnionWith(SyncState.DirtyComponents.Keys);
 
-			ConcurrentDictionary<FComponentName, FObjectMaterials> CurrentDocMaterialsMap = FObjectMaterials.LoadAssemblyMaterials(this, AllExportedComponents, swDisplayStateOpts_e.swThisDisplayState, null);
+			ConcurrentDictionary<FComponentName, FObjectMaterials> CurrentDocMaterialsMap =
+				FObjectMaterials.LoadAssemblyMaterials(this, AllExportedComponents,
+					swDisplayStateOpts_e.swThisDisplayState, null);
 
 			if (CurrentDocMaterialsMap == null && SyncState.ComponentsMaterialsMap == null)
 			{
@@ -257,6 +248,7 @@ namespace DatasmithSolidworks
 				{
 					SetComponentDirty(KVP.Key, EComponentDirtyState.Material);
 				}
+
 				return true;
 			}
 			else if (CurrentDocMaterialsMap != null && SyncState.ComponentsMaterialsMap == null)
@@ -265,6 +257,7 @@ namespace DatasmithSolidworks
 				{
 					SetComponentDirty(KVP.Key, EComponentDirtyState.Material);
 				}
+
 				return true;
 			}
 			else
@@ -292,7 +285,9 @@ namespace DatasmithSolidworks
 								Component2 Comp = SyncState.ExportedComponentsMap[CompName];
 								bShouldSyncComponentMaterial = !Comp.IsSuppressed() && (Comp.Visible == (int)swComponentVisibilityState_e.swComponentVisible);
 							}
-							catch { }
+							catch
+							{
+							}
 						}
 
 						if (bShouldSyncComponentMaterial)
@@ -324,7 +319,7 @@ namespace DatasmithSolidworks
 			}
 		}
 
-		private float[] GetComponentDatasmithTransform(Component2 InComponent)
+		public float[] GetComponentDatasmithTransform(Component2 InComponent)
 		{
 			MathTransform ComponentTransform = InComponent.GetTotalTransform(true);
 			if (ComponentTransform == null)
@@ -335,7 +330,7 @@ namespace DatasmithSolidworks
 			float[] DatasmithTransform = null;
 			if (ComponentTransform != null)
 			{
-				DatasmithTransform = MathUtils.ConvertFromSolidworksTransform(ComponentTransform, 100f/*GeomScale*/);
+				DatasmithTransform = MathUtils.ConvertFromSolidworksTransform(ComponentTransform, 100f /*GeomScale*/);
 			}
 			else
 			{
@@ -345,16 +340,15 @@ namespace DatasmithSolidworks
 			return DatasmithTransform;
 		}
 
-		/// <summary>
-		/// Get component transform in specified configuration
-		/// </summary>
+		// Get component transform in specified configuration
 		private float[] GetComponentDatasmithTransform(FConfigurationTree.FComponentTreeNode InNode,
 			FConfigurationTree.FComponentConfig ComponentConfig)
 		{
 			return ComponentConfig?.Transform ?? InNode.CommonConfig.Transform ?? MathUtils.TransformIdentity();
 		}
 
-		private void ExportComponentRecursive(Component2 InComponent, Component2 InParent, Dictionary<Component2, FActorName> OutMeshesToExportMap)
+		private void ExportComponentRecursive(Component2 InComponent, Component2 InParent,
+			FMeshes.FConfiguration OutMeshesToExportMap)
 		{
 			bool bHasDirtyTransform = false;
 
@@ -362,7 +356,8 @@ namespace DatasmithSolidworks
 			if (SyncState.ComponentsTransformsMap.ContainsKey(ComponentName))
 			{
 				float[] ComponentTm = GetComponentDatasmithTransform(InComponent);
-				bHasDirtyTransform = !MathUtils.TransformsAreEqual(SyncState.ComponentsTransformsMap[ComponentName], ComponentTm);
+				bHasDirtyTransform =
+					!MathUtils.TransformsAreEqual(SyncState.ComponentsTransformsMap[ComponentName], ComponentTm);
 			}
 
 			if (bHasDirtyTransform || !SyncState.CleanComponents.Contains(ComponentName))
@@ -380,7 +375,7 @@ namespace DatasmithSolidworks
 				{
 					ActorExportInfo.ParentName = new FComponentName(InParent).GetActorName();
 				}
-				
+
 				ActorExportInfo.bVisible = true;
 				ActorExportInfo.Type = Exporter.GetExportedActorType(ActorName) ?? EActorType.SimpleActor;
 				ActorExportInfo.Transform = GetComponentDatasmithTransform(InComponent);
@@ -392,7 +387,8 @@ namespace DatasmithSolidworks
 
 				if (!InComponent.IsSuppressed())
 				{
-					dynamic ComponentVisibility = InComponent.GetVisibilityInAsmDisplayStates((int)swDisplayStateOpts_e.swThisDisplayState, null);
+					dynamic ComponentVisibility =
+						InComponent.GetVisibilityInAsmDisplayStates((int)swDisplayStateOpts_e.swThisDisplayState, null);
 					if (ComponentVisibility != null)
 					{
 						int Visible = ComponentVisibility[0];
@@ -411,7 +407,7 @@ namespace DatasmithSolidworks
 
 				// ComponentDoc is null if component is suppressed or lightweight
 				ModelDoc2 ComponentDoc = InComponent.GetModelDoc2();
-				
+
 				if (ComponentDoc is PartDoc) // Also possible to test document type with swDocumentTypes_e.swDocPART
 				{
 					bool bFirstExport = !SyncState.ExportedComponentsMap.ContainsKey(ComponentName);
@@ -434,20 +430,18 @@ namespace DatasmithSolidworks
 					if (!SyncState.PartsMap.ContainsKey(PartPath))
 					{
 						// New part
-						int PartDocId = Addin.Instance.GetDocumentId(ComponentDoc as ModelDoc2);
-						SyncState.PartsMap[PartPath] = new FPartDocument(PartDocId, ComponentDoc as PartDoc, Exporter, this, ComponentName);
-						SyncState.PartsMap[PartPath].Init();
+						SyncState.PartsMap[PartPath] = AddTrackedAssemblyPart(InComponent, ComponentName);
 					}
 
 					SyncState.ComponentToPartMap[ComponentName] = PartPath;
 
 					// This component has associated part document -- treat is as a mesh actor
 					ActorExportInfo.Type = EActorType.MeshActor;
-					
-					OutMeshesToExportMap.Add(InComponent, ActorExportInfo.Name);
+
+					OutMeshesToExportMap.AddComponentWhichNeedsMesh(InComponent, ComponentName, ActorExportInfo.Name);
 				}
 
-				if (ActorExportInfo.Type == EActorType.MeshActor )
+				if (ActorExportInfo.Type == EActorType.MeshActor)
 				{
 					// todo: deduplicate mesh name calculation
 					ActorExportInfo.MeshName = FDatasmithExporter.SanitizeName($"{ActorExportInfo.Name}_Mesh");
@@ -477,17 +471,18 @@ namespace DatasmithSolidworks
 			FConfigurationExporter ConfigurationExporter,
 			FConfigurationTree.FComponentTreeNode InNode, FConfigurationTree.FComponentTreeNode InParent)
 		{
+			SetExportStatus(InNode.ComponentName.GetString());
 
-			FConfigurationTree.FComponentConfig ActiveComponentConfig = InNode.Configurations?.Find(Config => Config.ConfigName == ActiveConfigName);
+			FConfigurationTree.FComponentConfig ActiveConfig = InNode.Configurations?.Find(Config => Config.ConfigName == ActiveConfigName);
+			float[] Transform = GetComponentDatasmithTransform(InNode, ActiveConfig);
+			SyncState.ComponentsTransformsMap[InNode.ComponentName] = Transform;
 
-			if (NeedExportComponent(InNode, ActiveComponentConfig))
+			if (InNode.bGeometrySame)
 			{
-				SetExportStatus(InNode.ComponentName.GetString());
+				FActorName ActorName = InNode.ActorName;
 
 				FDatasmithActorExportInfo ActorExportInfo = new FDatasmithActorExportInfo();
 
-				FActorName ActorName = InNode.ActorName;
-
 				ActorExportInfo.Label = InNode.ComponentName.GetLabel();
 				ActorExportInfo.Name = ActorName;
 
@@ -495,14 +490,12 @@ namespace DatasmithSolidworks
 				{
 					ActorExportInfo.ParentName = InParent.GetActorName();
 				}
-				
+
 				ActorExportInfo.bVisible = true;
 				ActorExportInfo.Type = Exporter.GetExportedActorType(ActorName) ?? EActorType.SimpleActor;
-				ActorExportInfo.Transform = GetComponentDatasmithTransform(InNode, ActiveComponentConfig);
-
-				SyncState.ComponentsTransformsMap[InNode.ComponentName] = ActorExportInfo.Transform;
+				ActorExportInfo.Transform = Transform;
 
-				ActorExportInfo.bVisible = ActiveComponentConfig?.bVisible ?? InNode.CommonConfig.bVisible;
+				ActorExportInfo.bVisible = ActiveConfig?.bVisible ?? InNode.CommonConfig.bVisible;
 
 				if (InNode.IsPartComponent())
 				{
@@ -510,18 +503,74 @@ namespace DatasmithSolidworks
 					ActorExportInfo.Type = EActorType.MeshActor;
 				}
 
-				if (ActorExportInfo.Type == EActorType.MeshActor )
+				if (ActorExportInfo.Type == EActorType.MeshActor)
 				{
 					// todo: deduplicate mesh name calculation
-					ActorExportInfo.MeshName = FDatasmithExporter.SanitizeName($"{ActorExportInfo.Name}_Mesh");
+					ActorExportInfo.MeshName = FDatasmithExporter.SanitizeName(ConfigurationExporter.GetMeshName(InNode.ComponentName));
 				}
 
 				Exporter.ExportOrUpdateActor(ActorExportInfo);
+			}
+			else
+			{
+				Debug.Assert(InNode.IsPartComponent());  // Expecting only Part components to have 'mesh' variants
+
+				FActorName ParentName;
+				{
+					FActorName ActorName = InNode.ActorName;
+					ParentName = ActorName;
+
+					FDatasmithActorExportInfo ActorExportInfo = new FDatasmithActorExportInfo();
+
+					ActorExportInfo.Label = InNode.ComponentName.GetLabel();
+					ActorExportInfo.Name = ActorName;
+
+					if (InParent != null)
+					{
+						ActorExportInfo.ParentName = InParent.GetActorName();
+					}
+
+					ActorExportInfo.bVisible = true;
+					ActorExportInfo.Type = EActorType.SimpleActor;  // Actor for Component with Mesh Variants is a 'simple' actor(i.e. just a node which has children)
+					ActorExportInfo.Transform = Transform;
+
+					ActorExportInfo.bVisible = ActiveConfig?.bVisible ?? InNode.CommonConfig.bVisible;
+
+					Exporter.ExportOrUpdateActor(ActorExportInfo);
+				}
 
-				SyncState.CleanComponents.Add(InNode.ComponentName);
+				foreach (FConfigurationTree.FComponentConfig ComponentConfig in InNode.Configurations)
+				{
+					FActorName ActorName = ConfigurationExporter.GetMeshActorName(ComponentConfig.ConfigName, InNode.ComponentName);
+					string Label = ActorName.GetString();
+
+					FDatasmithActorExportInfo ActorExportInfo = new FDatasmithActorExportInfo();
+
+					ActorExportInfo.Label = Label;
+					ActorExportInfo.Name = ActorName;
+
+					ActorExportInfo.ParentName = ParentName;
+
+					ActorExportInfo.bVisible = true;
+					ActorExportInfo.Type = Exporter.GetExportedActorType(ActorName) ?? EActorType.SimpleActor;
+					ActorExportInfo.Transform = Transform;
+
+					SyncState.ComponentsTransformsMap[InNode.ComponentName] = ActorExportInfo.Transform;
+
+					ActorExportInfo.bVisible = ComponentConfig?.bVisible ?? InNode.CommonConfig.bVisible;
+
+					// This component has associated part document -- treat is as a mesh actor
+					ActorExportInfo.Type = EActorType.MeshActor;
+
+					;
+					// todo: deduplicate mesh name calculation
+					ActorExportInfo.MeshName = FDatasmithExporter.SanitizeName(ConfigurationExporter.GetMeshName(ComponentConfig.ConfigName, InNode.ComponentName));
+
+					Exporter.ExportOrUpdateActor(ActorExportInfo);
+				}
 			}
 
-			// AddExportedComponent(InNode);
+			SyncState.CleanComponents.Add(InNode.ComponentName);
 
 			// Export component children
 			foreach (FConfigurationTree.FComponentTreeNode Child in InNode.EnumChildren())
@@ -535,52 +584,54 @@ namespace DatasmithSolidworks
 			SyncState.ExportedComponentsMap[InNode.ComponentName] = InNode.Component;
 		}
 
-		public override bool NeedExportComponent(FConfigurationTree.FComponentTreeNode InComponent, FConfigurationTree.FComponentConfig ActiveComponentConfig)
+		public override bool NeedExportComponent(FConfigurationTree.FComponentTreeNode InComponent,
+			FConfigurationTree.FComponentConfig ActiveComponentConfig)
 		{
 			bool bHasDirtyTransform = false;
 			if (SyncState.ComponentsTransformsMap.ContainsKey(InComponent.ComponentName))
 			{
 				float[] ComponentTm = GetComponentDatasmithTransform(InComponent, ActiveComponentConfig);
 				bHasDirtyTransform =
-					!MathUtils.TransformsAreEqual(SyncState.ComponentsTransformsMap[InComponent.ComponentName], ComponentTm);
+					!MathUtils.TransformsAreEqual(SyncState.ComponentsTransformsMap[InComponent.ComponentName],
+						ComponentTm);
 			}
 
-			bool bNeedExportComponent = bHasDirtyTransform || !SyncState.CleanComponents.Contains(InComponent.ComponentName);
+			bool bNeedExportComponent =
+				bHasDirtyTransform || !SyncState.CleanComponents.Contains(InComponent.ComponentName);
 			return bNeedExportComponent;
 		}
 
-		public override bool NeedGeometryExport(FConfigurationTree.FComponentTreeNode InNode, FConfigurationTree.FComponentConfig ActiveComponentConfig)
+		public override void AddPartDocument(FConfigurationTree.FComponentTreeNode InNode)
 		{
-			bool bNeedsGeometryExport;
-			bool bFirstExport = !SyncState.ExportedComponentsMap.ContainsKey(InNode.ComponentName);
-			bNeedsGeometryExport = bFirstExport;
+			string PartPath = InNode.PartPath;
+			Component2 Component = InNode.Component;
+			FComponentName ComponentName = InNode.ComponentName;
 
-			if (!bFirstExport && SyncState.DirtyComponents.ContainsKey(InNode.ComponentName))
+			// Add part document to track its changes
+			// todo: replace this with only document notifications. Doesn't seem that anything else is needed (assembly component itself is used to extract geometry)
+			// Probably, we might use Part document info to identify components built from the same Part(beware - same Part can be different in different components - configured differently)
+			// Anyway, 'PartDocument' is a lot for now, something simple can be used here definitely
+			if (!SyncState.PartsMap.ContainsKey(PartPath))
 			{
-				uint DirtyState = SyncState.DirtyComponents[InNode.ComponentName];
-				bNeedsGeometryExport =
-					((DirtyState & (1u << (int)EComponentDirtyState.Material)) != 0) ||
-					((DirtyState & (1u << (int)EComponentDirtyState.Geometry)) != 0) ||
-					((DirtyState & (1u << (int)EComponentDirtyState.Delete)) != 0);
+				FPartDocument PartTracker = AddTrackedAssemblyPart(Component, ComponentName);
+				SyncState.PartsMap[PartPath] = PartTracker;
 			}
 
-			return bNeedsGeometryExport;
+			SyncState.ComponentToPartMap[ComponentName] = PartPath;
 		}
 
-		public override void AddPartDocument(FConfigurationTree.FComponentTreeNode InNode)
+		// Tracks changes in a Part that is a component in the assembly
+		private FPartDocument AddTrackedAssemblyPart(IComponent2 Component, FComponentName ComponentName)
 		{
-			//TODO this will be null for new part, think of more solid solution
-			if (!SyncState.PartsMap.ContainsKey(InNode.PartPath))
-			{
-				ModelDoc2 ComponentDoc = InNode.Component.GetModelDoc2();
-				// New part
-				int PartDocId = Addin.Instance.GetDocumentId(ComponentDoc as ModelDoc2);
-				SyncState.PartsMap[InNode.PartPath] = new FPartDocument(PartDocId, ComponentDoc as PartDoc, Exporter, this,
-					InNode.ComponentName);
-				SyncState.PartsMap[InNode.PartPath].Init();
-			}
+			ModelDoc2 ComponentDoc = Component.GetModelDoc2();
+			// New part
+			int PartDocId = Addin.Instance.GetDocumentId(ComponentDoc);
 
-			SyncState.ComponentToPartMap[InNode.ComponentName] = InNode.PartPath;
+			FPartDocument PartTracker = new FPartDocument(PartDocId, ComponentDoc as PartDoc,
+				Exporter, 
+				this,  // Propagate change notifications to this assembly
+				ComponentName);
+			return PartTracker;
 		}
 
 		public void SetComponentDirty(FComponentName InComponent, EComponentDirtyState InState)
@@ -603,9 +654,38 @@ namespace DatasmithSolidworks
 			SetComponentDirty(new FComponentName(InComponent), InState);
 		}
 
-		public override void Init()
+		public void ComponentDeleted(FComponentName ComponentName)
+		{
+			if (SyncState.ComponentToPartMap.ContainsKey(ComponentName) &&
+				!SyncState.ComponentsToDelete.Contains(ComponentName))
+			{
+				SyncState.ComponentsToDelete.Add(ComponentName);
+				SetComponentDirty(ComponentName, FAssemblyDocumentTracker.EComponentDirtyState.Delete);
+			}
+		}
+
+		// Active configuration changed
+		public void ActiveConfigChanged()
+		{
+			SyncState.CleanComponents?.Clear();
+			SyncState.DirtyComponents?.Clear();
+			SyncState.ExportedComponentsMap?.Clear();
+		}
+	}
+
+	// Handles Solidworks API document events and propagates them to the DocumentTracker
+	// todo: extract api which is used by the Tracker to receive events. This should help to be clear on how Tracker is controlled by external events
+	public class FAssemblyDocumentEvents
+	{
+		private readonly FAssemblyDocumentTracker AsmDocumentTracker;
+		private readonly FAssemblyDocument AssemblyDocument;
+		private AssemblyDoc SwAsmDoc => AssemblyDocument.SwAsmDoc;
+		private int DocId => AssemblyDocument.DocId;
+
+		public FAssemblyDocumentEvents(FAssemblyDocument InAssemblyDocument, FAssemblyDocumentTracker InAsmDocumentTracker)
 		{
-			base.Init();
+			AssemblyDocument = InAssemblyDocument;
+			AsmDocumentTracker = InAsmDocumentTracker;
 
 			SwAsmDoc.RegenNotify += new DAssemblyDocEvents_RegenNotifyEventHandler(OnRegenNotify);
 			SwAsmDoc.ActiveDisplayStateChangePreNotify += new DAssemblyDocEvents_ActiveDisplayStateChangePreNotifyEventHandler(OnActiveDisplayStateChangePreNotify);
@@ -649,10 +729,8 @@ namespace DatasmithSolidworks
 			SwAsmDoc.FileSaveAsNotify2 += new DAssemblyDocEvents_FileSaveAsNotify2EventHandler(OnFileSaveAsNotify2);
 		}
 
-		public override void Destroy()
+		public void Destroy()
 		{
-			base.Destroy();
-
 			SwAsmDoc.RegenNotify -= new DAssemblyDocEvents_RegenNotifyEventHandler(OnRegenNotify);
 			SwAsmDoc.ActiveDisplayStateChangePreNotify -= new DAssemblyDocEvents_ActiveDisplayStateChangePreNotifyEventHandler(OnActiveDisplayStateChangePreNotify);
 			SwAsmDoc.ActiveViewChangeNotify -= new DAssemblyDocEvents_ActiveViewChangeNotifyEventHandler(OnActiveViewChangeNotify);
@@ -715,7 +793,7 @@ namespace DatasmithSolidworks
 		{
 			if (InComponentObj is Component2 Comp)
 			{
-				SetComponentDirty(Comp, EComponentDirtyState.Visibility);
+				AsmDocumentTracker.SetComponentDirty(Comp, FAssemblyDocumentTracker.EComponentDirtyState.Visibility);
 			}
 			return 0;
 		}
@@ -724,7 +802,7 @@ namespace DatasmithSolidworks
 		{
 			if (InCompObject is Component2 Comp)
 			{
-				SetComponentDirty(Comp, EComponentDirtyState.Material);
+				AsmDocumentTracker.SetComponentDirty(Comp, FAssemblyDocumentTracker.EComponentDirtyState.Material);
 			}
 			return 0;
 		}
@@ -772,7 +850,7 @@ namespace DatasmithSolidworks
 
 		int OnComponentConfigurationChangeNotify(string componentName, string oldConfigurationName, string newConfigurationName)
 		{
-			SetComponentDirty(FComponentName.FromApiString(componentName), EComponentDirtyState.Geometry);
+			AsmDocumentTracker.SetComponentDirty(FComponentName.FromApiString(componentName), FAssemblyDocumentTracker.EComponentDirtyState.Geometry);
 			return 0;
 		}
 
@@ -780,17 +858,17 @@ namespace DatasmithSolidworks
 		{
 			// Check each exported component's transform for changes, since 
 			// this callback does not tell us what changed (and there's no other way to know that)!
-			foreach (var KVP in SyncState.ExportedComponentsMap)
+			foreach (var KVP in AsmDocumentTracker.SyncState.ExportedComponentsMap)
 			{
 				Component2 Comp = KVP.Value;
 				float[] PrevCompTransform = null;
-				if (SyncState.ComponentsTransformsMap.TryGetValue(new FComponentName(Comp), out PrevCompTransform))
+				if (AsmDocumentTracker.SyncState.ComponentsTransformsMap.TryGetValue(new FComponentName(Comp), out PrevCompTransform))
 				{
-					float[] CompTransform = GetComponentDatasmithTransform(Comp);
+					float[] CompTransform = AsmDocumentTracker.GetComponentDatasmithTransform(Comp);
 
 					if (!MathUtils.TransformsAreEqual(CompTransform, PrevCompTransform))
 					{
-						SetDirty(true);
+						AsmDocumentTracker.SetDirty(true);
 						break;
 					}
 				}
@@ -832,7 +910,7 @@ namespace DatasmithSolidworks
 		{
 			if (InComponentObj is Component2 Comp)
 			{
-				SetComponentDirty(new FComponentName(Comp), EComponentDirtyState.Visibility);
+				AsmDocumentTracker.SetComponentDirty(new FComponentName(Comp), FAssemblyDocumentTracker.EComponentDirtyState.Visibility);
 			}
 			return 0;
 		}
@@ -844,28 +922,18 @@ namespace DatasmithSolidworks
 
 		int OnDeleteItemNotify(int InEntityType, string InItemName)
 		{
-			FComponentName ComponentName = FComponentName.FromApiString(InItemName);
-
-			if (InEntityType == (int)swNotifyEntityType_e.swNotifyComponent && 
-				SyncState.ComponentToPartMap.ContainsKey(ComponentName) && 
-				!SyncState.ComponentsToDelete.Contains(ComponentName))
+			if (InEntityType == (int)swNotifyEntityType_e.swNotifyComponent)
 			{
-				SyncState.ComponentsToDelete.Add(ComponentName);
-				SetComponentDirty(ComponentName, EComponentDirtyState.Delete);
+				AsmDocumentTracker.ComponentDeleted(FComponentName.FromApiString(InItemName));
 			}
 			return 0;
 		}
 
 		int OnRenameItemNotify(int InEntityType, string InOldName, string InNewName)
 		{
-			FComponentName ComponentName = FComponentName.FromApiString(InOldName);
-
-			if (InEntityType == (int)swNotifyEntityType_e.swNotifyComponent &&
-			    SyncState.ComponentToPartMap.ContainsKey(ComponentName) &&
-			    !SyncState.ComponentsToDelete.Contains(ComponentName))
+			if (InEntityType == (int)swNotifyEntityType_e.swNotifyComponent)
 			{
-				SyncState.ComponentsToDelete.Add(ComponentName);
-				SetComponentDirty(ComponentName, EComponentDirtyState.Delete);
+				AsmDocumentTracker.ComponentDeleted(FComponentName.FromApiString(InOldName));
 			}
 			return 0;
 		}
@@ -874,7 +942,7 @@ namespace DatasmithSolidworks
 		{
 			if (InEntityType == (int)swNotifyEntityType_e.swNotifyComponent)
 			{
-				SetDirty(true);
+				AsmDocumentTracker.SetDirty(true);
 			}
 			return 0;
 		}
@@ -886,10 +954,8 @@ namespace DatasmithSolidworks
 
 		int OnActiveConfigChangeNotify()
 		{
-			SetDirty(true);
-			SyncState.CleanComponents?.Clear();
-			SyncState.DirtyComponents?.Clear();
-			SyncState.ExportedComponentsMap?.Clear();
+			AsmDocumentTracker.SetDirty(true);
+			AsmDocumentTracker.ActiveConfigChanged();
 			return 0;
 		}
 
@@ -941,7 +1007,7 @@ namespace DatasmithSolidworks
 				IComponent2 Comp = ObjComp as IComponent2;
 				if (Comp != null)
 				{
-					SetComponentDirty(new FComponentName(Comp), EComponentDirtyState.Transform);
+					AsmDocumentTracker.SetComponentDirty(new FComponentName(Comp), FAssemblyDocumentTracker.EComponentDirtyState.Transform);
 				}
 			}
 			return 0;
@@ -977,4 +1043,116 @@ namespace DatasmithSolidworks
 			return 0;
 		}
 	}
+
+	// Notifications for the top(synced) document, includes material update checker thread
+	public class FAssemblyDocumentNotifications
+	{
+		private readonly FAssemblyDocumentTracker DocumentTracker;
+		private readonly FAssemblyDocumentEvents Events;
+		private FMaterialUpdateChecker MaterialUpdateChecker;
+
+		public FAssemblyDocumentNotifications(FAssemblyDocument InAssemblyDocument, FAssemblyDocumentTracker InAsmDocumentTracker)
+		{
+			DocumentTracker = InAsmDocumentTracker;
+			Events = new FAssemblyDocumentEvents(InAssemblyDocument, InAsmDocumentTracker);
+		}
+
+		public void Destroy()
+		{
+			MaterialUpdateChecker?.Destroy();
+			Events.Destroy();
+		}
+
+		public void Start()
+		{
+			if (MaterialUpdateChecker == null)
+			{
+				MaterialUpdateChecker = new FMaterialUpdateChecker(DocumentTracker);
+			}
+		}
+
+		public void Resume()
+		{
+			MaterialUpdateChecker?.Resume();
+		}
+
+		public void Pause()
+		{
+			MaterialUpdateChecker?.Pause();
+		}
+	}
+
+	// See IDocumentSyncer interface, implemented for Assembly
+	public class FAssemblyDocumentSyncer : IDocumentSyncer
+	{
+		public readonly FAssemblyDocumentNotifications Notifications;
+		public readonly FAssemblyDocumentTracker DocumentTracker;
+
+		public FAssemblyDocumentSyncer(FAssemblyDocument InDoc, AssemblyDoc InSwDoc, FDatasmithExporter InExporter)
+		{
+			DocumentTracker = new FAssemblyDocumentTracker(InDoc, InSwDoc, InExporter);
+			Notifications = new FAssemblyDocumentNotifications(InDoc, DocumentTracker);
+		}
+
+		public void Destroy()
+		{
+			Notifications.Destroy();
+			DocumentTracker.Destroy();
+		}
+
+		public void Start()
+		{
+			Notifications.Start();
+		}
+
+		public void Resume()
+		{
+			Notifications.Resume();
+		}
+
+		public void Pause()  
+		{
+			Notifications.Pause();
+		}
+
+		public bool GetDirty()  // i.e. in autosync
+		{
+			return DocumentTracker.GetDirty();
+		}
+
+		public void SetDirty(bool bInDirty)  // After Sync, where called with true??? Probably not needed at all(i.e. all SetDirty is inside in notifiers, and with false can be in Sync)
+		{
+			DocumentTracker.SetDirty(bInDirty);
+		}
+
+		public void Sync(string InOutputPath)
+		{
+			DocumentTracker.Sync(InOutputPath);
+		}
+
+		public FDatasmithFacadeScene GetDatasmithScene()
+		{
+			return DocumentTracker.DatasmithScene;
+		}
+	}
+
+	public class FAssemblyDocument : FDocument
+	{
+		public AssemblyDoc SwAsmDoc { get; private set; } = null;
+
+		public FAssemblyDocument(int InDocId, AssemblyDoc InSwDoc, FDatasmithExporter InExporter) : base(InDocId, InSwDoc as ModelDoc2)
+		{
+			SwAsmDoc = InSwDoc;
+
+			DocumentSyncer = new FAssemblyDocumentSyncer(this, SwAsmDoc, InExporter);
+		}
+
+		public override void Export(string InFilePath)
+		{
+			FAssemblyDocumentTracker Tracker = new FAssemblyDocumentTracker(this, SwAsmDoc, null);
+			Tracker.DatasmithFileExportPath = InFilePath;
+			Tracker.Export();
+		}
+
+	}
 }
\ No newline at end of file
diff --git a/Engine/Source/Programs/Enterprise/Datasmith/DatasmithSolidworksExporter/Private/ConfigurationExporter.cs b/Engine/Source/Programs/Enterprise/Datasmith/DatasmithSolidworksExporter/Private/ConfigurationExporter.cs
--- a/Engine/Source/Programs/Enterprise/Datasmith/DatasmithSolidworksExporter/Private/ConfigurationExporter.cs
+++ b/Engine/Source/Programs/Enterprise/Datasmith/DatasmithSolidworksExporter/Private/ConfigurationExporter.cs
@@ -1,12 +1,14 @@
 // Copyright Epic Games, Inc. All Rights Reserved.
 
 using System;
+using System.Collections;
 using System.Collections.Generic;
 using SolidWorks.Interop.sldworks;
 using SolidWorks.Interop.swconst;
 using System.Runtime.InteropServices;
 using System.Collections.Concurrent;
 using System.Linq;
+using static DatasmithSolidworks.FConfigurationTree;
 
 namespace DatasmithSolidworks
 {
@@ -40,7 +42,7 @@ namespace DatasmithSolidworks
 	[ComVisible(false)]
 	public class FMeshes
 	{
-		private Dictionary<string, IConfiguration> Configurations = new Dictionary<string, IConfiguration>();
+		private Dictionary<string, FConfiguration> Configurations = new Dictionary<string, FConfiguration>();
 
 		private HashSet<FMesh> Meshes = new HashSet<FMesh>();
 
@@ -58,6 +60,11 @@ namespace DatasmithSolidworks
 
 				Hash = 0;
 
+				if (InMeshData == null)
+				{
+					return;
+				}
+
 				foreach (FVec3 Vertex in MeshData.Vertices)
 				{
 					Hash ^= Vertex.GetHashCode();
@@ -131,56 +138,95 @@ namespace DatasmithSolidworks
 			}
 		}
 
-		[ComVisible(false)]
-		public interface IConfiguration
-		{
-			void AddMesh(Component2 Component, FMeshData Mesh);
-			bool GetMeshForComponent(FComponentName ComponentName, out FMesh Result);
-		}
-
-		public IConfiguration GetMeshesConfiguration(string ConfigurationName)
+		public FConfiguration GetMeshesConfiguration(string ConfigurationName)
 		{
 			if (Configurations.TryGetValue(ConfigurationName, out var Configuration))
 			{
 				return Configuration;
 			}			
 
-			Configuration = new FConfigurationImpl(this, ConfigurationName);
+			Configuration = new FConfiguration(this, ConfigurationName);
 			Configurations.Add(ConfigurationName, Configuration);
 			return Configuration;
 		}
 
-		public bool IsSameMesh(FComponentName ComponentName, string ConfigNameA, string ConfigNameB)
+		public bool DoesComponentHaveVisibleButDifferentMeshesInConfigs(FComponentName ComponentName, string ConfigNameA, string ConfigNameB)
 		{
-			return GetMeshesConfiguration(ConfigNameA).GetMeshForComponent(ComponentName, out FMesh MeshA)
-			       && GetMeshesConfiguration(ConfigNameB).GetMeshForComponent(ComponentName, out FMesh MeshB)
-			       && MeshA == MeshB;
-		}
+			bool bHasMeshA = GetMeshesConfiguration(ConfigNameA).GetMeshForComponent(ComponentName, out FMesh MeshA);
+			bool bHasMeshB = GetMeshesConfiguration(ConfigNameB).GetMeshForComponent(ComponentName, out FMesh MeshB);
 
+			return (bHasMeshA && bHasMeshB) && (MeshA != MeshB);
+		}
 
 		public FMeshData GetMeshData(string ConfigName, FComponentName ComponentName)
 		{
 			return GetMeshesConfiguration(ConfigName).GetMeshForComponent(ComponentName, out FMesh Mesh) ? Mesh.MeshData : null;
 		}
 
-		private class FConfigurationImpl : IConfiguration
+		public class FConfiguration
 		{
-			private Dictionary<FComponentName, FMesh> MeshForComponent = new Dictionary<FComponentName, FMesh>();
+			public class FComponentMesh
+			{
+				public Component2 Component;
+				public FActorName ActorName;
+				public FMesh Mesh = new FMesh();
+			}
+
+			private ConcurrentDictionary<FComponentName, FComponentMesh> MeshForComponent = new ConcurrentDictionary<FComponentName, FComponentMesh>();
 
-			public FConfigurationImpl(FMeshes Meshes, string Name)
+			public FConfiguration(FMeshes Meshes, string Name)
 			{
 			}
 
-			public void AddMesh(Component2 Component, FMeshData Mesh)
+			public void AddComponentWhichNeedsMesh(Component2 InComponent, FComponentName InComponentName, FActorName InActorName)
 			{
-				MeshForComponent.Add(new FComponentName(Component), new FMesh(Mesh));
+				MeshForComponent.TryAdd(InComponentName, new FComponentMesh() {Component = InComponent, ActorName = InActorName});
+			}
+
+			public void AddMesh(Component2 Component, FMeshData MeshData)
+			{
+				FComponentMesh ComponentMesh;
+				if (MeshForComponent.TryGetValue(new FComponentName(Component), out ComponentMesh))
+				{
+					ComponentMesh.Mesh = new FMesh(MeshData);
+				}
 			}
 
 			public bool GetMeshForComponent(FComponentName ComponentName, out FMesh Result)
 			{
-				return MeshForComponent.TryGetValue(ComponentName, out Result);
+				if (MeshForComponent.TryGetValue(ComponentName, out FComponentMesh ComponentMesh))
+				{
+					Result = ComponentMesh.Mesh;
+					return true;
+				}
+				Result = new FMesh(null);
+				return false;
+			}
+
+			public IEnumerable<FComponentName> EnumerateComponentNames()
+			{
+				return MeshForComponent.Keys;
 			}
 
+			public IEnumerable<KeyValuePair<Component2, FActorName>> EnumerateComponentAndActorNames()
+			{
+				foreach (KeyValuePair<FComponentName, FComponentMesh> FVP in MeshForComponent)
+				{
+					yield return new KeyValuePair<Component2, FActorName>(FVP.Value.Component, FVP.Value.ActorName);
+				}
+			}
+		}
+
+		public string GetConfigurationNameForComponentMesh(FComponentName ComponentName)
+		{
+			foreach (KeyValuePair<string, FConfiguration> KVP in Configurations)
+			{
+				if (KVP.Value.GetMeshForComponent(ComponentName, out FMesh Mesh))
+				{
+					return KVP.Key;
+				}
+			}
+			return null;
 		}
 	}
 
@@ -189,16 +235,18 @@ namespace DatasmithSolidworks
 	{
 		public readonly FMeshes Meshes;
 		public readonly string[] ConfigurationNames;
+		private readonly string MainConfigurationName;
 
 		public FConfigurationTree.FComponentTreeNode CombinedTree;
 
-		public FConfigurationExporter(FMeshes InMeshes, string[] InConfigurationNames)
+		public FConfigurationExporter(FMeshes InMeshes, string[] InConfigurationNames, string InMainConfigurationName)
 		{
 			Meshes = InMeshes;
 			ConfigurationNames = InConfigurationNames;
+			MainConfigurationName = InMainConfigurationName;
 		}
 
-		public List<FConfigurationData> ExportConfigurations(FDocument InDoc)
+		public List<FConfigurationData> ExportConfigurations(FDocumentTracker InDoc)
 		{
 			// todo: may refactor to union behavior for 'no configuration'(i.e. active) and one-to-many configurations export
 			// Not sure how CfgNames might be technically null though
@@ -217,13 +265,28 @@ namespace DatasmithSolidworks
 			ConfigurationManager ConfigManager = InDoc.SwDoc.ConfigurationManager;
 			IConfiguration OriginalConfiguration = ConfigManager.ActiveConfiguration;
 
-			List<string> ExportedConfigurationNames = new List<string>();
+			string OriginalConfigurationName = ConfigManager.ActiveConfiguration.Name;
 
-			foreach (string CfgName in ConfigurationNames)
+			Dictionary<string, string> ExportedConfigurationNames = new Dictionary<string, string>();
+
+			// Enumerate configurations starting with Active - to avoid redundant config switching(which takes time) and,
+			// additionally, to export meshes in Active configuration without configuration name suffixes
+			List<string> ConfigurationNamesActiveFirst = ConfigurationNames.ToList();
+			ConfigurationNamesActiveFirst.Remove(OriginalConfigurationName);
+			ConfigurationNamesActiveFirst.Insert(0, OriginalConfigurationName);
+
+			// Keep configurations order
+			Dictionary<string, int> ConfigurationOrder = ConfigurationNames.Select((Name, Index) => new { Name, Index }).ToDictionary(V => V.Name, V => V.Index);
+
+			bool bIsActiveConfiguration = true;
+			foreach (string CfgName in ConfigurationNamesActiveFirst)
 			{
 				IConfiguration swConfiguration = InDoc.SwDoc.GetConfigurationByName(CfgName) as IConfiguration;
 
-				InDoc.SwDoc.ShowConfiguration2(CfgName);
+				if (!bIsActiveConfiguration)
+				{
+					InDoc.SwDoc.ShowConfiguration2(CfgName);
+				}
 
 				string ConfigName = CfgName;
 
@@ -241,14 +304,15 @@ namespace DatasmithSolidworks
 
 				ConfigNode.Children = new List<FConfigurationTree.FComponentTreeNode>();
 
-				Dictionary<Component2, FActorName> MeshesToExportMap = new Dictionary<Component2, FActorName>();
+				FMeshes.FConfiguration MeshesConfiguration = Meshes.GetMeshesConfiguration(ConfigurationName: CfgName);
 
 				// Build the tree and get default materials (which aren't affected by any configuration)
 				// Use GetRootComponent3() with Resolve = true to ensure suppressed components will be loaded
 				// todo: docs says that Part document SW returns null. Not the case. But probably better to add a guard
-				CollectComponentsRecursive(InDoc, swConfiguration.GetRootComponent3(true), ConfigNode, MeshesToExportMap);
+				CollectComponentsRecursive(InDoc, swConfiguration.GetRootComponent3(true), ConfigNode, MeshesConfiguration, CfgName);
+
 
-				ExportedConfigurationNames.Add(ConfigName);
+				ExportedConfigurationNames.Add(CfgName, ConfigName);
 
 				Dictionary<FComponentName, FObjectMaterials> MaterialsMap = GetComponentMaterials(InDoc, DisplayStates != null ? DisplayStates[0] : null, swConfiguration);
 				SetComponentTreeMaterials(ConfigNode, MaterialsMap, null, false);
@@ -269,9 +333,8 @@ namespace DatasmithSolidworks
 				// Export materials
 				InDoc.SetExportStatus($"Component Materials");
 				HashSet<FComponentName> ComponentNamesToExportSet = new HashSet<FComponentName>();
-				foreach (var KVP in MeshesToExportMap)
+				foreach (FComponentName ComponentName in MeshesConfiguration.EnumerateComponentNames())
 				{
-					FComponentName ComponentName = new FComponentName(KVP.Key);
 					if (!ComponentNamesToExportSet.Contains(ComponentName))
 					{
 						ComponentNamesToExportSet.Add(ComponentName);
@@ -294,35 +357,14 @@ namespace DatasmithSolidworks
 
 				// Export meshes
 				InDoc.SetExportStatus($"Component Meshes");
-				ConcurrentDictionary<Component2, Tuple<FDatasmithFacadeMeshElement, FDatasmithFacadeMesh>> CreatedMeshes = new ConcurrentDictionary<Component2, Tuple<FDatasmithFacadeMeshElement, FDatasmithFacadeMesh>>();
-
-				foreach (KeyValuePair<Component2, FActorName> KVP in MeshesToExportMap)
-				{
-					Component2 Comp = KVP.Key;
 
-					ConcurrentBag<FBody> Bodies = FBody.FetchBodies(Comp);
-					FMeshData MeshData = FStripGeometry.CreateMeshData(Bodies, InDoc.GetComponentMaterials(Comp));
-
-					if (MeshData != null)
-					{
-						Tuple<FDatasmithFacadeMeshElement, FDatasmithFacadeMesh> NewMesh = null;
-						if (InDoc.Exporter.ExportMesh($"{KVP.Value}_Mesh", MeshData, KVP.Value, out NewMesh))
-						{
-							CreatedMeshes[Comp] = NewMesh;
-						}
-					}
-				}
-
-				foreach (var KVP in CreatedMeshes)
-				{
-					FComponentName ComponentName = new FComponentName(KVP.Key);
-					string MeshName = InDoc.Exporter.AddMesh(KVP.Value);
-
-					InDoc.AddMeshForComponent(ComponentName, MeshName);
-				}
+				string MeshSuffix = bIsActiveConfiguration ? null : CfgName; // Suffix for main configuration is absent(so in default or single configuration export we don't have unnecessary long names)
+				InDoc.ProcessConfigurationMeshes(MeshesConfiguration, MeshSuffix);
 
 				// Combine separate scene trees into the single one with configuration-specific data
-				FConfigurationTree.Merge(CombinedTree, ConfigNode, CfgName);
+				FConfigurationTree.Merge(CombinedTree, ConfigNode, CfgName, bIsActiveConfiguration);
+
+				bIsActiveConfiguration = false;  // Next enumerated configurations will be switched to(they weren't active on export start)
 			}
 
 			if (OriginalConfiguration != null)
@@ -352,8 +394,12 @@ namespace DatasmithSolidworks
 
 			List<FConfigurationData> FlatConfigurationData = new List<FConfigurationData>();
 
+
+			// Order exported variants by original configurations order
+			IEnumerable<string> ExportedConfigurationNamesOrdered = ExportedConfigurationNames.OrderBy(V => ConfigurationOrder[V.Key]).Select(V => V.Value);
+
 			// Add roots for configurations
-			foreach (string CfgName in ExportedConfigurationNames)
+			foreach (string CfgName in ExportedConfigurationNamesOrdered)
 			{
 				FConfigurationData CfgData = new FConfigurationData();
 				CfgData.Name = CfgName;
@@ -381,14 +427,14 @@ namespace DatasmithSolidworks
 			return FlatConfigurationData;
 		}
 
-		private static Dictionary<FComponentName, FObjectMaterials> GetComponentMaterials(FDocument InDoc, string InDisplayState, IConfiguration InConfiguration)
+		private static Dictionary<FComponentName, FObjectMaterials> GetComponentMaterials(FDocumentTracker InDoc, string InDisplayState, IConfiguration InConfiguration)
 		{
 			Dictionary<FComponentName, FObjectMaterials> MaterialsMap = new Dictionary<FComponentName, FObjectMaterials>();
 
 			swDisplayStateOpts_e Option = InDisplayState != null ? swDisplayStateOpts_e.swSpecifyDisplayState : swDisplayStateOpts_e.swThisDisplayState;
 			string[] DisplayStates = InDisplayState != null ? new string[] { InDisplayState } : null;
 
-			if (InDoc is FAssemblyDocument AsmDoc)
+			if (InDoc is FAssemblyDocumentTracker AsmDoc)
 			{
 				HashSet<FComponentName> InComponentsSet = new HashSet<FComponentName>();
 
@@ -410,8 +456,8 @@ namespace DatasmithSolidworks
 			}
 			else
 			{
-				FPartDocument PartDocument = InDoc as FPartDocument;
-				FObjectMaterials PartMaterials = FObjectMaterials.LoadPartMaterials(InDoc, PartDocument.SwPartDoc, Option, DisplayStates);
+				FPartDocumentTracker PartDocumentTracker = InDoc as FPartDocumentTracker;
+				FObjectMaterials PartMaterials = FObjectMaterials.LoadPartMaterials(InDoc, PartDocumentTracker.SwPartDoc, Option, DisplayStates);
 				Component2 Comp = InConfiguration.GetRootComponent3(true);
 
 				MaterialsMap.Add(FComponentName.FromCustomString(Comp?.Name2 ?? ""), PartMaterials);
@@ -429,7 +475,7 @@ namespace DatasmithSolidworks
 				TargetConfig = InComponentTree.GetConfiguration(InConfigurationName, bIsDisplayState);
 				if (TargetConfig == null)
 				{
-					TargetConfig = InComponentTree.AddConfiguration(InConfigurationName, bIsDisplayState);
+					TargetConfig = InComponentTree.AddConfiguration(InConfigurationName, bIsDisplayState, bIsActiveConfiguration: false);
 				}
 			}
 			else
@@ -449,8 +495,8 @@ namespace DatasmithSolidworks
 			}
 		}
 
-		private static void CollectComponentsRecursive(FDocument InDoc, Component2 InComponent,
-			FConfigurationTree.FComponentTreeNode InParentNode, Dictionary<Component2, FActorName> OutMeshesToExportMap)
+		private static void CollectComponentsRecursive(FDocumentTracker InDoc, Component2 InComponent,
+			FComponentTreeNode InParentNode, FMeshes.FConfiguration Meshes, string CfgName)
 		{
 			FConfigurationTree.FComponentTreeNode NewNode = new FConfigurationTree.FComponentTreeNode(InComponent);
 			InParentNode.Children.Add(NewNode);
@@ -500,7 +546,7 @@ namespace DatasmithSolidworks
 				foreach (object ObjChild in Children)
 				{
 					Component2 Child = (Component2)ObjChild;
-					CollectComponentsRecursive(InDoc, Child, NewNode, OutMeshesToExportMap);
+					CollectComponentsRecursive(InDoc, Child, NewNode, Meshes, CfgName);
 				}
 
 				// todo: sorting ensures that component order is stable in export(e.g. in variant property bindings order)
@@ -511,24 +557,37 @@ namespace DatasmithSolidworks
 
 			if (InDoc.NeedExportComponent(NewNode, NewNode.CommonConfig))
 			{
-				if (NewNode.IsPartComponent() && InDoc.NeedGeometryExport(NewNode, NewNode.CommonConfig))
+				if (NewNode.IsPartComponent())
 				{
 					// Collect meshes to export when it's a part component
 					InDoc.AddPartDocument(NewNode);
-					OutMeshesToExportMap.Add(NewNode.Component, NewNode.ActorName);
+					Meshes.AddComponentWhichNeedsMesh(NewNode.Component, NewNode.ComponentName, NewNode.ActorName);
 				}
 			}
 			InDoc.AddExportedComponent(NewNode);
 		}
 
-		public bool IsSameMesh(FComponentName ComponentName, string ConfigNameA, string ConfigNameB)
+		public bool DoesComponentHaveVisibleButDifferentMeshesInConfigs(FComponentName ComponentName, string ConfigNameA, string ConfigNameB)
+		{
+			return Meshes.DoesComponentHaveVisibleButDifferentMeshesInConfigs(ComponentName, ConfigNameA, ConfigNameB);
+		}
+
+		// name/label for variant mesh actor 
+		public FActorName GetMeshActorName(string ConfigName, FComponentName ComponentName)
+		{
+			return FActorName.FromString($"{ComponentName}_{ConfigName}");
+		}
+
+		public string GetMeshName(string ConfigName, FComponentName ComponentName)
 		{
-			return Meshes.IsSameMesh(ComponentName, ConfigNameA, ConfigNameB);
+			return ConfigName == MainConfigurationName ? $"{ComponentName}_Mesh" : $"{ComponentName}_{ConfigName}_Mesh";
 		}
 
-		public FActorName GetMeshActorName(FComponentName ComponentName, string ConfigName)
+		// Get mesh for component which has only single mesh configuration
+		public string GetMeshName(FComponentName ComponentName)
 		{
-			return FActorName.FromString(ComponentName.GetString() + "_" + ConfigName);
+			string ConfigName = Meshes.GetConfigurationNameForComponentMesh(ComponentName);
+			return ConfigName == MainConfigurationName ? $"{ComponentName}_Mesh" : $"{ComponentName}_{ConfigName}_Mesh";
 		}
 	}
 }
\ No newline at end of file
diff --git a/Engine/Source/Programs/Enterprise/Datasmith/DatasmithSolidworksExporter/Private/ConfigurationTree.cs b/Engine/Source/Programs/Enterprise/Datasmith/DatasmithSolidworksExporter/Private/ConfigurationTree.cs
--- a/Engine/Source/Programs/Enterprise/Datasmith/DatasmithSolidworksExporter/Private/ConfigurationTree.cs
+++ b/Engine/Source/Programs/Enterprise/Datasmith/DatasmithSolidworksExporter/Private/ConfigurationTree.cs
@@ -14,7 +14,8 @@ namespace DatasmithSolidworks
 		// Component state for a single configuration
 		public class FComponentConfig
 		{
-			public string ConfigName = null;
+			public readonly string ConfigName;
+			public readonly bool bIsMainActive;  // if this Configuration is the active config
 
 			public bool bVisible = true;
 			public bool bSuppressed = false;
@@ -28,6 +29,12 @@ namespace DatasmithSolidworks
 
 			public List<Body2> Bodies = null;
 
+			public FComponentConfig(string InConfigName, bool bInIsMainActive)
+			{
+				ConfigName = InConfigName;
+				bIsMainActive = bInIsMainActive;
+			}
+
 			public void CopyFrom(FComponentConfig InOther)
 			{
 				bVisible = InOther.bVisible;
@@ -37,6 +44,7 @@ namespace DatasmithSolidworks
 				RelativeTransform = InOther.RelativeTransform;
 				Bodies = InOther.Bodies;
 			}
+
 		};
 
 		// A node contains all the configuration data for a single component
@@ -62,7 +70,7 @@ namespace DatasmithSolidworks
 			public string PartPath => ComponentInfo.PartPath;
 
 			// Common configuration data
-			public FComponentConfig CommonConfig = new FComponentConfig();
+			public FComponentConfig CommonConfig = new FComponentConfig(null, true);
 
 			// Per-configuration data
 			public List<FComponentConfig> Configurations = null;
@@ -95,7 +103,7 @@ namespace DatasmithSolidworks
 				Component = InComponent;
 			}
 
-			public FComponentConfig AddConfiguration(string InConfigurationName, bool bIsDisplayState)
+			public FComponentConfig AddConfiguration(string InConfigurationName, bool bIsDisplayState, bool bIsActiveConfiguration)
 			{
 				List<FComponentConfig> TargetList = null;
 
@@ -115,8 +123,7 @@ namespace DatasmithSolidworks
 					}
 					TargetList = Configurations;
 				}
-				FComponentConfig Result = new FComponentConfig();
-				Result.ConfigName = InConfigurationName;
+				FComponentConfig Result = new FComponentConfig(InConfigurationName, bIsActiveConfiguration);
 				TargetList.Add(Result);
 
 				return Result;
@@ -163,7 +170,8 @@ namespace DatasmithSolidworks
 			}
 		};
 
-		static public void Merge(FComponentTreeNode OutCombined, FComponentTreeNode InTree, string InConfigurationName)
+		static public void Merge(FComponentTreeNode OutCombined, FComponentTreeNode InTree, string InConfigurationName,
+			bool bIsActiveConfiguration)
 		{
 			foreach (FComponentTreeNode Child in InTree.EnumChildren())
 			{
@@ -202,11 +210,11 @@ namespace DatasmithSolidworks
 				}
 
 				// Make a NodeConfig and copy parameter values from 'tree' node
-				FComponentConfig NodeConfig = CombinedChild.AddConfiguration(InConfigurationName, false);
+				FComponentConfig NodeConfig = CombinedChild.AddConfiguration(InConfigurationName, false, bIsActiveConfiguration);
 				NodeConfig.CopyFrom(Child.CommonConfig);
 
 				// Recurse to children
-				Merge(CombinedChild, Child, InConfigurationName);
+				Merge(CombinedChild, Child, InConfigurationName, bIsActiveConfiguration);
 			}
 		}
 
@@ -293,7 +301,7 @@ namespace DatasmithSolidworks
 						bSuppressionSame = false;
 					}
 
-					if (!ConfigurationExporter.IsSameMesh(InNode.ComponentName, InNode.Configurations[0].ConfigName, InNode.Configurations[Idx].ConfigName))
+					if (ConfigurationExporter.DoesComponentHaveVisibleButDifferentMeshesInConfigs(InNode.ComponentName, InNode.Configurations[0].ConfigName, InNode.Configurations[Idx].ConfigName))
 					{
 						bGeometrySame = false;
 					}
@@ -306,8 +314,7 @@ namespace DatasmithSolidworks
 
 					foreach (FComponentConfig ComponentConfig in InNode.Configurations)
 					{
-						InNode.Meshes.Add(ConfigurationExporter.GetMeshActorName(InNode.ComponentName,
-							ComponentConfig.ConfigName));
+						InNode.Meshes.Add(ConfigurationExporter.GetMeshActorName(ComponentConfig.ConfigName, InNode.ComponentName));
 					}
 				}
 
@@ -390,7 +397,7 @@ namespace DatasmithSolidworks
 				{
 					OutData.ComponentGeometry.Add(InNode.ComponentName, new FConfigurationData.FComponentGeometryVariant
 						{
-							VisibleActor = ConfigurationExporter.GetMeshActorName(InNode.ComponentName, InConfigurationName),
+							VisibleActor = ConfigurationExporter.GetMeshActorName(InConfigurationName, InNode.ComponentName),
 							All = InNode.Meshes.ToList()
 						}
 					);
diff --git a/Engine/Source/Programs/Enterprise/Datasmith/DatasmithSolidworksExporter/Private/DatasmithExporter.cs b/Engine/Source/Programs/Enterprise/Datasmith/DatasmithSolidworksExporter/Private/DatasmithExporter.cs
--- a/Engine/Source/Programs/Enterprise/Datasmith/DatasmithSolidworksExporter/Private/DatasmithExporter.cs
+++ b/Engine/Source/Programs/Enterprise/Datasmith/DatasmithSolidworksExporter/Private/DatasmithExporter.cs
@@ -9,6 +9,7 @@ using System.IO;
 using System.Linq;
 using System.Runtime.InteropServices;
 using System.Threading.Tasks;
+using static DatasmithSolidworks.FConfigurationTree;
 
 namespace DatasmithSolidworks
 {
@@ -179,7 +180,7 @@ namespace DatasmithSolidworks
 	public class FDatasmithExporter
 	{
 		private Dictionary<FActorName, Tuple<EActorType, FDatasmithFacadeActor>> ExportedActorsMap = new Dictionary<FActorName, Tuple<EActorType, FDatasmithFacadeActor>>();
-		private ConcurrentDictionary<string, Tuple<FDatasmithFacadeMeshElement, FDatasmithFacadeMesh>> ExportedMeshesMap = new ConcurrentDictionary<string, Tuple<FDatasmithFacadeMeshElement, FDatasmithFacadeMesh>>();
+		private ConcurrentDictionary<string, FDatasmithFacadeMeshElement> ExportedMeshesMap = new ConcurrentDictionary<string, FDatasmithFacadeMeshElement>();
 		private ConcurrentDictionary<int, FDatasmithFacadeMaterialInstance> ExportedMaterialsMap = new ConcurrentDictionary<int, FDatasmithFacadeMaterialInstance>();
 		private ConcurrentDictionary<string, FDatasmithFacadeTexture> ExportedTexturesMap = new ConcurrentDictionary<string, FDatasmithFacadeTexture>();
 		private Dictionary<string, FDatasmithFacadeActorBinding> ExportedActorBindingsMap = new Dictionary<string, FDatasmithFacadeActorBinding>();
@@ -264,12 +265,12 @@ namespace DatasmithSolidworks
 			return Actor;
 		}
 
-		public string AddMesh(Tuple<FDatasmithFacadeMeshElement, FDatasmithFacadeMesh> NewMesh)
+		public string AddMesh(FDatasmithFacadeMeshElement NewMesh)
 		{
-			string Name = NewMesh.Item1.GetName();
+			string Name = NewMesh.GetName();
 			RemoveMesh(Name);
 			ExportedMeshesMap.TryAdd(Name, NewMesh);
-			DatasmithScene.AddMesh(NewMesh.Item1);
+			DatasmithScene.AddMesh(NewMesh);
 			return Name;
 		}
 
@@ -279,9 +280,9 @@ namespace DatasmithSolidworks
             {
                 return;
             }
-			if (ExportedMeshesMap.TryRemove(MeshName, out Tuple<FDatasmithFacadeMeshElement, FDatasmithFacadeMesh> OldMesh))
+			if (ExportedMeshesMap.TryRemove(MeshName, out FDatasmithFacadeMeshElement OldMesh))
 			{
-				DatasmithScene.RemoveMesh(OldMesh.Item1);
+				DatasmithScene.RemoveMesh(OldMesh);
 			}
 		}
 
@@ -356,13 +357,13 @@ namespace DatasmithSolidworks
 			}
 		}
 
-		public bool ExportMesh(string InMeshName, FMeshData InData, FActorName InUpdateMeshActor, out Tuple<FDatasmithFacadeMeshElement, FDatasmithFacadeMesh> OutMeshPair)
+		public bool ExportMesh(string InMeshName, FMeshData InData, FActorName InUpdateMeshActor, out FDatasmithFacadeMeshElement OutMeshElement)
 		{
 			Addin.Instance.LogDebug($"ExportMesh(MeshName: {InMeshName}, MeshActor: {InUpdateMeshActor.GetString()})");
 
 			InMeshName = SanitizeName(InMeshName);  //Compute mesh name early(it might be needed to remove old mesh)
 
-			OutMeshPair = null;
+			OutMeshElement = null;
 
 			if (InData.Vertices == null || InData.Normals == null || InData.TexCoords == null || InData.Triangles == null)
 			{
@@ -448,7 +449,7 @@ namespace DatasmithSolidworks
 				Addin.Instance.LogDebug($"  assigned materials: {string.Join(", ", MeshAddedMaterials)})");
 			}
 
-			OutMeshPair = new Tuple<FDatasmithFacadeMeshElement, FDatasmithFacadeMesh>(MeshElement, Mesh);
+			OutMeshElement = MeshElement;
 
 			DatasmithScene.ExportDatasmithMesh(MeshElement, Mesh);
 
diff --git a/Engine/Source/Programs/Enterprise/Datasmith/DatasmithSolidworksExporter/Private/Document.cs b/Engine/Source/Programs/Enterprise/Datasmith/DatasmithSolidworksExporter/Private/Document.cs
--- a/Engine/Source/Programs/Enterprise/Datasmith/DatasmithSolidworksExporter/Private/Document.cs
+++ b/Engine/Source/Programs/Enterprise/Datasmith/DatasmithSolidworksExporter/Private/Document.cs
@@ -7,66 +7,124 @@ using System.Collections.Generic;
 using System.Diagnostics;
 using System.IO;
 using System.Threading;
-using static DatasmithSolidworks.FAssemblyDocument;
+using System.Threading.Tasks;
+using static DatasmithSolidworks.FAssemblyDocumentTracker;
 
 namespace DatasmithSolidworks
 {
-	/**
-	 * Base document class.
-	 */
-	public abstract class FDocument
+	// Checks for material updates periodically, using a separate thread to execute checks
+	// todo: Solidworks API is not concurrent-friendly - https://forum.solidworks.com/thread/43129
+	//   need to test but it's probably faster to check material updates simply in Sync. Maybe 
+	//   with the same period as currently if needed so
+	public class FMaterialUpdateChecker
 	{
-		public ModelDoc2 SwDoc = null;
-		public bool bDirectLinkAutoSync { get; set; } = false;
-		public int DirectLinkSyncCount { get; private set; } = 0;
+		public Thread MaterialCheckerThread = null;
+		public bool bExitMaterialUpdateThread = false;
+		public ManualResetEvent MaterialCheckerEvent = null;
+
+		public FDocumentTracker DocumentTracker;
+
+		public FMaterialUpdateChecker(FDocumentTracker InDocumentTracker)
+		{
+			DocumentTracker = InDocumentTracker;
+
+			MaterialCheckerThread = new Thread(CheckForMaterialUpdatesProc);
+			MaterialCheckerEvent = new ManualResetEvent(false);
+			MaterialCheckerThread.Start();
+			MaterialCheckerEvent.Set();
+		}
+
+		public void Destroy()
+		{
+			bExitMaterialUpdateThread = true;
+
+			if (MaterialCheckerThread != null && !MaterialCheckerThread.Join(500))
+			{
+				MaterialCheckerThread.Abort();
+			}
+		}
+
+		public void Resume()
+		{
+			MaterialCheckerEvent.Set();
+		}
+
+		public void Pause()
+		{
+			MaterialCheckerEvent.Reset();
+		}
+
+		private void CheckForMaterialUpdatesProc()
+		{
+			while (!bExitMaterialUpdateThread)
+			{
+				MaterialCheckerEvent.WaitOne();
+
+				// Only check for material updates when we are not current exporting
+				if (!DocumentTracker.IsUpdateInProgress())
+				{
+					bool bSketchMode = false;
+					try
+					{
+						bSketchMode = (DocumentTracker.SwDoc.SketchManager?.ActiveSketch != null);
+					}
+					catch { }
+				
+					if (!bSketchMode && DocumentTracker.HasMaterialUpdates())
+					{
+						DocumentTracker.SetDirty(true);
+					}
+				}
+				Thread.Sleep(600);
+			}
+		}
+	}
+
+	// Build DatasmithScene from Solidworks Document, supporting scene changes/modifications
+	// Doesn't incorporate change notifications callbacks/handlers(i.e. notification handlers call into this class to tell what was changed and it updates Datasmith scene accordingly)
+	public abstract class FDocumentTracker
+	{
+		public readonly FDocument Doc;  // Top Document that is tracked/exported
+		public ModelDoc2 SwDoc => Doc.SwDoc;
+
 		public ConcurrentDictionary<int, FMaterial> ExportedMaterialsMap = new ConcurrentDictionary<int, FMaterial>();
 		public bool bHasConfigurations = false;
-		
-		protected FDatasmithFacadeScene DatasmithScene = null;
+
+		public FDatasmithFacadeScene DatasmithScene = null;
 		public FDatasmithExporter Exporter = null;
-		protected int DocId = -1;
-		protected bool bDirectLinkSyncInProgress = false;
-		protected bool bFileExportInProgress = false;
 
-		private string DirectLinkPath = "";
-		private FDatasmithFacadeDirectLink DatasmithDirectLink = null;
-		private bool bDocumentIsDirty = true;
-		private string DatasmithFileExportPath = "";
-		private Thread MaterialCheckerThread = null;
-		private ManualResetEvent MaterialCheckerEvent = null;
-		private bool bExitMaterialUpdateThread = false;
-		private uint FaceCounter = 1; // Face Id generator
+		public bool bDocumentIsDirty = true;
+		public string DatasmithFileExportPath = "";
 
+		public uint FaceCounter = 1; // Face Id generator
 
-		public FDocument(int InDocId, ModelDoc2 InSwDoc, FDatasmithExporter InExporter)
+		protected FDocumentTracker(FDocument InDoc, FDatasmithExporter InExporter)
 		{
-			DocId = InDocId;
-			SwDoc = InSwDoc;
+			Doc = InDoc;
+
 			DatasmithScene = new FDatasmithFacadeScene("StdMaterial", "Solidworks", "Solidworks", "");
-			DatasmithScene.SetName(InSwDoc.GetTitle());
+			DatasmithScene.SetName(InDoc.SwDoc.GetTitle());
 
-			DirectLinkPath = Path.Combine(Path.GetTempPath(), "sw_dl_" + Guid.NewGuid().ToString());
+			Exporter = InExporter ?? new FDatasmithExporter(DatasmithScene);
+		}
 
-			if (!Directory.Exists(DirectLinkPath))
-			{
-				Directory.CreateDirectory(DirectLinkPath);
-			}
+		public string GetPathName()
+		{
+			return SwDoc.GetPathName();
+		}
 
-			if (InExporter != null)
-			{
-				Exporter = InExporter;
-			}
-			else
-			{
-				Exporter = new FDatasmithExporter(DatasmithScene);
-			}
+		public bool GetDirty()
+		{
+			return bDocumentIsDirty;
 		}
 
-		public static bool IsValidFaceId(uint InFaceId)
+		public virtual void SetDirty(bool bInDirty)
 		{
-			return (InFaceId >> 24 == 0xAA);
+			bDocumentIsDirty = bInDirty;
 		}
 
+		public abstract bool HasMaterialUpdates();
+
 		public uint GetFaceId(IFace2 InFace)
 		{
 			uint FaceId = unchecked((uint)InFace.GetFaceId());
@@ -79,27 +137,19 @@ namespace DatasmithSolidworks
 			return FaceId;
 		}
 
-		protected bool GetDirty()
+		public static bool IsValidFaceId(uint InFaceId)
 		{
-			return bDocumentIsDirty;
+			return (InFaceId >> 24 == 0xAA);
 		}
 
-		protected virtual void SetDirty(bool bInDirty)
+		public bool IsUpdateInProgress()
 		{
-			bDocumentIsDirty = bInDirty;
+			return Doc.bDirectLinkSyncInProgress || Doc.bFileExportInProgress;
 		}
 
 		public void SetExportStatus(string InMessage)
 		{
-			if (bDirectLinkSyncInProgress)
-			{
-				Addin.Instance.LogStatusBarMessage($"DirectLink sync...{InMessage}");
-			}
-			else if (bFileExportInProgress)
-			{
-				string FileName = Path.GetFileName(DatasmithFileExportPath);
-				Addin.Instance.LogStatusBarMessage($"Exporting file {FileName}...{InMessage}");
-			}
+			Doc.SetExportStatus(InMessage);
 		}
 
 		private void ExportConfigurations(List<FConfigurationData> Configs)
@@ -121,31 +171,6 @@ namespace DatasmithSolidworks
 			}
 		}
 
-		private void CheckForMaterialUpdatesProc()
-		{
-			while (!bExitMaterialUpdateThread)
-			{
-				MaterialCheckerEvent.WaitOne();
-
-				// Only check for material updates when we are not current exporting
-				if (!bDirectLinkSyncInProgress && !bFileExportInProgress)
-				{
-					bool bSketchMode = false;
-					try
-					{
-						bSketchMode = (SwDoc.SketchManager?.ActiveSketch != null);
-					}
-					catch { }
-				
-					if (!bSketchMode && HasMaterialUpdates())
-					{
-						SetDirty(true);
-					}
-				}
-				Thread.Sleep(600);
-			}
-		}
-
 		private void ExportLights()
 		{
 			List<FLight> Lights  = FLightExporter.ExportLights(SwDoc);
@@ -156,90 +181,64 @@ namespace DatasmithSolidworks
 			}
 		}
 
-		private void RunExport(bool bIsDirectLinkExport)
+		public void Export()
 		{
-			if (bIsDirectLinkExport)
-			{
-				// DirectLink sync
-
-				bDirectLinkSyncInProgress = true;
-
-				DatasmithScene.SetName(SwDoc.GetTitle());
-				DatasmithScene.SetOutputPath(DirectLinkPath);
-
-#if DEBUG
-				Stopwatch Watch = Stopwatch.StartNew();
-#endif
-				FMeshes Meshes = new FMeshes();
-				
-				PreExport(Meshes, false);
-				ExportToDatasmithScene(Meshes);
+			string OutDir = Path.GetDirectoryName(DatasmithFileExportPath);
+			string CleanFileName = Path.GetFileNameWithoutExtension(DatasmithFileExportPath);
 
-				ExportLights();
+			// Save/restore scene and exporter in order not to mess up DirectLink state
+			FDatasmithFacadeScene OldScene = DatasmithScene;
+			FDatasmithExporter OldExporter = Exporter;
 
-#if DEBUG
-				Watch.Stop();
-				Debug.WriteLine($"EXPORT TIME: {(double)Watch.ElapsedMilliseconds / 1000.0}");
-#endif
-
-				DatasmithDirectLink.UpdateScene(DatasmithScene);
-			}
-			else
-			{
-				// Export to file
-				bFileExportInProgress = true;
-
-				string OutDir = Path.GetDirectoryName(DatasmithFileExportPath);
-				string CleanFileName = Path.GetFileNameWithoutExtension(DatasmithFileExportPath);
+			DatasmithScene = new FDatasmithFacadeScene("StdMaterial", "Solidworks", "Solidworks", "");
+			DatasmithScene.SetName(CleanFileName);
+			DatasmithScene.SetOutputPath(OutDir);
 
-				// Save/restore scene and exporter in order not to mess up DirectLink state
-				FDatasmithFacadeScene OldScene = DatasmithScene;
-				FDatasmithExporter OldExporter = Exporter;
+			Exporter = new FDatasmithExporter(DatasmithScene);
 
-				DatasmithScene = new FDatasmithFacadeScene("StdMaterial", "Solidworks", "Solidworks", "");
-				DatasmithScene.SetName(CleanFileName);
-				DatasmithScene.SetOutputPath(OutDir);
+			FMeshes Meshes = new FMeshes();
 
-				Exporter = new FDatasmithExporter(DatasmithScene);
+			PreExport(Meshes, true);
 
-				FMeshes Meshes = new FMeshes();
+			ConfigurationManager ConfigManager = SwDoc.ConfigurationManager;
+			string[] ConfigurationNames = SwDoc?.GetConfigurationNames();
+			FConfigurationExporter ConfigurationExporter = new FConfigurationExporter(Meshes, ConfigurationNames, ConfigManager.ActiveConfiguration.Name);
 
-				PreExport(Meshes, true);
+			List<FConfigurationData> Configs = ConfigurationExporter.ExportConfigurations(this);
 
-				string[] ConfigurationNames = SwDoc?.GetConfigurationNames();
-				FConfigurationExporter ConfigurationExporter = new FConfigurationExporter(Meshes, ConfigurationNames);
+			bHasConfigurations = (Configs != null) && (Configs.Count != 0);
 
-				List<FConfigurationData> Configs = ConfigurationExporter.ExportConfigurations(this);
+			ExportToDatasmithScene(ConfigurationExporter);
 
-				bHasConfigurations = (Configs != null) && (Configs.Count != 0);
+			ExportLights();
+			ExportConfigurations(Configs);
 
-				ExportToDatasmithScene(ConfigurationExporter);
-			
-				ExportLights();
-				ExportConfigurations(Configs);
+			DatasmithScene.PreExport();
+			DatasmithScene.ExportScene(DatasmithFileExportPath);
 
-				DatasmithScene.PreExport();
-				DatasmithScene.ExportScene(DatasmithFileExportPath);
+			DatasmithScene = OldScene;
+			Exporter = OldExporter;
+		}
 
-				DatasmithScene = OldScene;
-				Exporter = OldExporter;
-			}
+		public void Sync(string InOutputPath)
+		{
+			DatasmithScene.SetName(SwDoc.GetTitle());
+			DatasmithScene.SetOutputPath(InOutputPath);
 
-			SetExportStatus("Done");
+#if DEBUG
+			Stopwatch Watch = Stopwatch.StartNew();
+#endif
+			FMeshes Meshes = new FMeshes();
 
-			SetDirty(false);
+			PreExport(Meshes, false);
+			ExportToDatasmithScene(Meshes);
 
-			bDirectLinkSyncInProgress = false;
-			bFileExportInProgress = false;
+			ExportLights();
 
-			// Kickoff material checker thread after first export
-			if (MaterialCheckerThread == null)
-			{
-				MaterialCheckerThread = new Thread(CheckForMaterialUpdatesProc);
-				MaterialCheckerEvent = new ManualResetEvent(false);
-				MaterialCheckerThread.Start();
-				MaterialCheckerEvent.Set();
-			}
+#if DEBUG
+			Watch.Stop();
+			Debug.WriteLine($"EXPORT TIME: {(double)Watch.ElapsedMilliseconds / 1000.0}");
+#endif
 		}
 
 		// todo: make abstract
@@ -248,58 +247,129 @@ namespace DatasmithSolidworks
 			throw new NotImplementedException();
 		}
 
-		public abstract bool HasMaterialUpdates();
 		public abstract void ExportToDatasmithScene(FMeshes Meshes);
 		public abstract void PreExport(FMeshes Meshes, bool bConfigurations);  // Called before configurations are parsed to prepare meshes needed to identify if they create different configurations
 
-		public virtual void Init()
+		public virtual bool NeedExportComponent(FConfigurationTree.FComponentTreeNode InComponent,
+			FConfigurationTree.FComponentConfig ActiveComponentConfig)
 		{
+			return true;
 		}
 
-		public virtual void Destroy()
-		{
-			bExitMaterialUpdateThread = true;
+		public virtual void AddPartDocument(FConfigurationTree.FComponentTreeNode InNode){}
 
-			if (MaterialCheckerThread != null && !MaterialCheckerThread.Join(500))
-			{
-				MaterialCheckerThread.Abort();
-			}
-		}
+		public virtual void AddExportedComponent(FConfigurationTree.FComponentTreeNode InNode) {}
 
-		public string GetPathName()
+		public abstract ConcurrentDictionary<FComponentName, FObjectMaterials> LoadDocumentMaterials(HashSet<FComponentName> ComponentNamesToExportSet);
+		public abstract void AddComponentMaterials(FComponentName ComponentName, FObjectMaterials Materials);
+		public abstract FObjectMaterials GetComponentMaterials(Component2 Comp);
+
+		public abstract void AddMeshForComponent(FComponentName ComponentName, string MeshName);
+
+		public virtual void Destroy()
 		{
-			return SwDoc.GetPathName();
 		}
 
-		public void MakeActive(bool bInActive)
+		public abstract FMeshData ExtractComponentMeshData(Component2 Comp);
+
+		// Extracts, exports meshes used for the assembly configuration, assigns them to actors
+		// todo: separate just Datasmith Mesh export to its own thread, all other code should be in single thread
+		// Datasmith actor assignment doesn't need threading(and not supposed to be thread-safe?) and Solidworks multithreading is supposed to be slower(SW does all the work in main thread)
+		public void ProcessConfigurationMeshes(FMeshes.FConfiguration MeshesConfiguration, string MeshSuffix)
 		{
-			if (bInActive)
+			ConcurrentDictionary<Component2, FDatasmithFacadeMeshElement> CreatedMeshes =
+				new ConcurrentDictionary<Component2, FDatasmithFacadeMeshElement>();
+
+			Parallel.ForEach(MeshesConfiguration.EnumerateComponentAndActorNames(), KVP =>
 			{
-				if (DatasmithDirectLink == null)
+				Component2 Comp = KVP.Key;
+
+				FMeshData MeshData = ExtractComponentMeshData(Comp);
+
+				MeshesConfiguration.AddMesh(Comp, MeshData);
+
+				if (MeshData != null)
 				{
-					DatasmithDirectLink = new FDatasmithFacadeDirectLink();
-					if (!DatasmithDirectLink.InitializeForScene(DatasmithScene))
+					string MeshName = MeshSuffix == null ? $"{KVP.Value}_Mesh" : $"{KVP.Value}_{MeshSuffix}_Mesh";
+
+					if (Exporter.ExportMesh(MeshName, MeshData, KVP.Value, out FDatasmithFacadeMeshElement NewMesh))
 					{
-						throw new Exception("DirectLink: failed to initialize");
+						CreatedMeshes[Comp] = NewMesh;
 					}
 				}
+			});
 
-				MaterialCheckerEvent?.Set();
-			}
-			else
+			foreach (var KVP in CreatedMeshes)
 			{
-				// Suspend material checker
-				MaterialCheckerEvent?.Reset();
-
-				DatasmithDirectLink?.Dispose();
-				DatasmithDirectLink = null;
+				string MeshName = Exporter.AddMesh(KVP.Value);
+				AddMeshForComponent(new FComponentName(KVP.Key), MeshName);  // Register that this mesh was used for the component
 			}
 		}
+	};
 
-		public void OnExportToFile(string InFilePath)
+
+	// Controls Syncing of a tracked Document, handling all events
+	public interface IDocumentSyncer
+	{
+		// Make Datasmith Scene up to date with the Solidworks Document
+		void Sync(string InOutputPath);
+
+		// Indicate that change handling should start(after first Sync)
+		void Start();
+
+		// Resume all change tracking(after document was made active, foreground in Solidworks)
+		void Resume();
+
+		// Pause all change tracking(when document deactivated, background)
+		void Pause();
+
+		// Is any changes pending(simple test for AutoSync)
+		bool GetDirty();
+
+		// After Sync
+		// todo: Probably not needed at all(i.e. all SetDirty(false) is only deep inside in notifiers, and with false can be moved  into Sync?)
+		void SetDirty(bool bInDirty);  
+
+		FDatasmithFacadeScene GetDatasmithScene();
+
+		// Explicit immediate release of resources(line event handlers)
+		void Destroy();
+	}
+
+	/**
+	 * Base document class.
+	 */
+	public abstract class FDocument
+	{
+		public int DocId = -1;
+		public ModelDoc2 SwDoc = null;
+
+		// Datasmith scene synced with the Solidworks document
+		public IDocumentSyncer DocumentSyncer;
+
+		// DirectLink
+		public FDatasmithFacadeDirectLink DatasmithDirectLink = null;
+		public string DirectLinkPath;
+		public bool bDirectLinkAutoSync { get; set; } = false;
+		public int DirectLinkSyncCount { get; private set; } = 0;
+
+		// Sync/Export 
+		public bool bDirectLinkSyncInProgress = false;
+		public bool bFileExportInProgress = false;
+		public string DatasmithFileExportPath = null;
+
+
+		public FDocument(int InDocId, ModelDoc2 InSwDoc)
 		{
-			DatasmithFileExportPath = InFilePath;
-			RunExport(false);
+			DocId = InDocId;
+			SwDoc = InSwDoc;
+
+			DirectLinkPath = Path.Combine(Path.GetTempPath(), "sw_dl_" + Guid.NewGuid().ToString());
+
+			if (!Directory.Exists(DirectLinkPath))
+			{
+				Directory.CreateDirectory(DirectLinkPath);
+			}
 		}
 
 		public void OnDirectLinkSync()
@@ -312,41 +382,99 @@ namespace DatasmithSolidworks
 			}
 
 			DirectLinkSyncCount++;
-			RunExport(true);
+			
+			bDirectLinkSyncInProgress = true;
+
+			DocumentSyncer.Sync(DirectLinkPath);
+			DatasmithDirectLink.UpdateScene(DocumentSyncer.GetDatasmithScene());
+			DocumentSyncer.SetDirty(false);
+
+			bDirectLinkSyncInProgress = false;
+
+			SetExportStatus("Done");
+
+			DocumentSyncer.Start();
+		}
+
+		public void ToggleDirectLinkAutoSync()
+		{
+			bDirectLinkAutoSync = !bDirectLinkAutoSync;
+
+			if (bDirectLinkAutoSync && DirectLinkSyncCount == 0)
+			{
+				// Run first sync
+				OnDirectLinkSync();
+			}
 		}
 
 		public void OnIdle()
 		{
-			bool bSceneIsDirty = GetDirty();
-			if (bSceneIsDirty && bDirectLinkAutoSync)
+			if (bDirectLinkAutoSync && DocumentSyncer.GetDirty())
 			{
 				OnDirectLinkSync();
 			}
 		}
 
-		public virtual bool NeedExportComponent(FConfigurationTree.FComponentTreeNode InComponent,
-			FConfigurationTree.FComponentConfig ActiveComponentConfig)
+		public void OnExportToFile(string InFilePath)
 		{
-			return false;
+			bFileExportInProgress = true;
+			DatasmithFileExportPath = InFilePath;  // For status
+			Export(InFilePath);
+			bFileExportInProgress = false;
 		}
 
-		/// <summary>
-		/// Whether this component node needs geometry exported/reexported. Currenly doesn't use configuration
-		/// todo: DL configurations
-		/// </summary>
-		public virtual bool NeedGeometryExport(FConfigurationTree.FComponentTreeNode InNode, FConfigurationTree.FComponentConfig ActiveComponentConfig)
+		// Toggle if this document is synced with DirectLink
+		public void MakeActive(bool bInActive)
 		{
-			return false;
+			if (bInActive)
+			{
+				if (DatasmithDirectLink == null)
+				{
+					DatasmithDirectLink = new FDatasmithFacadeDirectLink();
+					if (!DatasmithDirectLink.InitializeForScene(DocumentSyncer.GetDatasmithScene()))
+					{
+						throw new Exception("DirectLink: failed to initialize");
+					}
+				}
+
+				DocumentSyncer.Resume();
+			}
+			else
+			{
+				// Suspend material checker(while it's a heavy threaded procedure)
+				// other event handling is not suspended to register changes in the model
+				DocumentSyncer?.Pause();
+
+				DatasmithDirectLink?.Dispose();  // aka 'Close DirectLink Connection'
+				DatasmithDirectLink = null;
+			}
 		}
 
-		public virtual void AddPartDocument(FConfigurationTree.FComponentTreeNode InNode){}
+		public abstract void Export(string InFilePath);
 
-		public virtual void AddExportedComponent(FConfigurationTree.FComponentTreeNode InNode) {}
+		public string GetPathName()
+		{
+			return SwDoc.GetPathName();
+		}
 
-		public abstract ConcurrentDictionary<FComponentName, FObjectMaterials> LoadDocumentMaterials(HashSet<FComponentName> ComponentNamesToExportSet);
-		public abstract void AddComponentMaterials(FComponentName ComponentName, FObjectMaterials Materials);
-		public abstract FObjectMaterials GetComponentMaterials(Component2 Comp);
+		// Dispose all resources explicitly
+		// todo: DatasmithDirectLink should be disposed here too?
+		public virtual void Destroy()
+		{
+			DocumentSyncer?.Destroy();
+		}
 
-		public abstract void AddMeshForComponent(FComponentName ComponentName, string MeshName);
+		public void SetExportStatus(string InMessage)
+		{
+			if (bDirectLinkSyncInProgress)
+			{
+				Addin.Instance.LogStatusBarMessage($"DirectLink sync...{InMessage}");
+			}
+			else if (bFileExportInProgress)
+			{
+				string FileName = Path.GetFileName(DatasmithFileExportPath);
+				Addin.Instance.LogStatusBarMessage($"Exporting file {FileName}...{InMessage}");
+			}
+		}
 	}
 }
\ No newline at end of file
diff --git a/Engine/Source/Programs/Enterprise/Datasmith/DatasmithSolidworksExporter/Private/ObjectMaterials.cs b/Engine/Source/Programs/Enterprise/Datasmith/DatasmithSolidworksExporter/Private/ObjectMaterials.cs
--- a/Engine/Source/Programs/Enterprise/Datasmith/DatasmithSolidworksExporter/Private/ObjectMaterials.cs
+++ b/Engine/Source/Programs/Enterprise/Datasmith/DatasmithSolidworksExporter/Private/ObjectMaterials.cs
@@ -17,7 +17,7 @@ namespace DatasmithSolidworks
 		// for example, when loading component materials, owner will be the document 
 		// that component resides in, while part doc will be the document that component references!
 		public PartDoc PartDocument {  get; private set; }
-		public FDocument OwnerDoc {  get; private set; }
+		public FDocumentTracker OwnerDoc {  get; private set; }
 
 		public int ComponentMaterialID { get; private set; } = -1;
 		public int PartMaterialID { get; private set; } = -1;
@@ -28,7 +28,7 @@ namespace DatasmithSolidworks
 
 		public ConcurrentDictionary<int, FMaterial> GlobalMaterialsMap { get; private set; } = null;
 
-		public FObjectMaterials(FDocument InOwnerDoc, PartDoc InPartDocument, ref ConcurrentDictionary<int, FMaterial> InOutMaterialsMap)
+		public FObjectMaterials(FDocumentTracker InOwnerDoc, PartDoc InPartDocument, ref ConcurrentDictionary<int, FMaterial> InOutMaterialsMap)
 		{
 			GlobalMaterialsMap = InOutMaterialsMap;
 			OwnerDoc = InOwnerDoc;
@@ -89,7 +89,7 @@ namespace DatasmithSolidworks
 			}
 
 			uint FaceId = unchecked((uint)(InFace?.GetFaceId() ?? 0));
-			if (FDocument.IsValidFaceId(FaceId) && FaceMaterialsMap.ContainsKey(FaceId.ToString()))
+			if (FDocumentTracker.IsValidFaceId(FaceId) && FaceMaterialsMap.ContainsKey(FaceId.ToString()))
 			{
 				int MatId = FaceMaterialsMap[FaceId.ToString()];
 				return GetMaterial(MatId);
@@ -101,7 +101,7 @@ namespace DatasmithSolidworks
 
 			if (Feat != null)
 			{
-				string FeatId = FPartDocument.GetFeaturePath(PartDocument, Feat);
+				string FeatId = FPartDocumentTracker.GetFeaturePath(PartDocument, Feat);
 				if (FeatureMaterialsMap.ContainsKey(FeatId))
 				{
 					return GetMaterial(FeatureMaterialsMap[FeatId]);
@@ -110,7 +110,7 @@ namespace DatasmithSolidworks
 
 			if (Body != null)
 			{
-				string BodyId = FPartDocument.GetBodyPath(PartDocument, Body);
+				string BodyId = FPartDocumentTracker.GetBodyPath(PartDocument, Body);
 				if (BodyMaterialsMap.ContainsKey(BodyId))
 				{
 					return GetMaterial(BodyMaterialsMap[BodyId]);
@@ -170,7 +170,7 @@ namespace DatasmithSolidworks
 			return InDict1.OrderBy(KVP => KVP.Key).SequenceEqual(InDict2.OrderBy(KVP => KVP.Key));
 		}
 
-		public static FObjectMaterials LoadPartMaterials(FDocument InOwnerDoc, PartDoc InPartDoc, swDisplayStateOpts_e InDisplayState, string[] InDisplayStateNames)
+		public static FObjectMaterials LoadPartMaterials(FDocumentTracker InOwnerDoc, PartDoc InPartDoc, swDisplayStateOpts_e InDisplayState, string[] InDisplayStateNames)
 		{
 			ModelDoc2 Doc = InPartDoc as ModelDoc2;
 			if (Doc == null)
@@ -210,7 +210,7 @@ namespace DatasmithSolidworks
 
 					if (ObjUser is IBody2 Body)
 					{
-						string BodyId = FPartDocument.GetBodyPath(InPartDoc as PartDoc, Body);
+						string BodyId = FPartDocumentTracker.GetBodyPath(InPartDoc as PartDoc, Body);
 						PartMaterials.RegisterMaterial(PartMaterials.BodyMaterialsMap, Doc, RenderMat, BodyId);
 						continue;
 					}
@@ -224,7 +224,7 @@ namespace DatasmithSolidworks
 
 					if (ObjUser is IFeature Feat)
 					{
-						string FeatureId = FPartDocument.GetFeaturePath(InPartDoc as PartDoc, Feat);
+						string FeatureId = FPartDocumentTracker.GetFeaturePath(InPartDoc as PartDoc, Feat);
 						PartMaterials.RegisterMaterial(PartMaterials.FeatureMaterialsMap, Doc, RenderMat, FeatureId);
 						continue;
 					}
@@ -234,7 +234,7 @@ namespace DatasmithSolidworks
 			return PartMaterials;
 		}
 
-		public static FObjectMaterials LoadComponentMaterials(FDocument InComponentOwner, Component2 InComponent, swDisplayStateOpts_e InDisplayState, string[] InDisplayStateNames)
+		public static FObjectMaterials LoadComponentMaterials(FDocumentTracker InComponentOwner, Component2 InComponent, swDisplayStateOpts_e InDisplayState, string[] InDisplayStateNames)
 		{
 			Addin.Instance.LogDebug($"LoadComponentMaterials: {InComponentOwner.GetPathName()}, {InDisplayState}, {InDisplayStateNames}");
 
@@ -272,7 +272,7 @@ namespace DatasmithSolidworks
 				{
 					if (ObjUser is IBody2 Body)
 					{
-						string BodyId = FPartDocument.GetBodyPath(ComponentDoc as PartDoc, Body);
+						string BodyId = FPartDocumentTracker.GetBodyPath(ComponentDoc as PartDoc, Body);
 						ComponentMaterials.RegisterMaterial(ComponentMaterials.BodyMaterialsMap, InComponentOwner.SwDoc, RenderMat, BodyId);
 						continue;
 					}
@@ -286,7 +286,7 @@ namespace DatasmithSolidworks
 
 					if (ObjUser is IFeature Feat)
 					{
-						string FeatureId = FPartDocument.GetFeaturePath(ComponentDoc as PartDoc, Feat);
+						string FeatureId = FPartDocumentTracker.GetFeaturePath(ComponentDoc as PartDoc, Feat);
 						ComponentMaterials.RegisterMaterial(ComponentMaterials.FeatureMaterialsMap, InComponentOwner.SwDoc, RenderMat, FeatureId);
 						continue;
 					}
@@ -300,7 +300,7 @@ namespace DatasmithSolidworks
 			return ComponentMaterials;
 		}
 
-		public static ConcurrentDictionary<FComponentName, FObjectMaterials> LoadAssemblyMaterials(FAssemblyDocument InAsmDoc, HashSet<FComponentName> InComponentsSet, swDisplayStateOpts_e InDisplayState, string[] InDisplayStateNames)
+		public static ConcurrentDictionary<FComponentName, FObjectMaterials> LoadAssemblyMaterials(FAssemblyDocumentTracker InAsmDoc, HashSet<FComponentName> InComponentsSet, swDisplayStateOpts_e InDisplayState, string[] InDisplayStateNames)
 		{
 			Addin.Instance.LogDebug($"LoadAssemblyMaterials: {InAsmDoc.GetPathName()}, {InDisplayState}, {InDisplayStateNames}");
 			Addin.Instance.LogDebug($"  for components: {string.Join(", ", InComponentsSet)}");
diff --git a/Engine/Source/Programs/Enterprise/Datasmith/DatasmithSolidworksExporter/Private/PartDocument.cs b/Engine/Source/Programs/Enterprise/Datasmith/DatasmithSolidworksExporter/Private/PartDocument.cs
--- a/Engine/Source/Programs/Enterprise/Datasmith/DatasmithSolidworksExporter/Private/PartDocument.cs
+++ b/Engine/Source/Programs/Enterprise/Datasmith/DatasmithSolidworksExporter/Private/PartDocument.cs
@@ -6,23 +6,28 @@ using System;
 using System.Collections.Concurrent;
 using System.Collections.Generic;
 using System.IO;
+using System.Security.Cryptography;
+using static System.Windows.Forms.VisualStyles.VisualStyleElement;
 
 namespace DatasmithSolidworks
 {
-	public class FPartDocument : FDocument
+	public class FPartDocumentTracker : FDocumentTracker
 	{
 		public PartDoc SwPartDoc { get; private set; } = null;
-		private FAssemblyDocument AsmDoc = null;
+		private FAssemblyDocumentTracker AsmDoc = null;
 		private FComponentName ComponentName;
 		private string MeshName = null;
+		public readonly FPartDocument PartDocument;
 
 		public FObjectMaterials ExportedPartMaterials { get; set; }
 
 		public string PathName { get; private set; }
 
-		public FPartDocument(int InDocId, PartDoc InPartDoc, FDatasmithExporter InExporter, FAssemblyDocument InAsmDoc, FComponentName InComponentName) 
-			: base(InDocId, InPartDoc as ModelDoc2, InExporter)
+		public FPartDocumentTracker(FPartDocument InDoc, PartDoc InPartDoc, FDatasmithExporter InExporter,
+			FAssemblyDocumentTracker InAsmDoc, FComponentName InComponentName)
+			: base(InDoc, InExporter)
 		{
+			PartDocument = InDoc;
 			SwPartDoc = InPartDoc;
 			AsmDoc = InAsmDoc;
 			ComponentName = InComponentName;
@@ -47,8 +52,11 @@ namespace DatasmithSolidworks
 					string Name = InBody.Name;
 					Path = Title + "\\\\" + Name;
 				}
-				catch { }
+				catch
+				{
+				}
 			}
+
 			return Path;
 		}
 
@@ -64,19 +72,22 @@ namespace DatasmithSolidworks
 					string Name = InFeature.Name;
 					Path = Title + "\\\\" + Name;
 				}
-				catch { }
+				catch
+				{
+				}
 			}
+
 			return Path;
 		}
 
-		protected override void SetDirty(bool bInDirty)
+		public override void SetDirty(bool bInDirty)
 		{
 			base.SetDirty(bInDirty);
 
 			if (bInDirty && AsmDoc != null && ComponentName.IsValid())
 			{
 				// Notify owner assembly that a component needs re-export
-				AsmDoc.SetComponentDirty(ComponentName, FAssemblyDocument.EComponentDirtyState.Geometry);
+				AsmDoc.SetComponentDirty(ComponentName, FAssemblyDocumentTracker.EComponentDirtyState.Geometry);
 			}
 		}
 
@@ -85,39 +96,8 @@ namespace DatasmithSolidworks
 			if (!bConfigurations)
 			{
 				// Not exporting configurations - just load materials here
-				ExportedPartMaterials = FObjectMaterials.LoadPartMaterials(this, SwPartDoc, swDisplayStateOpts_e.swThisDisplayState, null);
-				return;
-			}
-
-			// Load mesh data for each configuration - need to have information about meshes to understand what to export for configurations.
-			// (Also load materials - mesh processing later needs them)
-			ConfigurationManager ConfigManager = (SwPartDoc as ModelDoc2).ConfigurationManager;
-			IConfiguration OriginalConfiguration = ConfigManager.ActiveConfiguration;
-
-			foreach (string ConfigurationName in SwDoc.GetConfigurationNames())
-			{
-				SwDoc.ShowConfiguration2(ConfigurationName);
-
-				// Load materials for each configuration
-				ExportedPartMaterials = FObjectMaterials.LoadPartMaterials(this, SwPartDoc, swDisplayStateOpts_e.swThisDisplayState, null);
-
-				// todo: Solidworks api docs says that GetRootComponent3 is null for Part but our current 
-				//   parsing code relies on it's existence(and it exists in all cases). Probably better to refactor this without Component use for parts
-				//   This is possible, we don't need Component, just code would be be less 'generic'   
-				Component2 Component = ConfigManager.ActiveConfiguration.GetRootComponent3(true);
-
-				ConcurrentBag<FBody> Bodies = FBody.FetchBodies(SwPartDoc);
-				FMeshData MeshData = FStripGeometry.CreateMeshData(Bodies, ExportedPartMaterials);
-
-				if (MeshData != null)
-				{
-					Meshes.GetMeshesConfiguration(ConfigManager.ActiveConfiguration.Name).AddMesh(Component, MeshData);
-				}
-			}
-
-			if (OriginalConfiguration != null)
-			{
-				SwDoc.ShowConfiguration2(OriginalConfiguration.Name);
+				ExportedPartMaterials =
+					FObjectMaterials.LoadPartMaterials(this, SwPartDoc, swDisplayStateOpts_e.swThisDisplayState, null);
 			}
 		}
 
@@ -140,22 +120,26 @@ namespace DatasmithSolidworks
 			{
 				ConfigurationManager ConfigManager = (SwPartDoc as ModelDoc2).ConfigurationManager;
 				string ActiveConfigurationName = ConfigManager.ActiveConfiguration.Name;
+				FConfigurationExporter ConfigurationExporter = new FConfigurationExporter(Meshes, null, ActiveConfigurationName);
 
 				foreach (string ConfigurationName in SwDoc.GetConfigurationNames())
 				{
 
 					// todo: might want to not use GetRootComponent3 for Part - docs says that it's null for Part
-					FComponentName RootComponentName = new FComponentName(ConfigManager.ActiveConfiguration.GetRootComponent3(true));
+					FComponentName RootComponentName =
+						new FComponentName(ConfigManager.ActiveConfiguration.GetRootComponent3(true));
 					FMeshData MeshData = Meshes.GetMeshData(ConfigurationName, RootComponentName);
 					if (MeshData == null)
 					{
 						continue;
 					}
+
 					// todo: not make extra copy - use existing(i.e. instead of bHasConfigurations) 
-					FActorName ActorName = new FConfigurationExporter(Meshes, null).GetMeshActorName(RootComponentName, ConfigurationName);
+					FActorName ActorName =
+						ConfigurationExporter.GetMeshActorName(ConfigurationName, RootComponentName);
 
 					bool bHasMesh = Exporter.ExportMesh($"{ActorName}_Mesh", MeshData, ActorName,
-						out Tuple<FDatasmithFacadeMeshElement, FDatasmithFacadeMesh> NewMesh);
+						out FDatasmithFacadeMeshElement NewMesh);
 
 					string MeshNameForConfiguration = null; // todo: Sync not supported for variants, 
 					if (bHasMesh)
@@ -175,8 +159,8 @@ namespace DatasmithSolidworks
 			else
 			{
 				FActorName ActorName = FActorName.FromString(PartName);
-				ConcurrentBag<FBody> Bodies = FBody.FetchBodies(SwPartDoc);
-				FMeshData MeshData = FStripGeometry.CreateMeshData(Bodies, ExportedPartMaterials);
+				
+				FMeshData MeshData = ExtractPartMeshData();
 
 				Exporter.RemoveMesh(MeshName);
 				MeshName = null;
@@ -185,7 +169,7 @@ namespace DatasmithSolidworks
 				if (MeshData != null)
 				{
 					bHasMesh = Exporter.ExportMesh($"{ActorName}_Mesh", MeshData, ActorName,
-						out Tuple<FDatasmithFacadeMeshElement, FDatasmithFacadeMesh> NewMesh);
+						out FDatasmithFacadeMeshElement NewMesh);
 
 					if (bHasMesh)
 					{
@@ -205,13 +189,14 @@ namespace DatasmithSolidworks
 
 		public override bool HasMaterialUpdates()
 		{
-			FObjectMaterials CurrentMaterials = FObjectMaterials.LoadPartMaterials(this, SwPartDoc, swDisplayStateOpts_e.swThisDisplayState, null);
+			FObjectMaterials CurrentMaterials =
+				FObjectMaterials.LoadPartMaterials(this, SwPartDoc, swDisplayStateOpts_e.swThisDisplayState, null);
 			if (CurrentMaterials == null && ExportedPartMaterials == null)
 			{
 				return false;
 			}
 			else if ((CurrentMaterials != null && ExportedPartMaterials == null) ||
-					 (CurrentMaterials == null && ExportedPartMaterials != null))
+			         (CurrentMaterials == null && ExportedPartMaterials != null))
 			{
 				return true;
 			}
@@ -219,9 +204,47 @@ namespace DatasmithSolidworks
 			return !CurrentMaterials.EqualMaterials(ExportedPartMaterials);
 		}
 
-		public override void Init()
+		public override FObjectMaterials GetComponentMaterials(Component2 Comp)
+		{
+			return null;
+		}
+
+		public override void AddMeshForComponent(FComponentName ComponentName, string MeshName)
+		{
+		}
+
+		public override FMeshData ExtractComponentMeshData(Component2 Comp)
+		{
+			return ExtractPartMeshData();  // Part document has single component
+		}
+
+		// Extract mesh data for current active configuration. MeshData is ready for export to DatasmithMesh 
+		public FMeshData ExtractPartMeshData()
+		{
+			FObjectMaterials PartMaterials = FObjectMaterials.LoadPartMaterials(this, SwPartDoc, swDisplayStateOpts_e.swThisDisplayState, null);
+			ConcurrentBag<FBody> Bodies = FBody.FetchBodies(SwPartDoc);
+			return FStripGeometry.CreateMeshData(Bodies, PartMaterials);
+		}
+
+		public override void AddComponentMaterials(FComponentName ComponentName, FObjectMaterials Materials)
+		{
+		}
+
+		public override ConcurrentDictionary<FComponentName, FObjectMaterials> LoadDocumentMaterials(HashSet<FComponentName> ComponentNamesToExportSet)
+		{
+			return null;
+		}
+	}
+
+	public class FPartDocumentEvents
+	{
+		private readonly FPartDocumentTracker DocumentTracker;
+		private PartDoc SwPartDoc => DocumentTracker.PartDocument.SwPartDoc;
+		private int DocId => DocumentTracker.PartDocument.DocId;
+
+		public FPartDocumentEvents(FPartDocumentTracker InDocumentTracker)
 		{
-			base.Init();
+			DocumentTracker = InDocumentTracker;
 
 			SwPartDoc.UndoPostNotify += new DPartDocEvents_UndoPostNotifyEventHandler(OnPartUndoPostNotify);
 			SwPartDoc.SuppressionStateChangeNotify += new DPartDocEvents_SuppressionStateChangeNotifyEventHandler(OnPartSuppressionStateChangeNotify);
@@ -247,10 +270,8 @@ namespace DatasmithSolidworks
 			SwPartDoc.FileSavePostNotify += new DPartDocEvents_FileSavePostNotifyEventHandler(OnPartFileSavePostNotify);
 		}
 
-		public override void Destroy()
+		public void Destroy()
 		{
-			base.Destroy();
-
 			SwPartDoc.UndoPostNotify -= new DPartDocEvents_UndoPostNotifyEventHandler(OnPartUndoPostNotify);
 			SwPartDoc.SuppressionStateChangeNotify -= new DPartDocEvents_SuppressionStateChangeNotifyEventHandler(OnPartSuppressionStateChangeNotify);
 			SwPartDoc.DestroyNotify2 -= new DPartDocEvents_DestroyNotify2EventHandler(OnPartDocumentDestroyNotify2);
@@ -275,23 +296,6 @@ namespace DatasmithSolidworks
 			SwPartDoc.FileSavePostNotify -= new DPartDocEvents_FileSavePostNotifyEventHandler(OnPartFileSavePostNotify);
 		}
 
-		public override FObjectMaterials GetComponentMaterials(Component2 Comp)
-		{
-			return null;
-		}
-
-		public override void AddMeshForComponent(FComponentName ComponentName, string MeshName)
-		{
-		}
-
-		public override void AddComponentMaterials(FComponentName ComponentName, FObjectMaterials Materials)
-		{
-		}
-
-		public override ConcurrentDictionary<FComponentName, FObjectMaterials> LoadDocumentMaterials(HashSet<FComponentName> ComponentNamesToExportSet)
-		{
-			return null;
-		}
 
 		int OnPartDocumentDestroyNotify2(int DestroyType)
 		{
@@ -301,14 +305,14 @@ namespace DatasmithSolidworks
 
 		int OnPartUndoPostNotify()
 		{
-			SetDirty(true);
+			DocumentTracker.SetDirty(true);
 			return 0;
 		}
 
 		// part has become resolved or lightweight or suppressed
 		int OnPartSuppressionStateChangeNotify(Feature Feature, int NewSuppressionState, int PreviousSuppressionState, int ConfigurationOption, ref object ConfigurationNames)
 		{
-			SetDirty(true);
+			DocumentTracker.SetDirty(true);
 			return 0;
 		}
 
@@ -319,7 +323,7 @@ namespace DatasmithSolidworks
 
 		int OnPartConfigurationChangeNotify(string ConfigurationName, object Object, int ObjectType, int changeType)
 		{
-			SetDirty(true);
+			DocumentTracker.SetDirty(true);
 			return 0;
 		}
 
@@ -330,7 +334,7 @@ namespace DatasmithSolidworks
 
 		int OnPartDragStateChangeNotify(Boolean State)
 		{
-			SetDirty(true);
+			DocumentTracker.SetDirty(true);
 			return 0;
 		}
 
@@ -346,13 +350,13 @@ namespace DatasmithSolidworks
 
 		int OnPartRedoPostNotify()
 		{
-			SetDirty(true);
+			DocumentTracker.SetDirty(true);
 			return 0;
 		}
 
 		int OnPartAddItemNotify(int InEntityType, string InItemName)
 		{
-			SetDirty(true);
+			DocumentTracker.SetDirty(true);
 			return 0;
 		}
 
@@ -393,7 +397,7 @@ namespace DatasmithSolidworks
 
 		int OnPartRegenPostNotify2(object stopFeature)
 		{
-			SetDirty(true);
+			DocumentTracker.SetDirty(true);
 			return 0;
 		}
 
@@ -412,4 +416,132 @@ namespace DatasmithSolidworks
 			return 0;
 		}
 	}
+
+	public class FPartDocumentNotifications
+	{
+		private readonly FPartDocumentTracker DocumentTracker;
+		private readonly FPartDocumentEvents Events;
+		private FMaterialUpdateChecker MaterialUpdateChecker;
+
+		public FPartDocumentNotifications(int InDocId, PartDoc InPartDoc, FPartDocumentTracker InDocumentTracker)
+		{
+			DocumentTracker = InDocumentTracker;
+			Events = new FPartDocumentEvents(InDocumentTracker);
+		}
+
+		public void Destroy()
+		{
+			MaterialUpdateChecker?.Destroy();
+			Events.Destroy();
+		}
+
+		public void Start()
+		{
+			if (MaterialUpdateChecker == null)
+			{
+				MaterialUpdateChecker = new FMaterialUpdateChecker(DocumentTracker);
+			}
+		}
+
+		public void Resume()
+		{
+			MaterialUpdateChecker?.Resume();
+		}
+
+		public void Pause()
+		{
+			MaterialUpdateChecker?.Pause();
+		}
+	}
+	
+	public class FPartDocumentSyncer: IDocumentSyncer
+	{
+		public readonly FPartDocumentNotifications Notifications;
+		public readonly FPartDocumentTracker DocumentTracker;
+
+		public FPartDocumentSyncer(int InDocId, FPartDocument InDoc, PartDoc InPartDoc, FDatasmithExporter InExporter, FAssemblyDocumentTracker InAsmDoc, FComponentName InComponentName)
+		{
+			DocumentTracker = new FPartDocumentTracker(InDoc, InPartDoc, InExporter, InAsmDoc, InComponentName);
+			Notifications = new FPartDocumentNotifications(InDocId, InPartDoc, DocumentTracker);
+		}
+
+		public void Destroy()
+		{
+			Notifications.Destroy();
+			DocumentTracker.Destroy();
+		}
+
+		public void Start()
+		{
+			Notifications.Start();
+		}
+
+		public void Resume()
+		{
+			Notifications.Resume();
+		}
+
+		public void Pause()  
+		{
+			Notifications.Pause();
+		}
+
+		public bool GetDirty()  // i.e. in autosync
+		{
+			return DocumentTracker.GetDirty();
+		}
+
+		public void SetDirty(bool bInDirty)
+		{
+			DocumentTracker.SetDirty(bInDirty);
+		}
+
+		public void Sync(string InOutputPath)
+		{
+			DocumentTracker.Sync(InOutputPath);
+		}
+
+		public FDatasmithFacadeScene GetDatasmithScene()
+		{
+			return DocumentTracker.DatasmithScene;
+		}
+	}
+
+	public class FPartDocument: FDocument
+	{
+		public PartDoc SwPartDoc { get; private set; } = null;
+
+		private readonly FPartDocumentEvents Events;
+
+		public FPartDocument(int InDocId, PartDoc InPartDoc, FDatasmithExporter InExporter, FAssemblyDocumentTracker InAsmDoc, FComponentName InComponentName) 
+				: base(InDocId, InPartDoc as ModelDoc2)
+		{
+			SwPartDoc = InPartDoc;
+
+			if (InAsmDoc == null)  
+			{
+				// This is the top document and should track/sync its DatasmithScene
+				DocumentSyncer = new FPartDocumentSyncer(InDocId, this, InPartDoc, InExporter, InAsmDoc, InComponentName);
+			}
+			else
+			{
+				// Assembly Document passed when this Part Document is within an assembly and only needs to propagate events to that assembly
+				// todo: might split this from FPartDocument?
+				Events = new FPartDocumentEvents(new FPartDocumentTracker(this, InPartDoc, InExporter, InAsmDoc, InComponentName));
+			}
+		}
+
+		public override void Export(string InFilePath)
+		{
+			FPartDocumentTracker Tracker = new FPartDocumentTracker(this, SwPartDoc, null, null, new FComponentName());
+			Tracker.DatasmithFileExportPath = InFilePath;
+			Tracker.Export();
+		}
+
+		public override void Destroy()
+		{
+			Events?.Destroy();
+			base.Destroy();
+		}
+	}
 }
\ No newline at end of file
