[Solidworks] Fixed potential race condition exporting meshes; refactor

- In ExportMesh left to threaded execution only construction of DatatasmithMesh from extracted data and ExportDatasmithMesh itself

- refactored use of Mesh names all over to make more clear distintion between Datasmith names and names built from SW component names

diff --git a/Engine/Source/Programs/Enterprise/Datasmith/DatasmithSolidworksExporter/DatasmithSolidworks/Properties/AssemblyInfo.cs b/Engine/Source/Programs/Enterprise/Datasmith/DatasmithSolidworksExporter/DatasmithSolidworks/Properties/AssemblyInfo.cs
--- a/Engine/Source/Programs/Enterprise/Datasmith/DatasmithSolidworksExporter/DatasmithSolidworks/Properties/AssemblyInfo.cs
+++ b/Engine/Source/Programs/Enterprise/Datasmith/DatasmithSolidworksExporter/DatasmithSolidworks/Properties/AssemblyInfo.cs
@@ -2,6 +2,7 @@
 
 using System.Reflection;
 using System.Runtime.CompilerServices;
+using System.Runtime.InteropServices;
 
 //
 // General Information about an assembly is controlled through the following
@@ -60,3 +61,5 @@ using System.Runtime.CompilerServices;
 [assembly: AssemblyKeyName("")]
 [assembly: AssemblyFileVersion("1.0.0.21")]
 
+// todo: We don't need com visibility for every class. Only for the Addin and whatever we'd want to expose and automation API
+// [assembly: ComVisible(false)]
diff --git a/Engine/Source/Programs/Enterprise/Datasmith/DatasmithSolidworksExporter/Private/Addin.cs b/Engine/Source/Programs/Enterprise/Datasmith/DatasmithSolidworksExporter/Private/Addin.cs
--- a/Engine/Source/Programs/Enterprise/Datasmith/DatasmithSolidworksExporter/Private/Addin.cs
+++ b/Engine/Source/Programs/Enterprise/Datasmith/DatasmithSolidworksExporter/Private/Addin.cs
@@ -14,6 +14,7 @@ using System.Linq;
 using System.Reflection;
 using System.Runtime.InteropServices;
 using System.Threading;
+using DatasmithSolidworks.Names;
 using Environment = System.Environment;
 
 namespace DatasmithSolidworks
diff --git a/Engine/Source/Programs/Enterprise/Datasmith/DatasmithSolidworksExporter/Private/AssemblyDocument.cs b/Engine/Source/Programs/Enterprise/Datasmith/DatasmithSolidworksExporter/Private/AssemblyDocument.cs
--- a/Engine/Source/Programs/Enterprise/Datasmith/DatasmithSolidworksExporter/Private/AssemblyDocument.cs
+++ b/Engine/Source/Programs/Enterprise/Datasmith/DatasmithSolidworksExporter/Private/AssemblyDocument.cs
@@ -8,6 +8,7 @@ using System.Collections.Generic;
 using System.Diagnostics;
 using System.Linq;
 using System.Threading.Tasks;
+using DatasmithSolidworks.Names;
 
 namespace DatasmithSolidworks
 {
@@ -47,8 +48,8 @@ namespace DatasmithSolidworks
 				new Dictionary<FComponentName, Component2>();
 
 			/** Stores which mesh was exported for each component*/
-			public Dictionary<FComponentName, string> ComponentNameToMeshNameMap =
-				new Dictionary<FComponentName, string>();
+			public Dictionary<FComponentName, FMeshName> ComponentNameToMeshNameMap =
+				new Dictionary<FComponentName, FMeshName>();
 		}
 
 		public AssemblyDoc SwAsmDoc { get; private set; } = null;
@@ -111,11 +112,11 @@ namespace DatasmithSolidworks
 		{
 			foreach (FComponentName CompName in SyncState.ComponentsToDelete)
 			{
-				FActorName ActorName = CompName.GetActorName();
+				FActorName ActorName = Exporter.GetComponentActorName(CompName);
 				SyncState.ComponentToPartMap.Remove(CompName);
 				SyncState.ExportedComponentsMap.Remove(CompName);
 
-				if (SyncState.ComponentNameToMeshNameMap.TryGetValue(CompName, out string MeshName))
+				if (SyncState.ComponentNameToMeshNameMap.TryGetValue(CompName, out FMeshName MeshName))
 				{
 					SyncState.ComponentNameToMeshNameMap.Remove(CompName);
 					Exporter.RemoveMesh(MeshName);
@@ -142,7 +143,7 @@ namespace DatasmithSolidworks
 			SyncState.ComponentsMaterialsMap[ComponentName] = Materials;
 		}
 
-		public override void AddMeshForComponent(FComponentName ComponentName, string MeshName)
+		public override void AddMeshForComponent(FComponentName ComponentName, FMeshName MeshName)
 		{
 			SyncState.ComponentNameToMeshNameMap[ComponentName] = MeshName;
 		}
@@ -366,14 +367,14 @@ namespace DatasmithSolidworks
 
 				FDatasmithActorExportInfo ActorExportInfo = new FDatasmithActorExportInfo();
 
-				FActorName ActorName = ComponentName.GetActorName();
+				FActorName ActorName = Exporter.GetComponentActorName(ComponentName);
 
 				ActorExportInfo.Label = ComponentName.GetLabel();
 				ActorExportInfo.Name = ActorName;
 
 				if (InParent != null)
 				{
-					ActorExportInfo.ParentName = new FComponentName(InParent).GetActorName();
+					ActorExportInfo.ParentName = Exporter.GetComponentActorName(InParent);
 				}
 
 				ActorExportInfo.bVisible = true;
@@ -444,7 +445,7 @@ namespace DatasmithSolidworks
 				if (ActorExportInfo.Type == EActorType.MeshActor)
 				{
 					// todo: deduplicate mesh name calculation
-					ActorExportInfo.MeshName = FDatasmithExporter.SanitizeName($"{ActorExportInfo.Name}_Mesh");
+					ActorExportInfo.MeshName = new FMeshName(ComponentName);
 				}
 
 				Exporter.ExportOrUpdateActor(ActorExportInfo);
@@ -479,7 +480,7 @@ namespace DatasmithSolidworks
 
 			if (InNode.bGeometrySame)
 			{
-				FActorName ActorName = InNode.ActorName;
+				FActorName ActorName = Exporter.GetComponentActorName(InNode.ComponentName);
 
 				FDatasmithActorExportInfo ActorExportInfo = new FDatasmithActorExportInfo();
 
@@ -488,7 +489,7 @@ namespace DatasmithSolidworks
 
 				if (InParent != null)
 				{
-					ActorExportInfo.ParentName = InParent.GetActorName();
+					ActorExportInfo.ParentName = Exporter.GetComponentActorName(InParent.ComponentName);
 				}
 
 				ActorExportInfo.bVisible = true;
@@ -506,7 +507,8 @@ namespace DatasmithSolidworks
 				if (ActorExportInfo.Type == EActorType.MeshActor)
 				{
 					// todo: deduplicate mesh name calculation
-					ActorExportInfo.MeshName = FDatasmithExporter.SanitizeName(ConfigurationExporter.GetMeshName(InNode.ComponentName));
+					
+					ActorExportInfo.MeshName = ConfigurationExporter.GetMeshName(InNode.ComponentName);
 				}
 
 				Exporter.ExportOrUpdateActor(ActorExportInfo);
@@ -517,7 +519,7 @@ namespace DatasmithSolidworks
 
 				FActorName ParentName;
 				{
-					FActorName ActorName = InNode.ActorName;
+					FActorName ActorName = Exporter.GetComponentActorName(InNode.ComponentName);
 					ParentName = ActorName;
 
 					FDatasmithActorExportInfo ActorExportInfo = new FDatasmithActorExportInfo();
@@ -527,7 +529,7 @@ namespace DatasmithSolidworks
 
 					if (InParent != null)
 					{
-						ActorExportInfo.ParentName = InParent.GetActorName();
+						ActorExportInfo.ParentName = Exporter.GetComponentActorName(InParent.ComponentName);
 					}
 
 					ActorExportInfo.bVisible = true;
@@ -564,7 +566,7 @@ namespace DatasmithSolidworks
 
 					;
 					// todo: deduplicate mesh name calculation
-					ActorExportInfo.MeshName = FDatasmithExporter.SanitizeName(ConfigurationExporter.GetMeshName(ComponentConfig.ConfigName, InNode.ComponentName));
+					ActorExportInfo.MeshName = ConfigurationExporter.GetMeshName(ComponentConfig.ConfigName, InNode.ComponentName);
 
 					Exporter.ExportOrUpdateActor(ActorExportInfo);
 				}
diff --git a/Engine/Source/Programs/Enterprise/Datasmith/DatasmithSolidworksExporter/Private/ConfigurationExporter.cs b/Engine/Source/Programs/Enterprise/Datasmith/DatasmithSolidworksExporter/Private/ConfigurationExporter.cs
--- a/Engine/Source/Programs/Enterprise/Datasmith/DatasmithSolidworksExporter/Private/ConfigurationExporter.cs
+++ b/Engine/Source/Programs/Enterprise/Datasmith/DatasmithSolidworksExporter/Private/ConfigurationExporter.cs
@@ -8,6 +8,7 @@ using SolidWorks.Interop.swconst;
 using System.Runtime.InteropServices;
 using System.Collections.Concurrent;
 using System.Linq;
+using DatasmithSolidworks.Names;
 using static DatasmithSolidworks.FConfigurationTree;
 
 namespace DatasmithSolidworks
@@ -208,12 +209,9 @@ namespace DatasmithSolidworks
 				return MeshForComponent.Keys;
 			}
 
-			public IEnumerable<KeyValuePair<Component2, FActorName>> EnumerateComponentAndActorNames()
+			public IEnumerable<Component2> EnumerateComponents()
 			{
-				foreach (KeyValuePair<FComponentName, FComponentMesh> FVP in MeshForComponent)
-				{
-					yield return new KeyValuePair<Component2, FActorName>(FVP.Value.Component, FVP.Value.ActorName);
-				}
+				return MeshForComponent.Select(KVP => KVP.Value.Component);
 			}
 		}
 
@@ -578,16 +576,16 @@ namespace DatasmithSolidworks
 			return FActorName.FromString($"{ComponentName}_{ConfigName}");
 		}
 
-		public string GetMeshName(string ConfigName, FComponentName ComponentName)
+		public FMeshName GetMeshName(string ConfigName, FComponentName ComponentName)
 		{
-			return ConfigName == MainConfigurationName ? $"{ComponentName}_Mesh" : $"{ComponentName}_{ConfigName}_Mesh";
+			return FMeshName.FromString(ConfigName == MainConfigurationName ? $"{ComponentName}_Mesh" : $"{ComponentName}_{ConfigName}_Mesh");
 		}
 
 		// Get mesh for component which has only single mesh configuration
-		public string GetMeshName(FComponentName ComponentName)
+		public FMeshName GetMeshName(FComponentName ComponentName)
 		{
 			string ConfigName = Meshes.GetConfigurationNameForComponentMesh(ComponentName);
-			return ConfigName == MainConfigurationName ? $"{ComponentName}_Mesh" : $"{ComponentName}_{ConfigName}_Mesh";
+			return GetMeshName(ConfigName, ComponentName);
 		}
 	}
 }
\ No newline at end of file
diff --git a/Engine/Source/Programs/Enterprise/Datasmith/DatasmithSolidworksExporter/Private/ConfigurationTree.cs b/Engine/Source/Programs/Enterprise/Datasmith/DatasmithSolidworksExporter/Private/ConfigurationTree.cs
--- a/Engine/Source/Programs/Enterprise/Datasmith/DatasmithSolidworksExporter/Private/ConfigurationTree.cs
+++ b/Engine/Source/Programs/Enterprise/Datasmith/DatasmithSolidworksExporter/Private/ConfigurationTree.cs
@@ -5,6 +5,7 @@ using System.Collections.Generic;
 using System.Diagnostics;
 using System.Linq;
 using System.Runtime.InteropServices;
+using DatasmithSolidworks.Names;
 
 namespace DatasmithSolidworks
 {
diff --git a/Engine/Source/Programs/Enterprise/Datasmith/DatasmithSolidworksExporter/Private/DatasmithExporter.cs b/Engine/Source/Programs/Enterprise/Datasmith/DatasmithSolidworksExporter/Private/DatasmithExporter.cs
--- a/Engine/Source/Programs/Enterprise/Datasmith/DatasmithSolidworksExporter/Private/DatasmithExporter.cs
+++ b/Engine/Source/Programs/Enterprise/Datasmith/DatasmithSolidworksExporter/Private/DatasmithExporter.cs
@@ -9,162 +9,141 @@ using System.IO;
 using System.Linq;
 using System.Runtime.InteropServices;
 using System.Threading.Tasks;
+using DatasmithSolidworks.Names;
 using static DatasmithSolidworks.FConfigurationTree;
 
+
 namespace DatasmithSolidworks
 {
-	public enum EActorType
+	// Wrapping Component Name, Actor Name and Mesh Name to distinguish one from the other in the code. This helps to
+	//  - make explicit which name where
+	//  - avoid using sanitized(Datasmith) name to check component's name and vice versa
+	namespace Names
 	{
-		SimpleActor,
-		MeshActor,
-		PointLightActor,
-		SpotLightActor,
-		DirLightActor
-	};
+		[ComVisible(false)]
+		public struct FActorName : IEquatable<FActorName>
+		{
+			private string Value;
 
-	// Wrapping Component Name and Actor Name to distinguish one from the other in the code. This helps to
-	//  - make explicit which name where
-	//  - avoid using sanitized name to check component's name and vice versa
+			public FActorName(FComponentName ComponentName)
+			{
+				Value = FDatasmithExporter.SanitizeName(ComponentName.GetString());
+			}
+			private FActorName(string InValue) => Value = InValue;
 
-	[ComVisible(false)]
-	public struct FActorName
-	{
-		private string Value;
+			public static FActorName FromString(string Value) => new FActorName(Value);
 
-		public FActorName(FComponentName ComponentName)
-		{
-			Value = FDatasmithExporter.SanitizeName(ComponentName.GetString());
-		}
+			// Following is common between 'Names' but unfortunately C# has no way to generalize this for structs
+			public bool IsValid() => !string.IsNullOrEmpty(Value);
+			public string GetString() => Value;
 
-		public static FActorName FromString(string Value)
-		{
-			return new FActorName(Value);
-		}
+			public override string ToString() => Value;
 
-		private FActorName(string InValue)
-		{
-			Value = InValue;
-		}
+			public override bool Equals(object Obj) => Obj is FActorName Other && this == Other;
+			public bool Equals(FActorName Other) => Value == Other.Value;
+			public static bool operator ==(FActorName A, FActorName B) => A.Equals(B);
+			public static bool operator !=(FActorName A, FActorName B) => !(A == B);
 
-		public bool IsValid()
-		{
-			return !string.IsNullOrEmpty(Value);
+			public override int GetHashCode() => Value.GetHashCode();
 		}
 
-		public string GetString()
+		[ComVisible(false)]
+		public struct FComponentName: IEquatable<FComponentName>
 		{
-			return Value;
-		}
+			private string Value;
 
-		public override string ToString()
-		{
-			return Value;
-		}
+			public FComponentName(Component2 Component)
+			{
+				Value = Component.Name2;
+			}
 
-		public override bool Equals(object Obj)
-		{
-			if (Obj is FActorName Other)
+			public FComponentName(IComponent2 Component)
 			{
-				return this == Other;
-			}			
-			return false;
-		}
+				Value = Component.Name2;
+			}
 
-		public static bool operator ==(FActorName A, FActorName B)
-		{
-			return A.Value == B.Value;
-		}
+			private FComponentName(string Name)
+			{
+				Value = Name;
+			}
 
-		public static bool operator !=(FActorName A, FActorName B)
-		{
-			return !(A == B);
-		}
+			/// Used to convert string to component name ONLY when received from Solidworks API
+			public static FComponentName FromApiString(string Name)
+			{
+				return new FComponentName(Name);
+			}
 
-		public override int GetHashCode()
-		{
-			return Value.GetHashCode();
-		}
-	}
+			/// Used to convert from a custom string to represent something 'like' a Solidworks component
+			/// todo: may want to avoid trying to mix components and non-components into the same 'entity'?
+			public static FComponentName FromCustomString(string Value)
+			{
+				return new FComponentName(Value);
+			}
 
-	[ComVisible(false)]
-	public struct FComponentName
-	{
-		private string Value;
+			public FActorName GetActorName()
+			{
+				return new FActorName(this);
+			}
 
-		public FComponentName(Component2 Component)
-		{
-			Value = Component.Name2;
-		}
+			public string GetLabel()
+			{
+				return Value.Split('/').Last();
+			}
 
-		public FComponentName(IComponent2 Component)
-		{
-			Value = Component.Name2;
-		}
+			// Following is common between 'Names' but unfortunately C# has no way to generalize this for structs
+			public bool IsValid() => !string.IsNullOrEmpty(Value);
+			public string GetString() => Value;
 
-		private FComponentName(string Name)
-		{
-			Value = Name;
-		}
+			public override string ToString() => Value;
+			public override bool Equals(object Obj) => Obj is FComponentName Other && this == Other;
+			public bool Equals(FComponentName Other) => Value == Other.Value;
+			public static bool operator ==(FComponentName A, FComponentName B) => A.Equals(B);
+			public static bool operator !=(FComponentName A, FComponentName B) => !(A == B);
 
-		/// Used to convert string to component name ONLY when received from Solidworks API
-		public static FComponentName FromApiString(string Name)
-		{
-			return new FComponentName(Name);
+			public override int GetHashCode() => Value.GetHashCode();
 		}
 
-		/// Used to convert from a custom string to represent something 'like' a Solidworks component
-		/// todo: may want to avoid trying to mix components and non-components into the same 'entity'?
-		public static FComponentName FromCustomString(string Value)
+		[ComVisible(false)]
+		public struct FMeshName : IEquatable<FMeshName>
 		{
-			return new FComponentName(Value);
-		}
+			private string Value;
 
-		public bool IsValid() => !string.IsNullOrEmpty(Value);
+			public FMeshName(FComponentName ComponentName) => Value = $"{ComponentName}_Mesh";
+			public FMeshName(FComponentName ComponentName, string CfgName) => Value = $"{ComponentName}_{CfgName}_Mesh";
 
-		public string GetString()
-		{
-			return Value;
-		}
-
-		public FActorName GetActorName()
-		{
-			return new FActorName(this);
-		}
+			private FMeshName(string InValue) => Value = InValue;
 
-		public string GetLabel()
-		{
-			return Value.Split('/').Last();
-		}
+			public static FMeshName FromString(string Value) => new FMeshName(Value);
 
-		public override string ToString()
-		{
-			return Value;
-		}
+			// Following is common between 'Names' but unfortunately C# has no way to generalize this for structs
+			public bool IsValid() => !string.IsNullOrEmpty(Value);
+			public string GetString() => Value;
 
-		public override bool Equals(object Obj)
-		{
-			if (Obj is FComponentName Other)
-			{
-				return this == Other;
-			}			
-			return false;
-		}
+			public override string ToString() => Value;
 
-		public static bool operator ==(FComponentName A, FComponentName B)
-		{
-			return A.Value == B.Value;
-		}
+			public override bool Equals(object Obj) => Obj is FMeshName Other && this == Other;
+			public bool Equals(FMeshName Other) => Value == Other.Value;
+			public static bool operator ==(FMeshName A, FMeshName B) => A.Equals(B);
+			public static bool operator !=(FMeshName A, FMeshName B) => !(A == B);
 
-		public static bool operator !=(FComponentName A, FComponentName B)
-		{
-			return !(A == B);
+			public override int GetHashCode() => Value.GetHashCode();
 		}
 
-		public override int GetHashCode()
-		{
-			return Value.GetHashCode();
-		}
 	}
+}
+
+namespace DatasmithSolidworks
+{
+
+	public enum EActorType
+	{
+		SimpleActor,
+		MeshActor,
+		PointLightActor,
+		SpotLightActor,
+		DirLightActor
+	};
+
 
 	public class FDatasmithActorExportInfo
 	{
@@ -172,7 +151,7 @@ namespace DatasmithSolidworks
 		public string Label;
 		public FActorName Name;
 		public FActorName ParentName;
-		public string MeshName;
+		public FMeshName MeshName;
 		public float[] Transform;
 		public bool bVisible;
 	};
@@ -180,7 +159,7 @@ namespace DatasmithSolidworks
 	public class FDatasmithExporter
 	{
 		private Dictionary<FActorName, Tuple<EActorType, FDatasmithFacadeActor>> ExportedActorsMap = new Dictionary<FActorName, Tuple<EActorType, FDatasmithFacadeActor>>();
-		private ConcurrentDictionary<string, FDatasmithFacadeMeshElement> ExportedMeshesMap = new ConcurrentDictionary<string, FDatasmithFacadeMeshElement>();
+		private ConcurrentDictionary<FMeshName, FDatasmithFacadeMeshElement> ExportedMeshesMap = new ConcurrentDictionary<FMeshName, FDatasmithFacadeMeshElement>();
 		private ConcurrentDictionary<int, FDatasmithFacadeMaterialInstance> ExportedMaterialsMap = new ConcurrentDictionary<int, FDatasmithFacadeMaterialInstance>();
 		private ConcurrentDictionary<string, FDatasmithFacadeTexture> ExportedTexturesMap = new ConcurrentDictionary<string, FDatasmithFacadeTexture>();
 		private Dictionary<string, FDatasmithFacadeActorBinding> ExportedActorBindingsMap = new Dictionary<string, FDatasmithFacadeActorBinding>();
@@ -188,6 +167,17 @@ namespace DatasmithSolidworks
 
 		private FDatasmithFacadeScene DatasmithScene = null;
 
+		[ComVisible(false)]
+		public struct FMeshExportInfo
+		{
+			public FComponentName ComponentName;
+			public FMeshName MeshName;
+			public FActorName ActorName;
+			public FMeshData MeshData;
+
+			public FDatasmithFacadeMeshElement MeshElement;
+		}
+
 		public FDatasmithExporter(FDatasmithFacadeScene InScene)
 		{
 			DatasmithScene = InScene;
@@ -259,24 +249,27 @@ namespace DatasmithSolidworks
 			if (InExportInfo.Type == EActorType.MeshActor)
 			{
 				FDatasmithFacadeActorMesh MeshActor = Actor as FDatasmithFacadeActorMesh;
-				MeshActor.SetMesh(InExportInfo.MeshName);
+				if (InExportInfo.MeshName.IsValid())
+				{
+					MeshActor.SetMesh(GetDatasmithMeshName(InExportInfo.MeshName));
+				}
 			}
 
 			return Actor;
 		}
 
-		public string AddMesh(FDatasmithFacadeMeshElement NewMesh)
+		public void AddMesh(FMeshName MeshName, FDatasmithFacadeMeshElement NewMesh)
 		{
-			string Name = NewMesh.GetName();
-			RemoveMesh(Name);
-			ExportedMeshesMap.TryAdd(Name, NewMesh);
+			Debug.Assert(MeshName.IsValid());
+			Debug.Assert(NewMesh != null);
+			RemoveMesh(MeshName);
+			ExportedMeshesMap.TryAdd(MeshName, NewMesh);
 			DatasmithScene.AddMesh(NewMesh);
-			return Name;
 		}
 
-		public void RemoveMesh(string MeshName)
+		public void RemoveMesh(FMeshName MeshName)
 		{
-            if (MeshName == null)
+            if (!MeshName.IsValid())
             {
                 return;
             }
@@ -357,11 +350,9 @@ namespace DatasmithSolidworks
 			}
 		}
 
-		public bool ExportMesh(string InMeshName, FMeshData InData, FActorName InUpdateMeshActor, out FDatasmithFacadeMeshElement OutMeshElement)
+		public bool ExportMesh(FMeshName InMeshName, FMeshData InData, out FDatasmithFacadeMeshElement OutMeshElement)
 		{
-			Addin.Instance.LogDebug($"ExportMesh(MeshName: {InMeshName}, MeshActor: {InUpdateMeshActor.GetString()})");
-
-			InMeshName = SanitizeName(InMeshName);  //Compute mesh name early(it might be needed to remove old mesh)
+			string MeshNameInitial = InMeshName.GetString(); // Make Datasmith Mesh string to build Mesh/MeshElement from, Datasmith will modify it with FDatasmithUtils::SanitizeObjectName
 
 			OutMeshElement = null;
 
@@ -377,91 +368,83 @@ namespace DatasmithSolidworks
 				return false;
 			}
 
-			FDatasmithFacadeMesh Mesh = new FDatasmithFacadeMesh();
-			Mesh.SetName(InMeshName);
+			using(FDatasmithFacadeMesh Mesh = new FDatasmithFacadeMesh())
+			{
+				Mesh.SetName(MeshNameInitial);
 
-			FDatasmithFacadeMeshElement MeshElement = new FDatasmithFacadeMeshElement(InMeshName);
+				FDatasmithFacadeMeshElement MeshElement = new FDatasmithFacadeMeshElement(MeshNameInitial);
 
-			Mesh.SetVerticesCount(InData.Vertices.Length);
-			Mesh.SetFacesCount(InData.Triangles.Length);
+				Mesh.SetVerticesCount(InData.Vertices.Length);
+				Mesh.SetFacesCount(InData.Triangles.Length);
 
-			for (int i = 0; i < InData.Vertices.Length; i++)
-			{
-				Mesh.SetVertex(i, InData.Vertices[i].X, InData.Vertices[i].Y, InData.Vertices[i].Z);
-			}
-			for (int i = 0; i < InData.Normals.Length; i++)
-			{
-				Mesh.SetNormal(i, InData.Normals[i].X, InData.Normals[i].Y, InData.Normals[i].Z);
-			}
+				for (int i = 0; i < InData.Vertices.Length; i++)
+				{
+					Mesh.SetVertex(i, InData.Vertices[i].X, InData.Vertices[i].Y, InData.Vertices[i].Z);
+				}
 
-			if (InData.TexCoords != null)
-			{
-				Mesh.SetUVChannelsCount(1);
-				Mesh.SetUVCount(0, InData.TexCoords.Length);
-				for (int i = 0; i < InData.TexCoords.Length; i++)
+				for (int i = 0; i < InData.Normals.Length; i++)
 				{
-					Mesh.SetUV(0, i, InData.TexCoords[i].X, InData.TexCoords[i].Y);
+					Mesh.SetNormal(i, InData.Normals[i].X, InData.Normals[i].Y, InData.Normals[i].Z);
 				}
 
-			}
+				if (InData.TexCoords != null)
+				{
+					Mesh.SetUVChannelsCount(1);
+					Mesh.SetUVCount(0, InData.TexCoords.Length);
+					for (int i = 0; i < InData.TexCoords.Length; i++)
+					{
+						Mesh.SetUV(0, i, InData.TexCoords[i].X, InData.TexCoords[i].Y);
+					}
 
-			HashSet<int> MeshAddedMaterials = new HashSet<int>();
+				}
 
-			HashSet<int> MeshMaterialIds = new HashSet<int>();
+				HashSet<int> MeshAddedMaterials = new HashSet<int>();
 
-			for (int TriIndex = 0; TriIndex < InData.Triangles.Length; TriIndex++)
-			{
-				FTriangle Triangle = InData.Triangles[TriIndex];
-				int MatID = 0;
+				HashSet<int> MeshMaterialIds = new HashSet<int>();
 
-				MeshMaterialIds.Add(Triangle.MaterialID);
-				if (Triangle.MaterialID >= 1)
+				for (int TriIndex = 0; TriIndex < InData.Triangles.Length; TriIndex++)
 				{
-					if (!MeshAddedMaterials.Contains(Triangle.MaterialID))
+					FTriangle Triangle = InData.Triangles[TriIndex];
+					int MatID = 0;
+
+					MeshMaterialIds.Add(Triangle.MaterialID);
+					if (Triangle.MaterialID >= 1)
 					{
-						Addin.Instance.LogDebug($"  set material if it's exported: {Triangle.MaterialID}");
+						if (!MeshAddedMaterials.Contains(Triangle.MaterialID))
+						{
+							Addin.Instance.LogDebug($"  set material if it's exported: {Triangle.MaterialID}");
 
-						FDatasmithFacadeMaterialInstance Material = null;
-						ExportedMaterialsMap.TryGetValue(Triangle.MaterialID, out Material);
+							FDatasmithFacadeMaterialInstance Material = null;
+							ExportedMaterialsMap.TryGetValue(Triangle.MaterialID, out Material);
 
-						if (Material != null)
-						{
-							Addin.Instance.LogDebug($" SetMaterial({Material.GetName()}, SlotId: {Triangle.MaterialID})");
+							if (Material != null)
+							{
+								Addin.Instance.LogDebug(
+									$" SetMaterial({Material.GetName()}, SlotId: {Triangle.MaterialID})");
 
-							MeshAddedMaterials.Add(Triangle.MaterialID);
-							MeshElement.SetMaterial(Material.GetName(), Triangle.MaterialID);
+								MeshAddedMaterials.Add(Triangle.MaterialID);
+								MeshElement.SetMaterial(Material.GetName(), Triangle.MaterialID);
+								MatID = Triangle.MaterialID;
+							}
+						}
+						else
+						{
 							MatID = Triangle.MaterialID;
 						}
 					}
-					else
-					{
-						MatID = Triangle.MaterialID;
-					}
-				}
 
-				Mesh.SetFace(TriIndex, Triangle[0], Triangle[1], Triangle[2], MatID);
-				Mesh.SetFaceUV(TriIndex, 0, Triangle[0], Triangle[1], Triangle[2]);
-			}
-
-			Addin.Instance.LogDebug($"  material ids assigned to faces: {string.Join(", ", MeshMaterialIds)})");
-			if (MeshAddedMaterials.Count > 0)
-			{
-				Addin.Instance.LogDebug($"  assigned materials: {string.Join(", ", MeshAddedMaterials)})");
-			}
-
-			OutMeshElement = MeshElement;
-
-			DatasmithScene.ExportDatasmithMesh(MeshElement, Mesh);
+					Mesh.SetFace(TriIndex, Triangle[0], Triangle[1], Triangle[2], MatID);
+					Mesh.SetFaceUV(TriIndex, 0, Triangle[0], Triangle[1], Triangle[2]);
+				}
 
-			if (InUpdateMeshActor.IsValid())
-			{
-				Tuple<EActorType, FDatasmithFacadeActor> ExportedActorInfo = null;
-				if (ExportedActorsMap.TryGetValue(InUpdateMeshActor, out ExportedActorInfo) && ExportedActorInfo.Item1 == EActorType.MeshActor)
+				Addin.Instance.LogDebug($"  material ids assigned to faces: {string.Join(", ", MeshMaterialIds)})");
+				if (MeshAddedMaterials.Count > 0)
 				{
-					FDatasmithFacadeActorMesh MeshActor = ExportedActorInfo.Item2 as FDatasmithFacadeActorMesh;
-					Addin.Instance.LogDebug($"  SetMesh on MeshActor: {MeshActor.GetName()}: {MeshActor.GetLabel()}  )");
-					MeshActor.SetMesh(InMeshName);
+					Addin.Instance.LogDebug($"  assigned materials: {string.Join(", ", MeshAddedMaterials)})");
 				}
+
+				DatasmithScene.ExportDatasmithMesh(MeshElement, Mesh);
+				OutMeshElement = MeshElement;
 			}
 
 			return true;
@@ -506,7 +489,7 @@ namespace DatasmithSolidworks
 
 		FDatasmithFacadeActorBinding GetActorBinding(FComponentName InComponentName, FDatasmithFacadeVariant InVariant)
 		{
-			return GetActorBinding(InComponentName.GetActorName(), InVariant);
+			return GetActorBinding(GetComponentActorName(InComponentName), InVariant);
 		}
 
 		private FDatasmithFacadeActorBinding GetActorBinding(FActorName ActorName, FDatasmithFacadeVariant InVariant)
@@ -892,6 +875,17 @@ namespace DatasmithSolidworks
 			return true;
 		}
 
+		public FActorName GetComponentActorName(FComponentName ComponentName)
+		{
+			return ComponentName.GetActorName();
+		}
+
+		public FActorName GetComponentActorName(Component2 Component)
+		{
+			return GetComponentActorName(new FComponentName(Component));
+		}
+
+
 		public void ExportAnimation(FAnimation InAnim)
 		{
 			FDatasmithFacadeLevelSequence LevelSeq = new FDatasmithFacadeLevelSequence(InAnim.Name);
@@ -902,7 +896,8 @@ namespace DatasmithSolidworks
 			{
 				FAnimation.FChannel Chan = NodePair.Value;
 				Component2 Component = Chan.Target;
-				FDatasmithFacadeTransformAnimation Anim = new FDatasmithFacadeTransformAnimation(SanitizeName(Component.Name2));
+
+				FDatasmithFacadeTransformAnimation Anim = new FDatasmithFacadeTransformAnimation(GetComponentActorName(Component).GetString());
 
 				foreach (var Keyframe in Chan.Keyframes)
 				{
@@ -1046,5 +1041,56 @@ namespace DatasmithSolidworks
 			}
 			return RotMatrix;
 		}
+
+		// Get Datasmith Name for MeshName
+		// todo: Currently, name for the DatasmithMeshElement(as well as other elements) is computed within DatasmithCore's FDatasmithUtils::SanitizeObjectName
+		//   if we want to properly identify an element we need to construct it(or SetName) and then GetName from it.
+		//   The way to make element name ahead of constructing and element(e.g. to properly make unique names for every case) is to expose FDatasmithUtils::SanitizeObjectName
+		public string GetDatasmithMeshName(FMeshName MeshName)  
+		{
+			if (ExportedMeshesMap.TryGetValue(MeshName, out FDatasmithFacadeMeshElement MeshElement))
+			{
+				return MeshElement.GetName();
+			}
+
+			return null;
+		}
+
+		// Export meshes
+		public void ExportMeshes(List<FMeshExportInfo> MeshExportInfos, out List<FMeshExportInfo> OutCreatedMeshes)
+		{
+			List<FMeshExportInfo> CreatedMeshes = new List<FMeshExportInfo>();
+
+			Parallel.ForEach(MeshExportInfos, Info =>
+			{
+				Addin.Instance.LogDebug($"ExportMesh(MeshName: {Info.MeshName}, MeshActor: {Info.ActorName})");
+
+				if (ExportMesh(Info.MeshName, Info.MeshData, out Info.MeshElement))
+				{
+					Addin.Instance.LogDebug($"  exported (DatasmithMeshName: {Info.MeshElement.GetName()})");
+					CreatedMeshes.Add(Info);
+				}
+			});
+
+			// Remember exported meshes and assign to actor
+			foreach (FMeshExportInfo Info in CreatedMeshes)
+			{
+				// Info.MeshElement should be valid here
+				AddMesh(Info.MeshName, Info.MeshElement);
+
+				if (Info.ActorName.IsValid())
+				{
+					Tuple<EActorType, FDatasmithFacadeActor> ExportedActorInfo = null;
+					if (ExportedActorsMap.TryGetValue(Info.ActorName, out ExportedActorInfo) && ExportedActorInfo.Item1 == EActorType.MeshActor)
+					{
+						FDatasmithFacadeActorMesh MeshActor = ExportedActorInfo.Item2 as FDatasmithFacadeActorMesh;
+						Addin.Instance.LogDebug($"  SetMesh on MeshActor: {MeshActor.GetName()}: {MeshActor.GetLabel()}  )");
+						MeshActor.SetMesh(Info.MeshElement.GetName());
+					}
+				}
+			}
+
+			OutCreatedMeshes = CreatedMeshes;
+		}
 	}
 }
\ No newline at end of file
diff --git a/Engine/Source/Programs/Enterprise/Datasmith/DatasmithSolidworksExporter/Private/Document.cs b/Engine/Source/Programs/Enterprise/Datasmith/DatasmithSolidworksExporter/Private/Document.cs
--- a/Engine/Source/Programs/Enterprise/Datasmith/DatasmithSolidworksExporter/Private/Document.cs
+++ b/Engine/Source/Programs/Enterprise/Datasmith/DatasmithSolidworksExporter/Private/Document.cs
@@ -6,8 +6,10 @@ using System.Collections.Concurrent;
 using System.Collections.Generic;
 using System.Diagnostics;
 using System.IO;
+using System.Linq;
 using System.Threading;
 using System.Threading.Tasks;
+using DatasmithSolidworks.Names;
 using static DatasmithSolidworks.FAssemblyDocumentTracker;
 
 namespace DatasmithSolidworks
@@ -264,7 +266,7 @@ namespace DatasmithSolidworks
 		public abstract void AddComponentMaterials(FComponentName ComponentName, FObjectMaterials Materials);
 		public abstract FObjectMaterials GetComponentMaterials(Component2 Comp);
 
-		public abstract void AddMeshForComponent(FComponentName ComponentName, string MeshName);
+		public abstract void AddMeshForComponent(FComponentName ComponentName, FMeshName MeshName);
 
 		public virtual void Destroy()
 		{
@@ -277,32 +279,37 @@ namespace DatasmithSolidworks
 		// Datasmith actor assignment doesn't need threading(and not supposed to be thread-safe?) and Solidworks multithreading is supposed to be slower(SW does all the work in main thread)
 		public void ProcessConfigurationMeshes(FMeshes.FConfiguration MeshesConfiguration, string MeshSuffix)
 		{
-			ConcurrentDictionary<Component2, FDatasmithFacadeMeshElement> CreatedMeshes =
-				new ConcurrentDictionary<Component2, FDatasmithFacadeMeshElement>();
+			
+			List<FDatasmithExporter.FMeshExportInfo> MeshExportInfos = new List<FDatasmithExporter.FMeshExportInfo>();
 
-			Parallel.ForEach(MeshesConfiguration.EnumerateComponentAndActorNames(), KVP =>
+			// Extract meshes data and prepare for parallel datasmith export
+			foreach (Component2 Comp in MeshesConfiguration.EnumerateComponents())
 			{
-				Component2 Comp = KVP.Key;
-
 				FMeshData MeshData = ExtractComponentMeshData(Comp);
-
 				MeshesConfiguration.AddMesh(Comp, MeshData);
 
 				if (MeshData != null)
 				{
-					string MeshName = MeshSuffix == null ? $"{KVP.Value}_Mesh" : $"{KVP.Value}_{MeshSuffix}_Mesh";
 
-					if (Exporter.ExportMesh(MeshName, MeshData, KVP.Value, out FDatasmithFacadeMeshElement NewMesh))
+					FActorName ComponentActorName = Exporter.GetComponentActorName(Comp);
+					FComponentName ComponentName = new FComponentName(Comp);
+					MeshExportInfos.Add(new FDatasmithExporter.FMeshExportInfo()
 					{
-						CreatedMeshes[Comp] = NewMesh;
-					}
+						ComponentName = ComponentName,
+						MeshName = FMeshName.FromString(MeshSuffix == null
+							? $"{ComponentName}_Mesh"
+							: $"{ComponentName}_{MeshSuffix}_Mesh"),
+						ActorName = ComponentActorName,
+						MeshData = MeshData
+					});
 				}
-			});
+			}
+
+			Exporter.ExportMeshes(MeshExportInfos, out List<FDatasmithExporter.FMeshExportInfo> CreatedMeshes);
 
-			foreach (var KVP in CreatedMeshes)
+			foreach (FDatasmithExporter.FMeshExportInfo Info in CreatedMeshes)
 			{
-				string MeshName = Exporter.AddMesh(KVP.Value);
-				AddMeshForComponent(new FComponentName(KVP.Key), MeshName);  // Register that this mesh was used for the component
+				AddMeshForComponent(Info.ComponentName, Info.MeshName);  // Register that this mesh was used for the component
 			}
 		}
 	};
diff --git a/Engine/Source/Programs/Enterprise/Datasmith/DatasmithSolidworksExporter/Private/MetadataManager.cs b/Engine/Source/Programs/Enterprise/Datasmith/DatasmithSolidworksExporter/Private/MetadataManager.cs
--- a/Engine/Source/Programs/Enterprise/Datasmith/DatasmithSolidworksExporter/Private/MetadataManager.cs
+++ b/Engine/Source/Programs/Enterprise/Datasmith/DatasmithSolidworksExporter/Private/MetadataManager.cs
@@ -4,6 +4,7 @@ using SolidWorks.Interop.sldworks;
 using SolidWorks.Interop.swconst;
 using System.Collections.Generic;
 using System.Runtime.InteropServices;
+using DatasmithSolidworks.Names;
 
 namespace DatasmithSolidworks
 {
diff --git a/Engine/Source/Programs/Enterprise/Datasmith/DatasmithSolidworksExporter/Private/ObjectMaterials.cs b/Engine/Source/Programs/Enterprise/Datasmith/DatasmithSolidworksExporter/Private/ObjectMaterials.cs
--- a/Engine/Source/Programs/Enterprise/Datasmith/DatasmithSolidworksExporter/Private/ObjectMaterials.cs
+++ b/Engine/Source/Programs/Enterprise/Datasmith/DatasmithSolidworksExporter/Private/ObjectMaterials.cs
@@ -7,6 +7,7 @@ using SolidWorks.Interop.swconst;
 using System.Linq;
 using System.Collections.Concurrent;
 using System.Threading.Tasks;
+using DatasmithSolidworks.Names;
 
 namespace DatasmithSolidworks
 {
diff --git a/Engine/Source/Programs/Enterprise/Datasmith/DatasmithSolidworksExporter/Private/PartDocument.cs b/Engine/Source/Programs/Enterprise/Datasmith/DatasmithSolidworksExporter/Private/PartDocument.cs
--- a/Engine/Source/Programs/Enterprise/Datasmith/DatasmithSolidworksExporter/Private/PartDocument.cs
+++ b/Engine/Source/Programs/Enterprise/Datasmith/DatasmithSolidworksExporter/Private/PartDocument.cs
@@ -7,7 +7,9 @@ using System.Collections.Concurrent;
 using System.Collections.Generic;
 using System.IO;
 using System.Security.Cryptography;
+using DatasmithSolidworks.Names;
 using static System.Windows.Forms.VisualStyles.VisualStyleElement;
+using static DatasmithSolidworks.FDatasmithExporter;
 
 namespace DatasmithSolidworks
 {
@@ -16,7 +18,7 @@ namespace DatasmithSolidworks
 		public PartDoc SwPartDoc { get; private set; } = null;
 		private FAssemblyDocumentTracker AsmDoc = null;
 		private FComponentName ComponentName;
-		private string MeshName = null;
+		private FMeshName MeshName = new FMeshName();
 		public readonly FPartDocument PartDocument;
 
 		public FObjectMaterials ExportedPartMaterials { get; set; }
@@ -116,9 +118,9 @@ namespace DatasmithSolidworks
 
 			SetExportStatus($"{PartName} Meshes");
 
+			ConfigurationManager ConfigManager = (SwPartDoc as ModelDoc2).ConfigurationManager;
 			if (bHasConfigurations)
 			{
-				ConfigurationManager ConfigManager = (SwPartDoc as ModelDoc2).ConfigurationManager;
 				string ActiveConfigurationName = ConfigManager.ActiveConfiguration.Name;
 				FConfigurationExporter ConfigurationExporter = new FConfigurationExporter(Meshes, null, ActiveConfigurationName);
 
@@ -138,14 +140,21 @@ namespace DatasmithSolidworks
 					FActorName ActorName =
 						ConfigurationExporter.GetMeshActorName(ConfigurationName, RootComponentName);
 
-					bool bHasMesh = Exporter.ExportMesh($"{ActorName}_Mesh", MeshData, ActorName,
-						out FDatasmithFacadeMeshElement NewMesh);
+					FMeshName MeshNameForConfiguration = new FMeshName(RootComponentName, ConfigurationName);
 
-					string MeshNameForConfiguration = null; // todo: Sync not supported for variants, 
-					if (bHasMesh)
+					List<FMeshExportInfo> MeshExportInfos = new List<FMeshExportInfo>
 					{
-						MeshNameForConfiguration = Exporter.AddMesh(NewMesh);
-					}
+						new FMeshExportInfo
+						{
+							ComponentName = RootComponentName,
+							MeshName = MeshNameForConfiguration,
+							ActorName = ActorName,
+							MeshData = MeshData
+						}
+					};
+
+					Exporter.ExportMeshes(MeshExportInfos, out List<FMeshExportInfo> OutCreatedMeshes);
+					bool bHasMesh = OutCreatedMeshes.Count > 0;
 
 					FDatasmithActorExportInfo ExportInfo = new FDatasmithActorExportInfo();
 					ExportInfo.Name = ActorName;
@@ -158,23 +167,35 @@ namespace DatasmithSolidworks
 			}
 			else
 			{
+				FComponentName RootComponentName =
+					new FComponentName(ConfigManager.ActiveConfiguration.GetRootComponent3(true));
+
+
 				FActorName ActorName = FActorName.FromString(PartName);
 				
 				FMeshData MeshData = ExtractPartMeshData();
 
 				Exporter.RemoveMesh(MeshName);
-				MeshName = null;
+				MeshName = new FMeshName();
 				bool bHasMesh = false;
 
 				if (MeshData != null)
 				{
-					bHasMesh = Exporter.ExportMesh($"{ActorName}_Mesh", MeshData, ActorName,
-						out FDatasmithFacadeMeshElement NewMesh);
+					MeshName = new FMeshName(RootComponentName);
 
-					if (bHasMesh)
+					List<FMeshExportInfo> MeshExportInfos = new List<FMeshExportInfo>
 					{
-						MeshName = Exporter.AddMesh(NewMesh);
-					}
+						new FMeshExportInfo
+						{
+							ComponentName = RootComponentName,
+							MeshName = MeshName,
+							ActorName = ActorName,
+							MeshData = MeshData
+						}
+					};
+
+					Exporter.ExportMeshes(MeshExportInfos, out List<FMeshExportInfo> OutCreatedMeshes);
+					bHasMesh = OutCreatedMeshes.Count > 0;
 				}
 
 				FDatasmithActorExportInfo ExportInfo = new FDatasmithActorExportInfo();
@@ -209,7 +230,7 @@ namespace DatasmithSolidworks
 			return null;
 		}
 
-		public override void AddMeshForComponent(FComponentName ComponentName, string MeshName)
+		public override void AddMeshForComponent(FComponentName ComponentName, FMeshName MeshName)
 		{
 		}
 
