[Solidworks] Refactor - ComVisible removal

- made ComVisible false by default  - simplifies code and doesn't expose unneeded types

diff --git a/Engine/Source/Programs/Enterprise/Datasmith/DatasmithSolidworksExporter/DatasmithSolidworks/Properties/AssemblyInfo.cs b/Engine/Source/Programs/Enterprise/Datasmith/DatasmithSolidworksExporter/DatasmithSolidworks/Properties/AssemblyInfo.cs
--- a/Engine/Source/Programs/Enterprise/Datasmith/DatasmithSolidworksExporter/DatasmithSolidworks/Properties/AssemblyInfo.cs
+++ b/Engine/Source/Programs/Enterprise/Datasmith/DatasmithSolidworksExporter/DatasmithSolidworks/Properties/AssemblyInfo.cs
@@ -61,5 +61,5 @@ using System.Runtime.InteropServices;
 [assembly: AssemblyKeyName("")]
 [assembly: AssemblyFileVersion("1.0.0.21")]
 
-// todo: We don't need com visibility for every class. Only for the Addin and whatever we'd want to expose and automation API
-// [assembly: ComVisible(false)]
+// We don't need com visibility for every class. Only for the Addin and whatever we'd want to expose as automation API
+[assembly: ComVisible(false)]
diff --git a/Engine/Source/Programs/Enterprise/Datasmith/DatasmithSolidworksExporter/Private/Body.cs b/Engine/Source/Programs/Enterprise/Datasmith/DatasmithSolidworksExporter/Private/Body.cs
--- a/Engine/Source/Programs/Enterprise/Datasmith/DatasmithSolidworksExporter/Private/Body.cs
+++ b/Engine/Source/Programs/Enterprise/Datasmith/DatasmithSolidworksExporter/Private/Body.cs
@@ -9,14 +9,12 @@ using System.Collections.Concurrent;
 
 namespace DatasmithSolidworks
 {
-	[ComVisible(false)]
 	public class FBody
 	{
 		public List<FBodyFace> Faces { get; set; } = new List<FBodyFace>();
 		public Body2 Body { get; private set; } = null;
 		public FBoundingBox Bounds { get; private set; }
 
-		[ComVisible(false)]
 		public class FBodyFace
 		{
 			public Face2 Face { get; private set; } = null;
diff --git a/Engine/Source/Programs/Enterprise/Datasmith/DatasmithSolidworksExporter/Private/BoundingBox.cs b/Engine/Source/Programs/Enterprise/Datasmith/DatasmithSolidworksExporter/Private/BoundingBox.cs
--- a/Engine/Source/Programs/Enterprise/Datasmith/DatasmithSolidworksExporter/Private/BoundingBox.cs
+++ b/Engine/Source/Programs/Enterprise/Datasmith/DatasmithSolidworksExporter/Private/BoundingBox.cs
@@ -5,7 +5,6 @@ using System.Runtime.InteropServices;
 
 namespace DatasmithSolidworks
 {
-    [ComVisible(false)]
     public class FBoundingBox
     {
 		public FVec3 Min { get; private set; }
diff --git a/Engine/Source/Programs/Enterprise/Datasmith/DatasmithSolidworksExporter/Private/ConfigurationExporter.cs b/Engine/Source/Programs/Enterprise/Datasmith/DatasmithSolidworksExporter/Private/ConfigurationExporter.cs
--- a/Engine/Source/Programs/Enterprise/Datasmith/DatasmithSolidworksExporter/Private/ConfigurationExporter.cs
+++ b/Engine/Source/Programs/Enterprise/Datasmith/DatasmithSolidworksExporter/Private/ConfigurationExporter.cs
@@ -13,7 +13,6 @@ using static DatasmithSolidworks.FConfigurationTree;
 
 namespace DatasmithSolidworks
 {
-	[ComVisible(false)]
 	public class FConfigurationData
 	{
 		public string Name;
@@ -22,7 +21,6 @@ namespace DatasmithSolidworks
 		public Dictionary<FComponentName, bool> ComponentVisibility = new Dictionary<FComponentName, bool>();
 		public Dictionary<FComponentName, FObjectMaterials> ComponentMaterials = new Dictionary<FComponentName, FObjectMaterials>();
 
-		[ComVisible(false)]
 		public struct FComponentGeometryVariant
 		{
 			public FActorName VisibleActor; // Actor enabled for this variant
@@ -40,14 +38,12 @@ namespace DatasmithSolidworks
 		}
 	};
 
-	[ComVisible(false)]
 	public class FMeshes
 	{
 		private Dictionary<string, FConfiguration> Configurations = new Dictionary<string, FConfiguration>();
 
 		private HashSet<FMesh> Meshes = new HashSet<FMesh>();
 
-		[ComVisible(false)]
 		public struct FMesh
 		{
 			public FMeshData MeshData;
@@ -228,7 +224,6 @@ namespace DatasmithSolidworks
 		}
 	}
 
-	[ComVisible(false)]
 	public class FConfigurationExporter
 	{
 		public readonly FMeshes Meshes;
diff --git a/Engine/Source/Programs/Enterprise/Datasmith/DatasmithSolidworksExporter/Private/ConfigurationTree.cs b/Engine/Source/Programs/Enterprise/Datasmith/DatasmithSolidworksExporter/Private/ConfigurationTree.cs
--- a/Engine/Source/Programs/Enterprise/Datasmith/DatasmithSolidworksExporter/Private/ConfigurationTree.cs
+++ b/Engine/Source/Programs/Enterprise/Datasmith/DatasmithSolidworksExporter/Private/ConfigurationTree.cs
@@ -9,7 +9,6 @@ using DatasmithSolidworks.Names;
 
 namespace DatasmithSolidworks
 {
-	[ComVisible(false)]
 	public class FConfigurationTree
 	{
 		// Component state for a single configuration
@@ -49,12 +48,10 @@ namespace DatasmithSolidworks
 		};
 
 		// A node contains all the configuration data for a single component
-		[ComVisible(false)]
 		public class FComponentTreeNode
 		{
 			public readonly Component2 Component;
 
-			[ComVisible(false)]
 			public struct FComponentInfo
 			{
 				public FComponentName ComponentName;
diff --git a/Engine/Source/Programs/Enterprise/Datasmith/DatasmithSolidworksExporter/Private/ConstructionGeometry.cs b/Engine/Source/Programs/Enterprise/Datasmith/DatasmithSolidworksExporter/Private/ConstructionGeometry.cs
--- a/Engine/Source/Programs/Enterprise/Datasmith/DatasmithSolidworksExporter/Private/ConstructionGeometry.cs
+++ b/Engine/Source/Programs/Enterprise/Datasmith/DatasmithSolidworksExporter/Private/ConstructionGeometry.cs
@@ -7,7 +7,6 @@ using System.Runtime.InteropServices;
 
 namespace DatasmithSolidworks
 {
-	[ComVisible(false)]
 	public class FConstructionGeometry
 	{
 		public List<FVec3> CVertices;
diff --git a/Engine/Source/Programs/Enterprise/Datasmith/DatasmithSolidworksExporter/Private/DatasmithExporter.cs b/Engine/Source/Programs/Enterprise/Datasmith/DatasmithSolidworksExporter/Private/DatasmithExporter.cs
--- a/Engine/Source/Programs/Enterprise/Datasmith/DatasmithSolidworksExporter/Private/DatasmithExporter.cs
+++ b/Engine/Source/Programs/Enterprise/Datasmith/DatasmithSolidworksExporter/Private/DatasmithExporter.cs
@@ -20,7 +20,6 @@ namespace DatasmithSolidworks
 	//  - avoid using sanitized(Datasmith) name to check component's name and vice versa
 	namespace Names
 	{
-		[ComVisible(false)]
 		public struct FActorName : IEquatable<FActorName>
 		{
 			private string Value;
@@ -47,7 +46,6 @@ namespace DatasmithSolidworks
 			public override int GetHashCode() => Value.GetHashCode();
 		}
 
-		[ComVisible(false)]
 		public struct FComponentName: IEquatable<FComponentName>
 		{
 			private string Value;
@@ -103,7 +101,6 @@ namespace DatasmithSolidworks
 			public override int GetHashCode() => Value.GetHashCode();
 		}
 
-		[ComVisible(false)]
 		public struct FMeshName : IEquatable<FMeshName>
 		{
 			private string Value;
@@ -167,7 +164,6 @@ namespace DatasmithSolidworks
 
 		private FDatasmithFacadeScene DatasmithScene = null;
 
-		[ComVisible(false)]
 		public struct FMeshExportInfo
 		{
 			public FComponentName ComponentName;
diff --git a/Engine/Source/Programs/Enterprise/Datasmith/DatasmithSolidworksExporter/Private/GeometryChunk.cs b/Engine/Source/Programs/Enterprise/Datasmith/DatasmithSolidworksExporter/Private/GeometryChunk.cs
--- a/Engine/Source/Programs/Enterprise/Datasmith/DatasmithSolidworksExporter/Private/GeometryChunk.cs
+++ b/Engine/Source/Programs/Enterprise/Datasmith/DatasmithSolidworksExporter/Private/GeometryChunk.cs
@@ -4,7 +4,6 @@ using System.Runtime.InteropServices;
 
 namespace DatasmithSolidworks
 {
-    [ComVisible(false)]
     public class FGeometryChunk
     {
         public FVec3[] Vertices { get; set; } = null;
diff --git a/Engine/Source/Programs/Enterprise/Datasmith/DatasmithSolidworksExporter/Private/Material.cs b/Engine/Source/Programs/Enterprise/Datasmith/DatasmithSolidworksExporter/Private/Material.cs
--- a/Engine/Source/Programs/Enterprise/Datasmith/DatasmithSolidworksExporter/Private/Material.cs
+++ b/Engine/Source/Programs/Enterprise/Datasmith/DatasmithSolidworksExporter/Private/Material.cs
@@ -10,7 +10,6 @@ using SolidWorks.Interop.sldworks;
 
 namespace DatasmithSolidworks
 {
-	[ComVisible(false)]
 	public class FMaterial
 	{
 		private static List<Tuple<string, string>> SpecialFinishes = new List<Tuple<string, string>>();
diff --git a/Engine/Source/Programs/Enterprise/Datasmith/DatasmithSolidworksExporter/Private/MathUtils.cs b/Engine/Source/Programs/Enterprise/Datasmith/DatasmithSolidworksExporter/Private/MathUtils.cs
--- a/Engine/Source/Programs/Enterprise/Datasmith/DatasmithSolidworksExporter/Private/MathUtils.cs
+++ b/Engine/Source/Programs/Enterprise/Datasmith/DatasmithSolidworksExporter/Private/MathUtils.cs
@@ -7,7 +7,6 @@ using System.Runtime.InteropServices;
 
 namespace DatasmithSolidworks
 {
-	[ComVisible(false)]
     public static class MathUtils
     {
 		public static float Rad2Deg { get { return (float)(180.0 / Math.PI); } }
diff --git a/Engine/Source/Programs/Enterprise/Datasmith/DatasmithSolidworksExporter/Private/Matrix4.cs b/Engine/Source/Programs/Enterprise/Datasmith/DatasmithSolidworksExporter/Private/Matrix4.cs
--- a/Engine/Source/Programs/Enterprise/Datasmith/DatasmithSolidworksExporter/Private/Matrix4.cs
+++ b/Engine/Source/Programs/Enterprise/Datasmith/DatasmithSolidworksExporter/Private/Matrix4.cs
@@ -5,7 +5,6 @@ using System.Runtime.InteropServices;
 
 namespace DatasmithSolidworks
 {
-	[ComVisible(false)]
 	public class FMatrix4
 	{
 		public FVec3 XBasis { get { return new FVec3(Data[0], Data[1], Data[2]); } }
diff --git a/Engine/Source/Programs/Enterprise/Datasmith/DatasmithSolidworksExporter/Private/MeshData.cs b/Engine/Source/Programs/Enterprise/Datasmith/DatasmithSolidworksExporter/Private/MeshData.cs
--- a/Engine/Source/Programs/Enterprise/Datasmith/DatasmithSolidworksExporter/Private/MeshData.cs
+++ b/Engine/Source/Programs/Enterprise/Datasmith/DatasmithSolidworksExporter/Private/MeshData.cs
@@ -5,7 +5,6 @@ using System.Runtime.InteropServices;
 
 namespace DatasmithSolidworks
 {
-	[ComVisible(false)]
     public class FMeshData
     {
         public FVec3[] Vertices { get; set; } = null;
diff --git a/Engine/Source/Programs/Enterprise/Datasmith/DatasmithSolidworksExporter/Private/MetadataManager.cs b/Engine/Source/Programs/Enterprise/Datasmith/DatasmithSolidworksExporter/Private/MetadataManager.cs
--- a/Engine/Source/Programs/Enterprise/Datasmith/DatasmithSolidworksExporter/Private/MetadataManager.cs
+++ b/Engine/Source/Programs/Enterprise/Datasmith/DatasmithSolidworksExporter/Private/MetadataManager.cs
@@ -47,7 +47,6 @@ namespace DatasmithSolidworks
 		}
 	}
 
-	[ComVisible(false)]
 	public class FMetadata
 	{
 		public enum EOwnerType
diff --git a/Engine/Source/Programs/Enterprise/Datasmith/DatasmithSolidworksExporter/Private/ObjectMaterials.cs b/Engine/Source/Programs/Enterprise/Datasmith/DatasmithSolidworksExporter/Private/ObjectMaterials.cs
--- a/Engine/Source/Programs/Enterprise/Datasmith/DatasmithSolidworksExporter/Private/ObjectMaterials.cs
+++ b/Engine/Source/Programs/Enterprise/Datasmith/DatasmithSolidworksExporter/Private/ObjectMaterials.cs
@@ -11,7 +11,6 @@ using DatasmithSolidworks.Names;
 
 namespace DatasmithSolidworks
 {
-	[ComVisible(false)]
 	public class FObjectMaterials
 	{
 		// Note: part document might be a different document than the owner one:
diff --git a/Engine/Source/Programs/Enterprise/Datasmith/DatasmithSolidworksExporter/Private/StripGeometry.cs b/Engine/Source/Programs/Enterprise/Datasmith/DatasmithSolidworksExporter/Private/StripGeometry.cs
--- a/Engine/Source/Programs/Enterprise/Datasmith/DatasmithSolidworksExporter/Private/StripGeometry.cs
+++ b/Engine/Source/Programs/Enterprise/Datasmith/DatasmithSolidworksExporter/Private/StripGeometry.cs
@@ -7,7 +7,6 @@ using System.Threading.Tasks;
 
 namespace DatasmithSolidworks
 {
-	[ComVisible(false)]
 	public class FStripGeometryBody
 	{
 		public FBoundingBox Bounds { get; set; } = null;
@@ -26,7 +25,6 @@ namespace DatasmithSolidworks
 		}
 	}
 
-	[ComVisible(false)]
 	public class FStripGeometryFace
 	{
 		public FMaterial Material { get; set; } = null;
@@ -61,7 +59,6 @@ namespace DatasmithSolidworks
 		}
 	}
 
-	[ComVisible(false)]
 	public class FStripGeometry
 	{
 		private ConcurrentBag<FStripGeometryBody> Bodies { get; set; } = new ConcurrentBag<FStripGeometryBody>();
diff --git a/Engine/Source/Programs/Enterprise/Datasmith/DatasmithSolidworksExporter/Private/StripUnion.cs b/Engine/Source/Programs/Enterprise/Datasmith/DatasmithSolidworksExporter/Private/StripUnion.cs
--- a/Engine/Source/Programs/Enterprise/Datasmith/DatasmithSolidworksExporter/Private/StripUnion.cs
+++ b/Engine/Source/Programs/Enterprise/Datasmith/DatasmithSolidworksExporter/Private/StripUnion.cs
@@ -4,7 +4,6 @@ using System.Runtime.InteropServices;
 
 namespace DatasmithSolidworks
 {
-    [ComVisible(false)]
     [StructLayout(LayoutKind.Explicit)]
     public struct FStripUnion
     {
diff --git a/Engine/Source/Programs/Enterprise/Datasmith/DatasmithSolidworksExporter/Private/StripVector.cs b/Engine/Source/Programs/Enterprise/Datasmith/DatasmithSolidworksExporter/Private/StripVector.cs
--- a/Engine/Source/Programs/Enterprise/Datasmith/DatasmithSolidworksExporter/Private/StripVector.cs
+++ b/Engine/Source/Programs/Enterprise/Datasmith/DatasmithSolidworksExporter/Private/StripVector.cs
@@ -4,7 +4,6 @@ using System.Runtime.InteropServices;
 
 namespace DatasmithSolidworks
 {
-    [ComVisible(false)]
     public class FStripVector
     {
         private FStripUnion BaseArray;
diff --git a/Engine/Source/Programs/Enterprise/Datasmith/DatasmithSolidworksExporter/Private/Triangle.cs b/Engine/Source/Programs/Enterprise/Datasmith/DatasmithSolidworksExporter/Private/Triangle.cs
--- a/Engine/Source/Programs/Enterprise/Datasmith/DatasmithSolidworksExporter/Private/Triangle.cs
+++ b/Engine/Source/Programs/Enterprise/Datasmith/DatasmithSolidworksExporter/Private/Triangle.cs
@@ -4,7 +4,6 @@ using System.Runtime.InteropServices;
 
 namespace DatasmithSolidworks
 {
-	[ComVisible(false)]
 	public class FTriangle
 	{
 		private int[] Indices = new int[3];
diff --git a/Engine/Source/Programs/Enterprise/Datasmith/DatasmithSolidworksExporter/Private/TriangleStrip.cs b/Engine/Source/Programs/Enterprise/Datasmith/DatasmithSolidworksExporter/Private/TriangleStrip.cs
--- a/Engine/Source/Programs/Enterprise/Datasmith/DatasmithSolidworksExporter/Private/TriangleStrip.cs
+++ b/Engine/Source/Programs/Enterprise/Datasmith/DatasmithSolidworksExporter/Private/TriangleStrip.cs
@@ -4,7 +4,6 @@ using System.Runtime.InteropServices;
 
 namespace DatasmithSolidworks
 {
-	[ComVisible(false)]
 	public class FTriangleStrip
 	{
 		public FStripUnion Vertices { get; set; }
diff --git a/Engine/Source/Programs/Enterprise/Datasmith/DatasmithSolidworksExporter/Private/UVPlane.cs b/Engine/Source/Programs/Enterprise/Datasmith/DatasmithSolidworksExporter/Private/UVPlane.cs
--- a/Engine/Source/Programs/Enterprise/Datasmith/DatasmithSolidworksExporter/Private/UVPlane.cs
+++ b/Engine/Source/Programs/Enterprise/Datasmith/DatasmithSolidworksExporter/Private/UVPlane.cs
@@ -4,7 +4,6 @@ using System.Runtime.InteropServices;
 
 namespace DatasmithSolidworks
 {
-    [ComVisible(false)]
     public class FUVPlane
     {
         public FVec3 UDirection { get; set; }
diff --git a/Engine/Source/Programs/Enterprise/Datasmith/DatasmithSolidworksExporter/Private/Vec2.cs b/Engine/Source/Programs/Enterprise/Datasmith/DatasmithSolidworksExporter/Private/Vec2.cs
--- a/Engine/Source/Programs/Enterprise/Datasmith/DatasmithSolidworksExporter/Private/Vec2.cs
+++ b/Engine/Source/Programs/Enterprise/Datasmith/DatasmithSolidworksExporter/Private/Vec2.cs
@@ -5,7 +5,6 @@ using System.Runtime.InteropServices;
 
 namespace DatasmithSolidworks
 {
-    [ComVisible(false)]
     public class FVec2
     {
         public float X = 0f;
diff --git a/Engine/Source/Programs/Enterprise/Datasmith/DatasmithSolidworksExporter/Private/Vec3.cs b/Engine/Source/Programs/Enterprise/Datasmith/DatasmithSolidworksExporter/Private/Vec3.cs
--- a/Engine/Source/Programs/Enterprise/Datasmith/DatasmithSolidworksExporter/Private/Vec3.cs
+++ b/Engine/Source/Programs/Enterprise/Datasmith/DatasmithSolidworksExporter/Private/Vec3.cs
@@ -6,7 +6,6 @@ using System.Runtime.InteropServices;
 
 namespace DatasmithSolidworks
 {
-    [ComVisible(false)]
     public class FVec3
     {
         public float X = 0f;
diff --git a/Engine/Source/Programs/Enterprise/Datasmith/DatasmithSolidworksExporter/Private/Vertex.cs b/Engine/Source/Programs/Enterprise/Datasmith/DatasmithSolidworksExporter/Private/Vertex.cs
--- a/Engine/Source/Programs/Enterprise/Datasmith/DatasmithSolidworksExporter/Private/Vertex.cs
+++ b/Engine/Source/Programs/Enterprise/Datasmith/DatasmithSolidworksExporter/Private/Vertex.cs
@@ -4,7 +4,6 @@ using System.Runtime.InteropServices;
 
 namespace DatasmithSolidworks
 {
-    [ComVisible(false)]
     public class FVertex
     {
         public FVec3 P { get; set; }
