[Solidworks] display states for configurations

- added vast of material logging - it's a complicated in the api and currently in the code so should help with future debugging(and already helped)

- fixed another potential race condition that I've overlooked previously

diff --git a/Engine/Source/Programs/Enterprise/Datasmith/DatasmithSolidworksExporter/Private/Addin.cs b/Engine/Source/Programs/Enterprise/Datasmith/DatasmithSolidworksExporter/Private/Addin.cs
--- a/Engine/Source/Programs/Enterprise/Datasmith/DatasmithSolidworksExporter/Private/Addin.cs
+++ b/Engine/Source/Programs/Enterprise/Datasmith/DatasmithSolidworksExporter/Private/Addin.cs
@@ -41,17 +41,80 @@ namespace DatasmithSolidworks
 		public FDocument CurrentDocument { get; private set; } = null;
 
 		// DebugLog is enabled with DatasmithSolidworksDebugOutput conditional compilation symbol
-		private ConcurrentQueue<string> DebugLog;
-		private Thread LogWriterThread;
 
-		public static Addin Instance { get; private set; } = null;
+		class FDebugLog
+		{
+			private readonly int MainThreadId;
+			private ConcurrentQueue<string> MessagesQueue = new ConcurrentQueue<string>();
+			private Thread LogWriterThread;
+			private int Indentation = 0;
+
+			public FDebugLog(int InMainThreadId)
+			{
+				MainThreadId = InMainThreadId;
+				LogWriterThread = new Thread(() =>
+				{
+					LogWriterProc();
+				});
+
+				LogWriterThread.Start();
+			}
+
+			private void LogWriterProc()
+			{
+				string LogPath = Path.Combine(Environment.GetFolderPath(Environment.SpecialFolder.LocalApplicationData),
+					"UnrealDatasmithExporter/Saved/Logs/UnrealDatasmithSolidworksExporterDebug.log");
+
+				StreamWriter LogFile = new StreamWriter(LogPath);
+
+				while (true)
+				{
+					while (MessagesQueue.TryDequeue(out string Message))
+					{
+						LogFile.WriteLine(Message);
+					}
+					LogFile.Flush();
+
+					Thread.Sleep(10);
+				}
+			}
+
+			public void LogDebug(string Message)  
+			{
+				// todo: in general, Solidworks api shouldn't be called from another thread(it's slower) 
+				//   so better to identify all those places and fix them
+				// Debug.Assert(MainThreadId == Thread.CurrentThread.ManagedThreadId);
+				MessagesQueue.Enqueue(new string(' ', Indentation*2)+Message);
+			}
+
+			public void LogDebugThread(string Message)  
+			{
+				MessagesQueue.Enqueue(new string(' ', Indentation*2)+Message);
+			}
+
+			public void Dedent()
+			{
+				Indentation --;
+			}
+
+			public void Indent()
+			{
+				Indentation ++;
+			}
+		};
+
+		private FDebugLog DebugLog;
+
+		int SwThreadId;  // Main thread of Solidworks
 
+		public static Addin Instance { get; private set; } = null;
 
 		public Addin()
 		{
 			if (Instance == null)
 			{
 				Instance = this;
+				SwThreadId = Thread.CurrentThread.ManagedThreadId;
 				StartLogWriterTread();
 			}
 		}
@@ -533,40 +596,31 @@ namespace DatasmithSolidworks
 		[Conditional("DatasmithSolidworksDebugOutput")]
 		private void StartLogWriterTread()
 		{
-			DebugLog = new ConcurrentQueue<string>();
-
-			LogWriterThread = new Thread(() =>
-			{
-				LogWriterProc();
-			});
-
-			LogWriterThread.Start();
+			DebugLog = new FDebugLog(SwThreadId);
 		}
 
 		[Conditional("DatasmithSolidworksDebugOutput")]
-		private void LogWriterProc()
+		public void LogDebug(string Message)  
 		{
-			string LogPath = Path.Combine(Environment.GetFolderPath(Environment.SpecialFolder.LocalApplicationData),
-				"UnrealDatasmithExporter/Saved/Logs/UnrealDatasmithSolidworksExporterDebug.log");
-
-			StreamWriter LogFile = new StreamWriter(LogPath);
+			DebugLog.LogDebug(Message);
+		}
 
-			while (true)
-			{
-				while (DebugLog.TryDequeue(out string Message))
-				{
-					LogFile.WriteLine(Message);
-				}
-				LogFile.Flush();
+		[Conditional("DatasmithSolidworksDebugOutput")]
+		public void LogDebugThread(string Message)  
+		{
+			DebugLog.LogDebugThread(Message);
+		}
 
-				Thread.Sleep(10);
-			}
+		[Conditional("DatasmithSolidworksDebugOutput")]
+		public void LogIndent()  
+		{
+			DebugLog.Indent();
 		}
 
 		[Conditional("DatasmithSolidworksDebugOutput")]
-		public void LogDebug(string Message)  
+		public void LogDedent()  
 		{
-			DebugLog.Enqueue(Message);
+			DebugLog.Dedent();
 		}
 
 		#endregion
diff --git a/Engine/Source/Programs/Enterprise/Datasmith/DatasmithSolidworksExporter/Private/ConfigurationExporter.cs b/Engine/Source/Programs/Enterprise/Datasmith/DatasmithSolidworksExporter/Private/ConfigurationExporter.cs
--- a/Engine/Source/Programs/Enterprise/Datasmith/DatasmithSolidworksExporter/Private/ConfigurationExporter.cs
+++ b/Engine/Source/Programs/Enterprise/Datasmith/DatasmithSolidworksExporter/Private/ConfigurationExporter.cs
@@ -10,6 +10,9 @@ using System.Collections.Concurrent;
 using System.Linq;
 using DatasmithSolidworks.Names;
 using static DatasmithSolidworks.FConfigurationTree;
+using System.ComponentModel;
+using System.Diagnostics;
+using static DatasmithSolidworks.FMeshes;
 
 namespace DatasmithSolidworks
 {
@@ -241,6 +244,7 @@ namespace DatasmithSolidworks
 
 		public List<FConfigurationData> ExportConfigurations(FDocumentTracker InDoc)
 		{
+			Addin.Instance.LogDebug($"FConfigurationExporter.ExportConfigurations [{string.Join(",", ConfigurationNames)}]");
 			// todo: may refactor to union behavior for 'no configuration'(i.e. active) and one-to-many configurations export
 			// Not sure how CfgNames might be technically null though
 			if (ConfigurationNames == null) // || CfgNames.Length <= 1)
@@ -274,12 +278,15 @@ namespace DatasmithSolidworks
 			bool bIsActiveConfiguration = true;
 			foreach (string CfgName in ConfigurationNamesActiveFirst)
 			{
+				Addin.Instance.LogDebug($"Configuration '{CfgName}'");
 				IConfiguration swConfiguration = InDoc.SwDoc.GetConfigurationByName(CfgName) as IConfiguration;
 
 				if (!bIsActiveConfiguration)
 				{
+					Addin.Instance.LogDebug($"ShowConfiguration2 '{CfgName}'");
 					InDoc.SwDoc.ShowConfiguration2(CfgName);
 				}
+				Addin.Instance.LogDebug($"ActiveConfiguration '{ConfigManager.ActiveConfiguration.Name}'");
 
 				string ConfigName = CfgName;
 
@@ -302,11 +309,12 @@ namespace DatasmithSolidworks
 				// Build the tree and get default materials (which aren't affected by any configuration)
 				// Use GetRootComponent3() with Resolve = true to ensure suppressed components will be loaded
 				// todo: docs says that Part document SW returns null. Not the case. But probably better to add a guard
+				Addin.Instance.LogDebug($"Components:");
 				CollectComponentsRecursive(InDoc, swConfiguration.GetRootComponent3(true), ConfigNode, MeshesConfiguration, CfgName);
 
-
 				ExportedConfigurationNames.Add(CfgName, ConfigName);
 
+				Addin.Instance.LogDebug($"Materials:");
 				Dictionary<FComponentName, FObjectMaterials> MaterialsMap = GetComponentMaterials(InDoc, DisplayStates != null ? DisplayStates[0] : null, swConfiguration);
 				SetComponentTreeMaterials(ConfigNode, MaterialsMap, null, false);
 
@@ -317,8 +325,10 @@ namespace DatasmithSolidworks
 					for (int Index = 1; Index < DisplayStateCount; ++Index)
 					{
 						string DisplayState = DisplayStates[Index];
-						MaterialsMap = GetComponentMaterials(InDoc, DisplayState, swConfiguration);
 						string DisplayStateTreeName = $"{CfgName}_DisplayState_{FDatasmithExporter.SanitizeName(DisplayState)}";
+						Addin.Instance.LogDebug($"Linked DisplayState '{DisplayState}' Materials for variant '{DisplayStateTreeName}'");
+
+						MaterialsMap = GetComponentMaterials(InDoc, DisplayState, swConfiguration);
 						SetComponentTreeMaterials(ConfigNode, MaterialsMap, DisplayStateTreeName, false);
 					}
 				}
@@ -360,9 +370,36 @@ namespace DatasmithSolidworks
 				bIsActiveConfiguration = false;  // Next enumerated configurations will be switched to(they weren't active on export start)
 			}
 
-			if (OriginalConfiguration != null)
+			if (OriginalConfiguration != null && ConfigurationNames.Length > 1)
+			{
+				Addin.Instance.LogDebug($"ShowConfiguration2 '{OriginalConfiguration.Name}'");
+				bool bShowConfigurationResult = InDoc.SwDoc.ShowConfiguration2(OriginalConfiguration.Name);
+				Debug.Assert(bShowConfigurationResult);
+			}
+
+			// Reload components for Original configuration(without this extracting materials for Display States later won't work - old component objects simply return 0) 
 			{
-				InDoc.SwDoc.ShowConfiguration2(OriginalConfiguration.Name);
+				Component2 RootComponent = OriginalConfiguration.GetRootComponent3(true);
+				Stack<Component2> Components = new Stack<Component2>();
+				Components.Push(RootComponent);
+
+				while (Components.Count != 0)
+				{
+					Component2 Component = Components.Pop();
+
+					FConfigurationTree.FComponentTreeNode NewNode = new FConfigurationTree.FComponentTreeNode(Component);
+					NewNode.ComponentInfo.ComponentName = new FComponentName(Component);
+					InDoc.AddExportedComponent(NewNode);
+
+					object[] Children = (object[])Component.GetChildren();
+					if (Children != null)
+					{
+						foreach (Component2 Child in Children)
+						{
+							Components.Push(Child);
+						}
+					}
+				}
 			}
 
 			List<string> DisplayStateConfigurations = new List<string>();
@@ -372,14 +409,33 @@ namespace DatasmithSolidworks
 				// Export display states as separate configurations
 				string[] DisplayStates = OriginalConfiguration.GetDisplayStates();
 
-				foreach (string DisplayState in DisplayStates)
+				for (int Index = 0; Index < DisplayStates.Length; Index++)
 				{
+					string DisplayState = DisplayStates[Index];
+					if (Index != 0)
+					{
+						OriginalConfiguration.ApplyDisplayState(DisplayState);
+					}
+
 					string DisplayStateConfigName = $"DisplayState_{FDatasmithExporter.SanitizeName(DisplayState)}";
 					DisplayStateConfigurations.Add(DisplayStateConfigName);
-					Dictionary<FComponentName, FObjectMaterials> MaterialsMap = GetComponentMaterials(InDoc, DisplayState, OriginalConfiguration);
+					Addin.Instance.LogDebug(
+						$"DisplayState '{DisplayState}' Materials for variant '{DisplayStateConfigName}'");
+
+					// Specifying display state doesn't seem to work in the api like this:
+					// GetComponentMaterials(InDoc, DisplayState, OriginalConfiguration);
+					//   (which becomes Ext.GetRenderMaterials2((int)swDisplayStateOpts_e.swSpecifyDisplayState, InDisplayStateNames))
+					// - component materials are empty (Comp.GetRenderMaterialsCount2 gives 0 for assembly components)
+					Dictionary<FComponentName, FObjectMaterials> MaterialsMap = GetComponentMaterials(InDoc, null, OriginalConfiguration);
 
 					SetComponentTreeMaterials(CombinedTree, MaterialsMap, DisplayStateConfigName, true);
 				}
+
+				// Restore original active display state, if there were more than one display state
+				if (DisplayStates.Length > 1)
+				{
+					OriginalConfiguration.ApplyDisplayState(DisplayStates[0]);  // Active display state is the first in the list returned by GetDisplayStates
+				}
 			}
 
 			// Remove configuration data when it's the same
@@ -429,15 +485,15 @@ namespace DatasmithSolidworks
 
 			if (InDoc is FAssemblyDocumentTracker AsmDoc)
 			{
-				HashSet<FComponentName> InComponentsSet = new HashSet<FComponentName>();
+				HashSet<FComponentName> ComponentsSet = new HashSet<FComponentName>();
 
 				foreach (FComponentName CompName in AsmDoc.SyncState.ExportedComponentsMap.Keys)
 				{
-					InComponentsSet.Add(CompName);
+					ComponentsSet.Add(CompName);
 				}
 
 				ConcurrentDictionary<FComponentName, FObjectMaterials> ComponentMaterials =
-					FObjectMaterials.LoadAssemblyMaterials(AsmDoc, InComponentsSet, Option, DisplayStates);
+					FObjectMaterials.LoadAssemblyMaterials(AsmDoc, ComponentsSet, Option, DisplayStates);
 
 				if (ComponentMaterials != null)
 				{
@@ -461,6 +517,8 @@ namespace DatasmithSolidworks
 
 		private static void SetComponentTreeMaterials(FConfigurationTree.FComponentTreeNode InComponentTree, Dictionary<FComponentName, FObjectMaterials> InComponentMaterialsMap, string InConfigurationName, bool bIsDisplayState)
 		{
+			Addin.Instance.LogDebug($"SetComponentTreeMaterials: '{InComponentTree.ComponentName}'");
+
 			FConfigurationTree.FComponentConfig TargetConfig = null;
 
 			if (InConfigurationName != null)
@@ -482,15 +540,21 @@ namespace DatasmithSolidworks
 				TargetConfig.Materials = Materials;
 			}
 
+			Addin.Instance.LogDebug($"  Materials: {TargetConfig.Materials}");
+
 			foreach (FConfigurationTree.FComponentTreeNode Child in InComponentTree.EnumChildren())
 			{
+				Addin.Instance.LogIndent();
 				SetComponentTreeMaterials(Child, InComponentMaterialsMap, InConfigurationName, bIsDisplayState);
+				Addin.Instance.LogDedent();
 			}
 		}
 
 		private static void CollectComponentsRecursive(FDocumentTracker InDoc, Component2 InComponent,
 			FComponentTreeNode InParentNode, FMeshes.FConfiguration Meshes, string CfgName)
 		{
+			Addin.Instance.LogDebug($"'{InComponent.Name2}'");
+
 			FConfigurationTree.FComponentTreeNode NewNode = new FConfigurationTree.FComponentTreeNode(InComponent);
 			InParentNode.Children.Add(NewNode);
 
@@ -539,7 +603,9 @@ namespace DatasmithSolidworks
 				foreach (object ObjChild in Children)
 				{
 					Component2 Child = (Component2)ObjChild;
+					Addin.Instance.LogIndent();
 					CollectComponentsRecursive(InDoc, Child, NewNode, Meshes, CfgName);
+					Addin.Instance.LogDedent();
 				}
 
 				// todo: sorting ensures that component order is stable in export(e.g. in variant property bindings order)
diff --git a/Engine/Source/Programs/Enterprise/Datasmith/DatasmithSolidworksExporter/Private/ConfigurationTree.cs b/Engine/Source/Programs/Enterprise/Datasmith/DatasmithSolidworksExporter/Private/ConfigurationTree.cs
--- a/Engine/Source/Programs/Enterprise/Datasmith/DatasmithSolidworksExporter/Private/ConfigurationTree.cs
+++ b/Engine/Source/Programs/Enterprise/Datasmith/DatasmithSolidworksExporter/Private/ConfigurationTree.cs
@@ -229,10 +229,23 @@ namespace DatasmithSolidworks
 				bool bAllMaterialsAreSame = true;
 				for (int Idx = 1; Idx < ConfigList.Count; Idx++)
 				{
-					if (Materials != null && (!ConfigList[Idx].Materials?.EqualMaterials(Materials) ?? false))
+					FObjectMaterials OtherMaterials = ConfigList[Idx].Materials;
+					if (Materials != null && OtherMaterials != null)
 					{
-						bAllMaterialsAreSame = false;
-						break;
+						if (!OtherMaterials.EqualMaterials(Materials))
+						{
+							bAllMaterialsAreSame = false;
+							break;
+						}
+					}
+					else
+					{
+						if (Materials != OtherMaterials)
+						{
+							// One of the material sets is null(empty) but another isn't
+							bAllMaterialsAreSame = false;
+							break;
+						}
 					}
 				}
 
@@ -402,12 +415,9 @@ namespace DatasmithSolidworks
 				}
 				if (!InNode.bMaterialSame)
 				{
-					FObjectMaterials Materials = NodeConfig.Materials;
-					if (Materials == null)
-					{
-						// There's no material set for the config, use default one
-						Materials = InNode.CommonConfig.Materials;
-					}
+					// There's no material set for the config, use default one
+					// todo: handle missing materials for a config/display set by making a 'default' material
+					FObjectMaterials Materials = NodeConfig.Materials ?? InNode.CommonConfig.Materials;
 					if (Materials != null)
 					{
 						OutData.ComponentMaterials.Add(InNode.ComponentName, Materials);
diff --git a/Engine/Source/Programs/Enterprise/Datasmith/DatasmithSolidworksExporter/Private/DatasmithExporter.cs b/Engine/Source/Programs/Enterprise/Datasmith/DatasmithSolidworksExporter/Private/DatasmithExporter.cs
--- a/Engine/Source/Programs/Enterprise/Datasmith/DatasmithSolidworksExporter/Private/DatasmithExporter.cs
+++ b/Engine/Source/Programs/Enterprise/Datasmith/DatasmithSolidworksExporter/Private/DatasmithExporter.cs
@@ -172,6 +172,11 @@ namespace DatasmithSolidworks
 			public FMeshData MeshData;
 
 			public FDatasmithFacadeMeshElement MeshElement;
+
+			public override string ToString()
+			{
+				return $"FMeshExportInfo(ComponentName={ComponentName}, MeshName={MeshName}, ActorName={ActorName}), MeshElement={(MeshElement==null?"<null>":MeshElement.GetName())}";
+			}
 		}
 
 		public FDatasmithExporter(FDatasmithFacadeScene InScene)
@@ -354,13 +359,13 @@ namespace DatasmithSolidworks
 
 			if (InData.Vertices == null || InData.Normals == null || InData.TexCoords == null || InData.Triangles == null)
 			{
-				Addin.Instance.LogDebug($"  skipping(some attributes are null) - Vertices:{InData.Vertices == null}, Normals:{InData.Normals == null}, TexCoords: {InData.TexCoords == null}, Triangles:{InData.Triangles}");
+				Addin.Instance.LogDebugThread($"  skipping(some attributes are null) - Vertices:{InData.Vertices == null}, Normals:{InData.Normals == null}, TexCoords: {InData.TexCoords == null}, Triangles:{InData.Triangles}");
 				return false;
 			}
 
 			if (InData.Vertices.Length == 0 || InData.Normals.Length == 0 || InData.TexCoords.Length == 0 || InData.Triangles.Length == 0)
 			{
-				Addin.Instance.LogDebug($"  skipping - Vertices: {InData.Vertices.Length}, Normals: {InData.Normals.Length}, TexCoords: {InData.TexCoords.Length} Triangles: {InData.Triangles.Length}");
+				Addin.Instance.LogDebugThread($"  skipping - Vertices: {InData.Vertices.Length}, Normals: {InData.Normals.Length}, TexCoords: {InData.TexCoords.Length} Triangles: {InData.Triangles.Length}");
 				return false;
 			}
 
@@ -408,14 +413,14 @@ namespace DatasmithSolidworks
 					{
 						if (!MeshAddedMaterials.Contains(Triangle.MaterialID))
 						{
-							Addin.Instance.LogDebug($"  set material if it's exported: {Triangle.MaterialID}");
+							Addin.Instance.LogDebugThread($"  set material if it's exported: {Triangle.MaterialID}");
 
 							FDatasmithFacadeMaterialInstance Material = null;
 							ExportedMaterialsMap.TryGetValue(Triangle.MaterialID, out Material);
 
 							if (Material != null)
 							{
-								Addin.Instance.LogDebug(
+								Addin.Instance.LogDebugThread(
 									$" SetMaterial({Material.GetName()}, SlotId: {Triangle.MaterialID})");
 
 								MeshAddedMaterials.Add(Triangle.MaterialID);
@@ -433,10 +438,10 @@ namespace DatasmithSolidworks
 					Mesh.SetFaceUV(TriIndex, 0, Triangle[0], Triangle[1], Triangle[2]);
 				}
 
-				Addin.Instance.LogDebug($"  material ids assigned to faces: {string.Join(", ", MeshMaterialIds)})");
+				Addin.Instance.LogDebugThread($"  material ids assigned to faces: {string.Join(", ", MeshMaterialIds)})");
 				if (MeshAddedMaterials.Count > 0)
 				{
-					Addin.Instance.LogDebug($"  assigned materials: {string.Join(", ", MeshAddedMaterials)})");
+					Addin.Instance.LogDebugThread($"  assigned materials: {string.Join(", ", MeshAddedMaterials)})");
 				}
 
 				DatasmithScene.ExportDatasmithMesh(MeshElement, Mesh);
@@ -732,7 +737,7 @@ namespace DatasmithSolidworks
 
 			ConcurrentBag<FDatasmithFacadeTexture> CreatedTextures = new ConcurrentBag<FDatasmithFacadeTexture>();
 			ConcurrentBag<FDatasmithFacadeMaterialInstance> CreatedMaterials = new ConcurrentBag<FDatasmithFacadeMaterialInstance>();
-			Parallel.ForEach(InMaterialsMap, MatKVP =>
+			foreach(KeyValuePair<int, FMaterial> MatKVP in InMaterialsMap)
 			{
 				List<FDatasmithFacadeTexture> NewMaterialTextures = null;
 				FDatasmithFacadeMaterialInstance NewMaterial = null;
@@ -745,7 +750,7 @@ namespace DatasmithSolidworks
 						CreatedTextures.Add(Texture);
 					}
 				}
-			});
+			}
 
 			Addin.Instance.LogDebug($"  AddMaterials");
 			// Adding stuff to a datasmith scene cannot be multithreaded!
@@ -1055,22 +1060,22 @@ namespace DatasmithSolidworks
 		// Export meshes
 		public void ExportMeshes(List<FMeshExportInfo> MeshExportInfos, out List<FMeshExportInfo> OutCreatedMeshes)
 		{
-			List<FMeshExportInfo> CreatedMeshes = new List<FMeshExportInfo>();
+			Addin.Instance.LogDebug($"ExportMeshes: {string.Join(", ", MeshExportInfos)}");
+
+			ConcurrentQueue<FMeshExportInfo> CreatedMeshes = new ConcurrentQueue<FMeshExportInfo>();
 
 			Parallel.ForEach(MeshExportInfos, Info =>
 			{
-				Addin.Instance.LogDebug($"ExportMesh(MeshName: {Info.MeshName}, MeshActor: {Info.ActorName})");
-
 				if (ExportMesh(Info.MeshName, Info.MeshData, out Info.MeshElement))
 				{
-					Addin.Instance.LogDebug($"  exported (DatasmithMeshName: {Info.MeshElement.GetName()})");
-					CreatedMeshes.Add(Info);
+					CreatedMeshes.Enqueue(Info);
 				}
 			});
 
 			// Remember exported meshes and assign to actor
 			foreach (FMeshExportInfo Info in CreatedMeshes)
 			{
+				Addin.Instance.LogDebug($"  AddMesh(DatasmithMeshName: {Info}");
 				// Info.MeshElement should be valid here
 				AddMesh(Info.MeshName, Info.MeshElement);
 
@@ -1086,7 +1091,7 @@ namespace DatasmithSolidworks
 				}
 			}
 
-			OutCreatedMeshes = CreatedMeshes;
+			OutCreatedMeshes = CreatedMeshes.ToList();
 		}
 	}
 }
\ No newline at end of file
diff --git a/Engine/Source/Programs/Enterprise/Datasmith/DatasmithSolidworksExporter/Private/ObjectMaterials.cs b/Engine/Source/Programs/Enterprise/Datasmith/DatasmithSolidworksExporter/Private/ObjectMaterials.cs
--- a/Engine/Source/Programs/Enterprise/Datasmith/DatasmithSolidworksExporter/Private/ObjectMaterials.cs
+++ b/Engine/Source/Programs/Enterprise/Datasmith/DatasmithSolidworksExporter/Private/ObjectMaterials.cs
@@ -78,7 +78,12 @@ namespace DatasmithSolidworks
 
 		public override string ToString()
 		{
-			return $"FObjectMaterials:ComponentMaterialID='{ComponentMaterialID}', PartMaterialID='{PartMaterialID}'";
+			return $"FObjectMaterials:ComponentMaterialID='{ComponentMaterialID}', PartMaterialID='{PartMaterialID}', BodyMaterialsMap='{ToString(BodyMaterialsMap)}', FeatureMaterialsMap='{ToString(FeatureMaterialsMap)}', FaceMaterialsMap='{ToString(FaceMaterialsMap)}'";
+		}
+
+		private static string ToString(Dictionary<string, int> Map)
+		{
+			return string.Join(", ", Map.Select(KVP => $"{KVP.Key}:{KVP.Value}" ));
 		}
 
 		public FMaterial GetMaterial(Face2 InFace)
@@ -236,7 +241,7 @@ namespace DatasmithSolidworks
 
 		public static FObjectMaterials LoadComponentMaterials(FDocumentTracker InComponentOwner, Component2 InComponent, swDisplayStateOpts_e InDisplayState, string[] InDisplayStateNames)
 		{
-			Addin.Instance.LogDebug($"LoadComponentMaterials: {InComponentOwner.GetPathName()}, {InDisplayState}, {InDisplayStateNames}");
+			Addin.Instance.LogDebug($"LoadComponentMaterials: '{InComponent.Name2}', '{InComponentOwner.GetPathName()}', {InDisplayState}, {(InDisplayStateNames == null ? "<null>" : string.Join(", ", InDisplayStateNames))}");
 
 			ModelDoc2 ComponentDoc = InComponent.GetModelDoc2() as ModelDoc2;
 			if (ComponentDoc == null || !(ComponentDoc is PartDoc))
@@ -245,8 +250,9 @@ namespace DatasmithSolidworks
 				return null;
 			}
 
+			Addin.Instance.LogDebug($"  GetRenderMaterialsCount2({InDisplayState}, {InDisplayStateNames})");
 			int NumMaterials = InComponent.GetRenderMaterialsCount2((int)InDisplayState, InDisplayStateNames);
-			Addin.Instance.LogDebug($"  GetRenderMaterialsCount2: {NumMaterials}");
+			Addin.Instance.LogDebug($"    GetRenderMaterialsCount2 -> {NumMaterials}");
 
 			if (NumMaterials == 0)
 			{
@@ -262,6 +268,9 @@ namespace DatasmithSolidworks
 				RenderMaterial RenderMat = ObjMat as RenderMaterial;
 				int NumUsers = RenderMat.GetEntitiesCount();
 
+				Addin.Instance.LogDebug($"  FileName: {RenderMat.FileName}");
+				Addin.Instance.LogDebug($"    Users({NumUsers}):");
+
 				if (NumUsers == 0)
 				{
 					continue;
@@ -270,29 +279,41 @@ namespace DatasmithSolidworks
 				object[] ObjUsers = RenderMat.GetEntities();
 				foreach (object ObjUser in ObjUsers)
 				{
-					if (ObjUser is IBody2 Body)
-					{
-						string BodyId = FPartDocumentTracker.GetBodyPath(ComponentDoc as PartDoc, Body);
-						ComponentMaterials.RegisterMaterial(ComponentMaterials.BodyMaterialsMap, InComponentOwner.SwDoc, RenderMat, BodyId);
-						continue;
-					}
-
-					if (ObjUser is IFace2 Face)
-					{
-						uint FaceId = ComponentMaterials.OwnerDoc.GetFaceId(Face);
-						ComponentMaterials.RegisterMaterial(ComponentMaterials.FaceMaterialsMap, InComponentOwner.SwDoc, RenderMat, FaceId.ToString());
-						continue;
-					}
-
-					if (ObjUser is IFeature Feat)
+					switch (ObjUser)
 					{
-						string FeatureId = FPartDocumentTracker.GetFeaturePath(ComponentDoc as PartDoc, Feat);
-						ComponentMaterials.RegisterMaterial(ComponentMaterials.FeatureMaterialsMap, InComponentOwner.SwDoc, RenderMat, FeatureId);
-						continue;
-					}
-					if (ObjUser is IPartDoc Doc)
-					{
-						ComponentMaterials.SetPartMaterial(RenderMat, InComponentOwner.SwDoc);
+						case IBody2 Body:
+						{
+							Addin.Instance.LogDebug($"    Body '{Body.Name}'");
+							string BodyId = FPartDocumentTracker.GetBodyPath(ComponentDoc as PartDoc, Body);
+							ComponentMaterials.RegisterMaterial(ComponentMaterials.BodyMaterialsMap, InComponentOwner.SwDoc, RenderMat, BodyId);
+							break;
+						}
+						case IFace2 Face:
+						{
+							// note: plugins sets FaceId on each face to identify per-face materials later in CreateMeshData
+							uint FaceId = ComponentMaterials.OwnerDoc.GetFaceId(Face);
+							Addin.Instance.LogDebug($"    Face '{Face.GetFaceId()}'");
+							ComponentMaterials.RegisterMaterial(ComponentMaterials.FaceMaterialsMap, InComponentOwner.SwDoc, RenderMat, FaceId.ToString());
+							break;
+						}
+						case IFeature Feat:
+						{
+							Addin.Instance.LogDebug($"    Feature '{Feat.Name}'");
+							string FeatureId = FPartDocumentTracker.GetFeaturePath(ComponentDoc as PartDoc, Feat);
+							ComponentMaterials.RegisterMaterial(ComponentMaterials.FeatureMaterialsMap, InComponentOwner.SwDoc, RenderMat, FeatureId);
+							break;
+						}
+						case IPartDoc Doc:
+						{
+							Addin.Instance.LogDebug($"    Part Doc '{(Doc as ModelDoc2).GetPathName()}'");
+							ComponentMaterials.SetPartMaterial(RenderMat, InComponentOwner.SwDoc);
+							break;
+						}
+						default:
+						{
+							Addin.Instance.LogDebug($"    <unknown>'");
+							break;
+						}
 					}
 				}
 			}
@@ -302,13 +323,13 @@ namespace DatasmithSolidworks
 
 		public static ConcurrentDictionary<FComponentName, FObjectMaterials> LoadAssemblyMaterials(FAssemblyDocumentTracker InAsmDoc, HashSet<FComponentName> InComponentsSet, swDisplayStateOpts_e InDisplayState, string[] InDisplayStateNames)
 		{
-			Addin.Instance.LogDebug($"LoadAssemblyMaterials: {InAsmDoc.GetPathName()}, {InDisplayState}, {InDisplayStateNames}");
-			Addin.Instance.LogDebug($"  for components: {string.Join(", ", InComponentsSet)}");
+			Addin.Instance.LogDebug($"LoadAssemblyMaterials: {InAsmDoc.GetPathName()}, {InDisplayState}, [{(InDisplayStateNames == null ? "" : string.Join(",", InDisplayStateNames))}]");
+			Addin.Instance.LogDebug($"  for components: [{string.Join(", ", InComponentsSet.OrderBy(Name => Name.GetString()))}]");
 
 			IModelDocExtension Ext = InAsmDoc.SwDoc.Extension;
 			int NumMaterials = Ext.GetRenderMaterialsCount2((int)InDisplayState, InDisplayStateNames);
 
-			Addin.Instance.LogDebug($"  NumMaterials={NumMaterials}");
+			Addin.Instance.LogDebug($"  Assembly materials={NumMaterials}");
 
 			ConcurrentDictionary<FComponentName, FObjectMaterials> DocMaterials = new ConcurrentDictionary<FComponentName, FObjectMaterials>();
 
@@ -319,8 +340,13 @@ namespace DatasmithSolidworks
 				foreach (object ObjMat in ObjMaterials)
 				{
 					RenderMaterial RenderMat = ObjMat as RenderMaterial;
+					
+					Addin.Instance.LogDebug($"    Material: '{RenderMat.FileName}' ");
+
 					int NumUsers = RenderMat.GetEntitiesCount();
 
+					Addin.Instance.LogDebug($"    NumUsers: '{NumUsers}' ");
+
 					if (NumUsers == 0)
 					{
 						continue;
@@ -329,57 +355,94 @@ namespace DatasmithSolidworks
 					object[] ObjUsers = RenderMat.GetEntities();
 					foreach (object ObjUser in ObjUsers)
 					{
-						if (ObjUser is Component2 Comp)
+						if (ObjUser is IEntity Entity)
 						{
-							PartDoc CompDoc = Comp.GetModelDoc2() as PartDoc;
-							if (CompDoc != null)
-							{
-								FObjectMaterials ComponentMaterials = new FObjectMaterials(InAsmDoc, CompDoc, ref InAsmDoc.ExportedMaterialsMap);
-								ComponentMaterials.SetComponentMaterial(RenderMat, InAsmDoc.SwDoc);
-								DocMaterials[new FComponentName(Comp)] = ComponentMaterials;
-							}
-							continue;
+							Addin.Instance.LogDebug($"    User: type '{Entity.GetType()}' ");
+							
 						}
-
-						if (ObjUser is IPartDoc Doc)
+						
+						switch (ObjUser)
 						{
-							FObjectMaterials PartMaterials = LoadPartMaterials(InAsmDoc, Doc as PartDoc, InDisplayState, InDisplayStateNames);
-							if (PartMaterials != null)
+							case Component2 Comp:
 							{
-								// todo: why trying to present PathName as component name?
-								//  how it can be a PartDoc in Assembly function
-								DocMaterials[FComponentName.FromCustomString((Doc as ModelDoc2).GetPathName())] = PartMaterials;
+								Addin.Instance.LogDebug($"      Component '{Comp.Name2}' ");
+								if (Comp.GetModelDoc2() is PartDoc CompDoc)
+								{
+									Addin.Instance.LogDebug($"        registering material for Part Component in DocMaterials");
+									FObjectMaterials ComponentMaterials = new FObjectMaterials(InAsmDoc, CompDoc, ref InAsmDoc.ExportedMaterialsMap);
+
+									Addin.Instance.LogIndent();
+									ComponentMaterials.SetComponentMaterial(RenderMat, InAsmDoc.SwDoc);
+									Addin.Instance.LogDedent();
+
+									DocMaterials[new FComponentName(Comp)] = ComponentMaterials;
+								}
+								break;
+							}
+							case IPartDoc Doc:
+							{
+								Addin.Instance.LogDebug($"      Part Doc at '{(Doc as ModelDoc2).GetPathName()}' ");
+
+								Addin.Instance.LogIndent();
+								FObjectMaterials PartMaterials = LoadPartMaterials(InAsmDoc, Doc as PartDoc, InDisplayState, InDisplayStateNames);
+								Addin.Instance.LogDedent();
+
+								if (PartMaterials != null)
+								{
+									// todo: why trying to present PathName as component name?
+									//  how it can be a PartDoc in Assembly function
+									DocMaterials[FComponentName.FromCustomString((Doc as ModelDoc2).GetPathName())] = PartMaterials;
+								}
+								break;
 							}
-							continue;
 						}
 					}
 				}
 			}
 
-			// Check for materials that are not per component (but instead per face, feature, part etc.)
-			// todo: consider replacing Parallel.ForEach with simple loop
-			Parallel.ForEach(InComponentsSet,
-				CompName =>
+			Addin.Instance.LogDebug($"Load Each Component's materials");
+			Addin.Instance.LogIndent();
+			foreach(FComponentName CompName in InComponentsSet.OrderBy(Name => Name.GetString()))
 			{
 				Component2 Comp = null;
 
+				Addin.Instance.LogDebug($"Component '{CompName}' - adding DocMaterials");
+				Addin.Instance.LogIndent();
+
 				if (InAsmDoc.SyncState.ExportedComponentsMap.TryGetValue(CompName, out Comp))
 				{
 					if (!DocMaterials.ContainsKey(CompName))
 					{
+						Addin.Instance.LogDebug($"Adding to DocMaterials");
+
+						Addin.Instance.LogIndent();
 						FObjectMaterials ComponentMaterials = FObjectMaterials.LoadComponentMaterials(InAsmDoc, Comp, swDisplayStateOpts_e.swThisDisplayState, null);
+						Addin.Instance.LogDedent();
 
 						if (ComponentMaterials == null && Comp.GetModelDoc2() is PartDoc)
 						{
+							Addin.Instance.LogIndent();
 							ComponentMaterials = FObjectMaterials.LoadPartMaterials(InAsmDoc, Comp.GetModelDoc2() as PartDoc, swDisplayStateOpts_e.swThisDisplayState, null);
+							Addin.Instance.LogDedent();
 						}
 						if (ComponentMaterials != null)
 						{
 							DocMaterials.TryAdd(CompName, ComponentMaterials);
 						}
 					}
+					else
+					{
+						Addin.Instance.LogDebug($"...already in DocMaterials");
+					}
+				}
+				else
+				{
+					Addin.Instance.LogDebug($"...not in ExportedComponentsMap");
 				}
-			});
+
+				Addin.Instance.LogDedent();
+			}
+			Addin.Instance.LogDedent();
 
 			return DocMaterials.Count > 0 ? DocMaterials : null;
 		}
