[Solidworks] Stability fixes

- Fixed Sync/Export when plugin enabled manually(not on Startup)
- Fixed(temporary) plugin crashing after being re-enabled
- DirectLink source not staying alive anymore when closing last document or disabling the plugin

diff --git a/Engine/Source/Programs/Enterprise/Datasmith/DatasmithSolidworksExporter/Private/Addin.cs b/Engine/Source/Programs/Enterprise/Datasmith/DatasmithSolidworksExporter/Private/Addin.cs
--- a/Engine/Source/Programs/Enterprise/Datasmith/DatasmithSolidworksExporter/Private/Addin.cs
+++ b/Engine/Source/Programs/Enterprise/Datasmith/DatasmithSolidworksExporter/Private/Addin.cs
@@ -171,6 +171,8 @@ namespace DatasmithSolidworks
 			}
 		}
 
+		private bool bDatasmithFacadeDirectLinkInitialized = false;
+
 		public bool ConnectToSW(object InThisSW, int InCookie)
 		{
 			SolidworksApp = (ISldWorks)InThisSW;
@@ -187,16 +189,31 @@ namespace DatasmithSolidworks
 
 			FDatasmithFacadeElement.SetCoordinateSystemType(FDatasmithFacadeElement.ECoordinateSystemType.RightHandedZup);
 
-			FDatasmithFacadeDirectLink.Init();
+			if (!bDatasmithFacadeDirectLinkInitialized)
+			{
+				FDatasmithFacadeDirectLink.Init();
+				bDatasmithFacadeDirectLinkInitialized = true;
+			}
 
 			FMaterial.InitializeMaterialTypes();
 
+			OnActiveDocChange();
+
 			return true;
 		}
 
 		public bool DisconnectFromSW()
 		{
-			FDatasmithFacadeDirectLink.Shutdown();
+			// Disabled Shutdown as Initing again crashes if the plugin is re-enabled in SW
+			// FDatasmithFacadeDirectLink.Shutdown it is required to have plugin dll unloaded but Solidworks doesn't unload the plugin assembly when plugin is disabled
+			// So, something like this doesn't work:
+			// if (bDatasmithFacadeDirectLinkInitialized)
+			// {
+			// 	FDatasmithFacadeDirectLink.Shutdown();
+			// 	bDatasmithFacadeDirectLinkInitialized = false;
+			// }
+			// At least make sure to close current connection
+			CurrentDocument?.MakeActive(false);
 
 			DetachEventHandlers();
 			DestroyToolbarCommands();
diff --git a/Engine/Source/Programs/Enterprise/Datasmith/DatasmithSolidworksExporter/Private/Document.cs b/Engine/Source/Programs/Enterprise/Datasmith/DatasmithSolidworksExporter/Private/Document.cs
--- a/Engine/Source/Programs/Enterprise/Datasmith/DatasmithSolidworksExporter/Private/Document.cs
+++ b/Engine/Source/Programs/Enterprise/Datasmith/DatasmithSolidworksExporter/Private/Document.cs
@@ -452,6 +452,7 @@ namespace DatasmithSolidworks
 				// other event handling is not suspended to register changes in the model
 				DocumentSyncer?.Pause();
 
+				DatasmithDirectLink?.CloseCurrentSource();
 				DatasmithDirectLink?.Dispose();  // aka 'Close DirectLink Connection'
 				DatasmithDirectLink = null;
 			}
@@ -483,5 +484,6 @@ namespace DatasmithSolidworks
 				Addin.Instance.LogStatusBarMessage($"Exporting file {FileName}...{InMessage}");
 			}
 		}
+
 	}
 }
\ No newline at end of file
