[Datasmith Revit Exporter] Fixed Duplicated ids for Level elements

- also added more logging for more efficient debugging of parsing/name generation

diff --git a/Engine/Source/Programs/Enterprise/Datasmith/DatasmithRevitExporter/Private/DatasmithRevitApplication.cs b/Engine/Source/Programs/Enterprise/Datasmith/DatasmithRevitExporter/Private/DatasmithRevitApplication.cs
--- a/Engine/Source/Programs/Enterprise/Datasmith/DatasmithRevitExporter/Private/DatasmithRevitApplication.cs
+++ b/Engine/Source/Programs/Enterprise/Datasmith/DatasmithRevitExporter/Private/DatasmithRevitApplication.cs
@@ -1,11 +1,13 @@
 // Copyright Epic Games, Inc. All Rights Reserved.
 
 using System;
+using System.Collections.Concurrent;
 using System.Collections.Generic;
 using System.Diagnostics;
 using System.IO;
 using System.Linq;
 using System.Reflection;
+using System.Threading;
 using System.Windows.Media.Imaging;
 using Autodesk.Revit.DB;
 using Autodesk.Revit.DB.Events;
@@ -81,6 +83,12 @@ namespace DatasmithRevitExporter
 		BitmapImage AutoSyncIconOn_Large;
 		BitmapImage AutoSyncIconOff_Small;
 		BitmapImage AutoSyncIconOff_Large;
+		private FDebugLog DebugLog;
+
+		public DatasmithRevitApplication()
+		{
+			CreateDebugLog();
+		}
 
 		public object Properties { get; private set; }
 
@@ -439,5 +447,58 @@ namespace DatasmithRevitExporter
 			return true;
 #endif
 		}
+
+		[Conditional("DatasmithRevitDebugOutput")]
+		private void CreateDebugLog()
+		{
+			DebugLog = new FDebugLog();
+		}
+
+		[Conditional("DatasmithRevitDebugOutput")]
+		public void LogDebug(string Message)  
+		{
+			DebugLog.LogDebug(Message);
+		}
 	}
+
+	class FDebugLog
+	{
+		private ConcurrentQueue<string> MessagesQueue = new ConcurrentQueue<string>();
+		private Thread LogWriterThread;
+
+		public FDebugLog()
+		{
+			LogWriterThread = new Thread(() =>
+			{
+				LogWriterProc();
+			});
+
+			LogWriterThread.Start();
+		}
+
+		private void LogWriterProc()
+		{
+			string LogPath = Path.Combine(Environment.GetFolderPath(Environment.SpecialFolder.LocalApplicationData),
+				"UnrealDatasmithExporter/Saved/Logs/UnrealDatasmithRevitExporterDebug.log");
+
+			StreamWriter LogFile = new StreamWriter(LogPath);
+
+			while (true)
+			{
+				while (MessagesQueue.TryDequeue(out string Message))
+				{
+					LogFile.WriteLine(Message);
+				}
+				LogFile.Flush();
+
+				Thread.Sleep(10);
+			}
+		}
+
+		public void LogDebug(string Message)  
+		{
+			MessagesQueue.Enqueue(Message);
+		}
+	};
+
 }
diff --git a/Engine/Source/Programs/Enterprise/Datasmith/DatasmithRevitExporter/Private/DatasmithRevitCommand.cs b/Engine/Source/Programs/Enterprise/Datasmith/DatasmithRevitExporter/Private/DatasmithRevitCommand.cs
--- a/Engine/Source/Programs/Enterprise/Datasmith/DatasmithRevitExporter/Private/DatasmithRevitCommand.cs
+++ b/Engine/Source/Programs/Enterprise/Datasmith/DatasmithRevitExporter/Private/DatasmithRevitCommand.cs
@@ -113,6 +113,7 @@ namespace DatasmithRevitExporter
 				catch (System.Exception exception)
 				{
 					OutCommandMessage = string.Format("Cannot export the 3D view:\n\n{0}\n\n{1}", exception.Message, exception.StackTrace);
+					ExportContext.LogDebug(OutCommandMessage);
 					MessageBox.Show(OutCommandMessage, DatasmithRevitCommandUtils.DIALOG_CAPTION, MessageBoxButtons.OK, MessageBoxIcon.Error);
 					return Result.Failed;
 				}
diff --git a/Engine/Source/Programs/Enterprise/Datasmith/DatasmithRevitExporter/Private/DatasmithRevitDocumentData.cs b/Engine/Source/Programs/Enterprise/Datasmith/DatasmithRevitExporter/Private/DatasmithRevitDocumentData.cs
--- a/Engine/Source/Programs/Enterprise/Datasmith/DatasmithRevitExporter/Private/DatasmithRevitDocumentData.cs
+++ b/Engine/Source/Programs/Enterprise/Datasmith/DatasmithRevitExporter/Private/DatasmithRevitDocumentData.cs
@@ -1153,6 +1153,7 @@ namespace DatasmithRevitExporter
 					ActorLabel += string.IsNullOrEmpty(InstanceName) ? "" : "*" + InstanceName;
 				}
 
+				DocumentData.Context.LogDebug($"GetActorLabel(CategoryName='{CategoryName}', FamilyName='{FamilyName}', TypeName='{TypeName}', InstanceName='{InstanceName}') -> '{ActorLabel}'");
 				return ActorLabel;
 			}
 
@@ -2496,6 +2497,7 @@ namespace DatasmithRevitExporter
 			Func<ElementId, Element> InGetParentElement
 		)
 		{
+			Context.LogDebug("AddParentElementHierarchy");
 			Queue<ElementId> ElementIdQueue = new Queue<ElementId>(ActorMap.Keys);
 
 			// Make sure the Datasmith actor dictionary contains actors for all the Revit parent elements.
@@ -2523,6 +2525,7 @@ namespace DatasmithRevitExporter
 				}
 				else
 				{
+					Context.LogDebug("  Add element");
 					const FDatasmithFacadeScene NullScene = null;
 					PushElement(ParentElement, Transform.Identity);
 					PopElement(NullScene);
diff --git a/Engine/Source/Programs/Enterprise/Datasmith/DatasmithRevitExporter/Private/DatasmithRevitExportContext.cs b/Engine/Source/Programs/Enterprise/Datasmith/DatasmithRevitExporter/Private/DatasmithRevitExportContext.cs
--- a/Engine/Source/Programs/Enterprise/Datasmith/DatasmithRevitExporter/Private/DatasmithRevitExportContext.cs
+++ b/Engine/Source/Programs/Enterprise/Datasmith/DatasmithRevitExporter/Private/DatasmithRevitExportContext.cs
@@ -231,7 +231,9 @@ namespace DatasmithRevitExporter
 		[Conditional("DatasmithRevitDebugOutput")]
 		public void LogDebug(string Message)  
 		{
-			MessageList.Add(String.Concat(Enumerable.Repeat("    ", DebugIndent)) + Message);
+			string OutputMessage = String.Concat(Enumerable.Repeat("    ", DebugIndent)) + Message;
+			MessageList.Add(OutputMessage);
+			DatasmithRevitApplication.Instance.LogDebug(OutputMessage);
 		}
 
 		[Conditional("DatasmithRevitDebugOutput")]
@@ -307,6 +309,22 @@ namespace DatasmithRevitExporter
 			LogDebug($"OnInstanceEnd({InInstanceNode.NodeName})");
 		}
 
+		[Conditional("DatasmithRevitDebugOutput")]
+		private void LogDebugLinkBegin(LinkNode InLinkNode)
+		{
+			LogDebug($"LogDebugLinkBegin({InLinkNode.NodeName})");
+
+			++DebugIndent;
+		}
+
+		[Conditional("DatasmithRevitDebugOutput")]
+		private void LogDebugLinkEnd(LinkNode InLinkNode)
+		{
+			--DebugIndent;
+
+			LogDebug($"OnLinkEnd({InLinkNode.NodeName})");
+		}
+
 		// OnElementBegin marks the beginning of an element to be exported.
 		public RenderNodeAction OnElementBegin(
 			ElementId InElementId // exported element ID
@@ -422,6 +440,7 @@ namespace DatasmithRevitExporter
 			LinkNode InLinkNode // linked Revit document output node
 		)
 		{
+			LogDebugLinkBegin(InLinkNode);
 #if REVIT_API_2023
 			ElementType CurrentInstanceType = GetElement(InLinkNode.SymbolId) as ElementType;
 #else
@@ -454,6 +473,7 @@ namespace DatasmithRevitExporter
 			LinkNode InLinkNode // linked Revit document output node
 		)
 		{
+			LogDebugLinkEnd(InLinkNode);
 			if (InLinkNode.GetDocument() != null)
 			{
 				// Forget the current linked document being exported.
@@ -744,18 +764,25 @@ namespace DatasmithRevitExporter
 
 		private FDocumentData PopDocument()
 		{
-			FDocumentData DocumentData = DocumentDataStack.Pop();
+			FDocumentData DocumentData = DocumentDataStack.Peek();
 
-			if (DocumentDataStack.Count == 0)
+			if (DocumentDataStack.Count == 1)
 			{
 				DocumentData.WrapupScene(DatasmithScene, UniqueTextureNameSet);
 			}
 			else
 			{
-				DocumentData.WrapupLink(DatasmithScene, DocumentDataStack.Peek().GetCurrentActor(), UniqueTextureNameSet);
+				// Get element data which is the parent of linked elements - need document under the current Linked document on the stack
+				// Stack.ToArray returns array with Stack top at [0] array element, so second to top is [1]
+				// todo: may replace Stack with List
+				FDocumentData[] DocumentDataArray = DocumentDataStack.ToArray();
+				FDocumentData.FBaseElementData ParentElementData = DocumentDataArray[1].GetCurrentActor();
+			
+				DocumentData.WrapupLink(DatasmithScene, ParentElementData, UniqueTextureNameSet);
 				DirectLink?.OnEndLinkedDocument();
 			}
-
+			// Pop document only when all elements are collected for it (WrapupScene may also make actors for host elements like Level which needs to have proper document stack to build unique element name-path)
+			DocumentDataStack.Pop(); 
 			return DocumentData;
 		}
 
