[Solidworks] Fixed Part materials in multiple configurations

- searching for actual RenderMaterial for each Body/Feature of the part using entity's Appearance(somewhat alternale api to RenderMaterial). In order to deal with api flow which doesn't return corrent entities assigned to a RenderMaterial in active configuration

diff --git a/Engine/Source/Programs/Enterprise/Datasmith/DatasmithSolidworksExporter/Private/Material.cs b/Engine/Source/Programs/Enterprise/Datasmith/DatasmithSolidworksExporter/Private/Material.cs
--- a/Engine/Source/Programs/Enterprise/Datasmith/DatasmithSolidworksExporter/Private/Material.cs
+++ b/Engine/Source/Programs/Enterprise/Datasmith/DatasmithSolidworksExporter/Private/Material.cs
@@ -1055,5 +1055,41 @@ namespace DatasmithSolidworks
 
 			return true;
 		}
+
+		// Used to backtrack Appearance to Material in case Solidworks api of RenderMaterial didn't return all entities this material is used on
+		public bool EqualsAppearance(AppearanceSetting OtherAppearance)
+		{
+			if (Source.PrimaryColor != OtherAppearance.Color)
+			{
+				return false;
+			}
+
+			if (Math.Abs(Source.Emission - OtherAppearance.Luminous) > 0.001)
+			{
+				return false;
+			}
+			
+			if (Math.Abs(Source.Transparency - OtherAppearance.Transparent) > 0.001)
+			{
+				return false;
+			}
+
+			if (Math.Abs(Source.Reflectivity - OtherAppearance.Reflection) > 0.001)
+			{
+				return false;
+			}
+
+			if (Math.Abs(Source.Specular - OtherAppearance.Specular) > 0.001)
+			{
+				return false;
+			}
+
+			if (Source.SpecularColor != OtherAppearance.SpecularColor)
+			{
+				return false;
+			}
+
+			return true;
+		}
 	}
 }
diff --git a/Engine/Source/Programs/Enterprise/Datasmith/DatasmithSolidworksExporter/Private/ObjectMaterials.cs b/Engine/Source/Programs/Enterprise/Datasmith/DatasmithSolidworksExporter/Private/ObjectMaterials.cs
--- a/Engine/Source/Programs/Enterprise/Datasmith/DatasmithSolidworksExporter/Private/ObjectMaterials.cs
+++ b/Engine/Source/Programs/Enterprise/Datasmith/DatasmithSolidworksExporter/Private/ObjectMaterials.cs
@@ -1,12 +1,10 @@
 // Copyright Epic Games, Inc. All Rights Reserved.
 
-using System.Runtime.InteropServices;
 using System.Collections.Generic;
 using SolidWorks.Interop.sldworks;
 using SolidWorks.Interop.swconst;
 using System.Linq;
 using System.Collections.Concurrent;
-using System.Threading.Tasks;
 using DatasmithSolidworks.Names;
 
 namespace DatasmithSolidworks
@@ -177,7 +175,10 @@ namespace DatasmithSolidworks
 
 		public static FObjectMaterials LoadPartMaterials(FDocumentTracker InOwnerDoc, PartDoc InPartDoc, swDisplayStateOpts_e InDisplayState, string[] InDisplayStateNames)
 		{
+			Addin.Instance.LogDebug($"LoadPartMaterials: '{(InPartDoc as ModelDoc2).GetPathName()}', {InDisplayState}, {(InDisplayStateNames == null ? "<null>" : string.Join(", ", InDisplayStateNames))}");
+
 			ModelDoc2 Doc = InPartDoc as ModelDoc2;
+
 			if (Doc == null)
 			{
 				return null;
@@ -185,6 +186,9 @@ namespace DatasmithSolidworks
 
 			IModelDocExtension Ext = Doc.Extension;
 			int NumMaterials = Ext.GetRenderMaterialsCount2((int)InDisplayState, InDisplayStateNames);
+
+			Addin.Instance.LogDebug($"    GetRenderMaterialsCount2 -> {NumMaterials}");
+
 			if (NumMaterials == 0)
 			{
 				return null;
@@ -198,6 +202,8 @@ namespace DatasmithSolidworks
 			{
 				RenderMaterial RenderMat = ObjMat as RenderMaterial;
 				int NumUsers = RenderMat.GetEntitiesCount();
+				Addin.Instance.LogDebug($"  FileName: {RenderMat.FileName}");
+				Addin.Instance.LogDebug($"    Users({NumUsers}):");
 
 				if (NumUsers == 0)
 				{
@@ -207,34 +213,151 @@ namespace DatasmithSolidworks
 				object[] ObjUsers = RenderMat.GetEntities();
 				foreach (object ObjUser in ObjUsers)
 				{
-					if (ObjUser is IPartDoc)
+					switch (ObjUser)
 					{
-						PartMaterials.SetPartMaterial(RenderMat, Doc);
-						continue;
+						case IPartDoc _:
+							Addin.Instance.LogDebug($"    PartDoc '{(ObjUser as ModelDoc2).GetPathName()}'");
+							PartMaterials.SetPartMaterial(RenderMat, Doc);
+							continue;
+						case IBody2 Body:
+						{
+							Addin.Instance.LogDebug($"    Body '{Body.Name}'");
+
+							string BodyId = FPartDocumentTracker.GetBodyPath(InPartDoc as PartDoc, Body);
+							PartMaterials.RegisterMaterial(PartMaterials.BodyMaterialsMap, Doc, RenderMat, BodyId);
+							continue;
+						}
+						case IFace2 Face:
+						{
+							uint FaceId = PartMaterials.OwnerDoc.GetFaceId(Face);
+							Addin.Instance.LogDebug($"    Face '{Face.GetFaceId()}'");
+							PartMaterials.RegisterMaterial(PartMaterials.FaceMaterialsMap, Doc, RenderMat, FaceId.ToString());
+							continue;
+						}
+						case IFeature Feat:
+						{
+							Addin.Instance.LogDebug($"    Feature '{Feat.Name}'");
+							string FeatureId = FPartDocumentTracker.GetFeaturePath(InPartDoc as PartDoc, Feat);
+							PartMaterials.RegisterMaterial(PartMaterials.FeatureMaterialsMap, Doc, RenderMat, FeatureId);
+							continue;
+						}
+						default:
+						{
+							Addin.Instance.LogDebug($"    <unknown>'");
+							break;
+						}
 					}
+				}
+			}
+
+			// Parse Part entities material properties to workaround issue with SW api
+			//
+			// Solidworks API gives away render materials IModelDocExtension.GetRenderMaterialsCount2
+			// and entities these materials are assigned to are retrieved using RenderMaterial.GetEntities[Count] method
+			// But this has a flaw that when a Part has multiple configurationa those entities returned are only for the configuration
+			// that was active when this Part file was saved. Not matter if we switch configuration - still same list of entities returned.
+			// E.g. in Config C0, feature F0 was visible and assigned material M0, in config C1, feature F0 is not visible
+			// but rather feature F1  is revealed and has this material M0. If model was saved with C0 active
+			// M0.GetEntities returns only F0, even if we switch to C1 after part is loaded
+			// But get_DisplayStateSpecMaterialPropertyValues seems to return proper material properties for proper(i.e. visible in currently active config/display state) features
+			{
+				DisplayStateSetting swDSS = Ext.GetDisplayStateSetting((int)swDisplayStateOpts_e.swThisDisplayState);
+				swDSS.Option = (int)swDisplayStateOpts_e.swThisDisplayState;
+
+				// Check Bodies
+				object[] BodiesArray = InPartDoc.GetBodies((int)swBodyType_e.swSolidBody);
+				if (BodiesArray != null)
+				{
+					swDSS.Entities = BodiesArray.Cast<Body2>().ToArray();
+					// Get appearances for entities array
+					object[] Appearances = (object[])Ext.get_DisplayStateSpecMaterialPropertyValues(swDSS);
 
-					if (ObjUser is IBody2 Body)
+					for (int SettingIndex = 0; SettingIndex < Appearances.Length; SettingIndex++)
 					{
+						Body2 Body = swDSS.Entities[SettingIndex];
+						AppearanceSetting Appearance = Appearances[SettingIndex] as AppearanceSetting;
+						
 						string BodyId = FPartDocumentTracker.GetBodyPath(InPartDoc as PartDoc, Body);
-						PartMaterials.RegisterMaterial(PartMaterials.BodyMaterialsMap, Doc, RenderMat, BodyId);
-						continue;
+						if (PartMaterials.BodyMaterialsMap.TryGetValue(BodyId, out int AssignedMaterialId))
+						{
+							Addin.Instance.LogDebug($"  Body '{Body.Name}' already has material assigned {PartMaterials.GetMaterial(AssignedMaterialId).Name}");
+							continue;
+						}
+
+						Addin.Instance.LogDebug($"  Body '{Body.Name}' has no material assigned, searching material matching its Appearance..");
+
+						foreach (KeyValuePair<int, FMaterial> KVP in PartMaterials.GlobalMaterialsMap)
+						{
+							FMaterial Material = KVP.Value;
+
+							if (Material.EqualsAppearance(Appearance))
+							{
+								Addin.Instance.LogDebug($"Matching material for Body '{Body.Name}'  is {Material.Name}");
+
+								PartMaterials.BodyMaterialsMap[BodyId] = KVP.Key; 
+								break;
+							}
+						}
 					}
+				}
+
+				// Check Features
+				Stack<Feature> Features = new Stack<Feature>();
 
-					if (ObjUser is IFace2 Face)
+				// Collect features
+				{
+					Feature Feat = InPartDoc.FirstFeature();
+					while (Feat != null)
 					{
-						uint FaceId = PartMaterials.OwnerDoc.GetFaceId(Face);
-						PartMaterials.RegisterMaterial(PartMaterials.FaceMaterialsMap, Doc, RenderMat, FaceId.ToString());
-						continue;
+						Features.Push(Feat);
+
+						Feature SubFeature = Feat.GetFirstSubFeature();
+						while (SubFeature != null)
+						{
+							Features.Push(SubFeature);
+							SubFeature = SubFeature.GetNextSubFeature();
+						}
+
+						Feat = Feat.GetNextFeature();
 					}
+				}
 
-					if (ObjUser is IFeature Feat)
+				if (Features.Count > 0)
+				{
+					swDSS.Entities = Features.ToArray();
+
+					// Get appearances for entities array
+					object[] Appearances = (object[])Ext.get_DisplayStateSpecMaterialPropertyValues(swDSS);
+
+					for (int SettingIndex = 0; SettingIndex < Appearances.Length; SettingIndex++)
 					{
+						Feature Feat = swDSS.Entities[SettingIndex];
+						AppearanceSetting Appearance = Appearances[SettingIndex] as AppearanceSetting;
+						
 						string FeatureId = FPartDocumentTracker.GetFeaturePath(InPartDoc as PartDoc, Feat);
-						PartMaterials.RegisterMaterial(PartMaterials.FeatureMaterialsMap, Doc, RenderMat, FeatureId);
-						continue;
+						if (PartMaterials.FeatureMaterialsMap.TryGetValue(FeatureId, out int AssignedMaterialId))
+						{
+							Addin.Instance.LogDebug($"  Feature '{Feat.Name}' already has material assigned {PartMaterials.GetMaterial(AssignedMaterialId).Name}");
+							continue;
+						}
+
+						Addin.Instance.LogDebug($"  Feature '{Feat.Name}' has no material assigned, searching material matching its Appearance..");
+
+						foreach (KeyValuePair<int, FMaterial> KVP in PartMaterials.GlobalMaterialsMap)
+						{
+							FMaterial Material = KVP.Value;
+
+							if (Material.EqualsAppearance(Appearance))
+							{
+								Addin.Instance.LogDebug($"Matching material for Feature '{Feat.Name}'  is {Material.Name}");
+
+								PartMaterials.FeatureMaterialsMap[FeatureId] = KVP.Key; 
+								break;
+							}
+						}
 					}
 				}
-			}
+			}			
 
 			return PartMaterials;
 		}
@@ -244,6 +367,7 @@ namespace DatasmithSolidworks
 			Addin.Instance.LogDebug($"LoadComponentMaterials: '{InComponent.Name2}', '{InComponentOwner.GetPathName()}', {InDisplayState}, {(InDisplayStateNames == null ? "<null>" : string.Join(", ", InDisplayStateNames))}");
 
 			ModelDoc2 ComponentDoc = InComponent.GetModelDoc2() as ModelDoc2;
+
 			if (ComponentDoc == null || !(ComponentDoc is PartDoc))
 			{
 				// Component's model doc might be null if component is suppressed or lightweight (in which case we treat it as hidden)
