[Solidworks] Make an actor invisible in the Datasmith scene tree on export if it's suppressed in the selected config

diff --git a/Engine/Source/Programs/Enterprise/Datasmith/DatasmithSolidworksExporter/Private/AssemblyDocument.cs b/Engine/Source/Programs/Enterprise/Datasmith/DatasmithSolidworksExporter/Private/AssemblyDocument.cs
--- a/Engine/Source/Programs/Enterprise/Datasmith/DatasmithSolidworksExporter/Private/AssemblyDocument.cs
+++ b/Engine/Source/Programs/Enterprise/Datasmith/DatasmithSolidworksExporter/Private/AssemblyDocument.cs
@@ -559,7 +559,7 @@ namespace DatasmithSolidworks
 
 					SyncState.ComponentsTransformsMap[InNode.ComponentName] = ActorExportInfo.Transform;
 
-					ActorExportInfo.bVisible = ComponentConfig?.bVisible ?? InNode.CommonConfig.bVisible;
+					ActorExportInfo.bVisible = ComponentConfig.bVisible && (ComponentConfig.ConfigName == ActiveConfigName);
 
 					// This component has associated part document -- treat is as a mesh actor
 					ActorExportInfo.Type = EActorType.MeshActor;
