[Solidworks] Remove the -1 on spot light scale

diff --git a/Engine/Source/Programs/Enterprise/Datasmith/DatasmithSolidworksExporter/Private/DatasmithExporter.cs b/Engine/Source/Programs/Enterprise/Datasmith/DatasmithSolidworksExporter/Private/DatasmithExporter.cs
--- a/Engine/Source/Programs/Enterprise/Datasmith/DatasmithSolidworksExporter/Private/DatasmithExporter.cs
+++ b/Engine/Source/Programs/Enterprise/Datasmith/DatasmithSolidworksExporter/Private/DatasmithExporter.cs
@@ -305,7 +305,7 @@ namespace DatasmithSolidworks
 			{
 				case FLight.EType.Directional:
 				{
-					LightDirection = InLight.DirLightDirection;
+					LightDirection = - InLight.DirLightDirection; // Direction FROM the light source, SW gives away direction TOWARDS light source
 					ExportInfo.Type = EActorType.DirLightActor;
 				}
 				break;
@@ -318,7 +318,7 @@ namespace DatasmithSolidworks
 				case FLight.EType.Spot:
 				{
 					LightPosition = InLight.SpotLightPosition;
-					LightDirection = (InLight.SpotLightPosition - InLight.SpotLightTarget ).Normalized();  // Inverted direction
+					LightDirection = (InLight.SpotLightTarget - InLight.SpotLightPosition).Normalized();  // Direction FROM the light source
 					ExportInfo.Type = EActorType.SpotLightActor;
 				}
 				break;
diff --git a/Engine/Source/Programs/Enterprise/Datasmith/DatasmithSolidworksExporter/Private/Material.cs b/Engine/Source/Programs/Enterprise/Datasmith/DatasmithSolidworksExporter/Private/Material.cs
--- a/Engine/Source/Programs/Enterprise/Datasmith/DatasmithSolidworksExporter/Private/Material.cs
+++ b/Engine/Source/Programs/Enterprise/Datasmith/DatasmithSolidworksExporter/Private/Material.cs
@@ -909,7 +909,7 @@ namespace DatasmithSolidworks
 				XPos * UDir.Y + YPos * VDir.Y - CenterPoint.Y,
 				XPos * UDir.Z + YPos * VDir.Z - CenterPoint.Z);
 
-			UVPlane.Normal = new FVec3() - FVec3.Cross(UVPlane.UDirection, UVPlane.VDirection);
+			UVPlane.Normal = - FVec3.Cross(UVPlane.UDirection, UVPlane.VDirection);
 
 			return UVPlane;
 		}
@@ -929,7 +929,7 @@ namespace DatasmithSolidworks
 				case FMaterial.EMappingType.TYPE_AUTOMATIC:
 				case FMaterial.EMappingType.TYPE_CYLINDRICAL:
 				{
-					FVec3 cross = new FVec3(0f, 0f, 0f) - FVec3.Cross(UDirection, VDirection);
+					FVec3 cross = - FVec3.Cross(UDirection, VDirection);
 					Planes.Add(GetUVPlane(UDirection, VDirection));
 					Planes.Add(GetUVPlane(UDirection, cross));
 					Planes.Add(GetUVPlane(VDirection, cross));
@@ -989,7 +989,7 @@ namespace DatasmithSolidworks
 		public FVec2 ComputeNormalAtan2(FVec3 InVertex, FVec3 InModelCenter)
 		{
 			FVec3 Normal = InVertex - InModelCenter;
-			FVec3 UVCross = FVec3.Cross(new FVec3() - UDirection, VDirection);
+			FVec3 UVCross = FVec3.Cross(- UDirection, VDirection);
 			FVec3 UVNormal = new FVec3(-FVec3.Dot(Normal, UDirection), FVec3.Dot(Normal, VDirection), FVec3.Dot(Normal, UVCross));
 			return new FVec2(((float)Math.Atan2(UVNormal.Z, UVNormal.X) / (Math.PI * 2)), UVNormal.Y);
 		}
diff --git a/Engine/Source/Programs/Enterprise/Datasmith/DatasmithSolidworksExporter/Private/MathUtils.cs b/Engine/Source/Programs/Enterprise/Datasmith/DatasmithSolidworksExporter/Private/MathUtils.cs
--- a/Engine/Source/Programs/Enterprise/Datasmith/DatasmithSolidworksExporter/Private/MathUtils.cs
+++ b/Engine/Source/Programs/Enterprise/Datasmith/DatasmithSolidworksExporter/Private/MathUtils.cs
@@ -95,8 +95,8 @@ namespace DatasmithSolidworks
 			if (InDirection != null)
 			{
 				FVec3 Forward = InDirection.Normalized();
-				FVec3 Right = FVec3.Cross(FVec3.YAxis, Forward);
-				FVec3 Up = FVec3.Cross(Forward, Right);
+				FVec3 Right = - FVec3.Cross(FVec3.YAxis, Forward);
+				FVec3 Up = - FVec3.Cross(Forward, Right);
 
 				Ret[0] = Forward.X;
 				Ret[1] = Forward.Y;
diff --git a/Engine/Source/Programs/Enterprise/Datasmith/DatasmithSolidworksExporter/Private/Matrix4.cs b/Engine/Source/Programs/Enterprise/Datasmith/DatasmithSolidworksExporter/Private/Matrix4.cs
--- a/Engine/Source/Programs/Enterprise/Datasmith/DatasmithSolidworksExporter/Private/Matrix4.cs
+++ b/Engine/Source/Programs/Enterprise/Datasmith/DatasmithSolidworksExporter/Private/Matrix4.cs
@@ -156,36 +156,6 @@ namespace DatasmithSolidworks
 			return Result;
 		}
 
-		public static float[] LookAt(FVec3 InDirectionZ, FVec3 InFromPoint = null, FVec3 InUp = null)
-		{
-			float[] Ret = null;
-
-			if (InDirectionZ != null)
-			{
-				if (InUp == null)
-				{
-					InUp = new FVec3(0f, 1f, 0f);
-				}
-
-				var DirectionZ = InDirectionZ.Normalized();
-				var DirectionX = FVec3.Cross(InUp, DirectionZ);
-				var DirectionY = FVec3.Cross(DirectionZ, DirectionX);
-				Ret = new[] { DirectionX.X, DirectionY.X, DirectionZ.X, 0f ,
-							  DirectionX.Y, DirectionY.Y, DirectionZ.Y, 0f ,
-							  DirectionX.Z, DirectionY.Z, DirectionZ.Z, 0f ,
-							  0f, 0f, 0f, 1f };
-
-				if (InFromPoint != null)
-				{
-					Ret[3] = InFromPoint.X;
-					Ret[7] = InFromPoint.Y;
-					Ret[11] = InFromPoint.Z;
-				}
-			}
-
-			return Ret;
-		}
-
 		public static float[] FMatrix4x4Multiply(float[] InMatrixA, float[] InMatrixB)
 		{
 			float[] Ret = null;
diff --git a/Engine/Source/Programs/Enterprise/Datasmith/DatasmithSolidworksExporter/Private/Vec3.cs b/Engine/Source/Programs/Enterprise/Datasmith/DatasmithSolidworksExporter/Private/Vec3.cs
--- a/Engine/Source/Programs/Enterprise/Datasmith/DatasmithSolidworksExporter/Private/Vec3.cs
+++ b/Engine/Source/Programs/Enterprise/Datasmith/DatasmithSolidworksExporter/Private/Vec3.cs
@@ -188,6 +188,11 @@ namespace DatasmithSolidworks
             return new FVec3(V.X - M.X, V.Y - M.Y, V.Z - M.Z);
         }
 
+        public static FVec3 operator -(FVec3 V)
+        {
+	        return new FVec3(-V.X, -V.Y, -V.Z);
+        }
+
         public static FVec3 operator +(FVec3 V, float M)
         {
             return new FVec3(V.X + M, V.Y + M, V.Z + M);
