[Solidworks] Infinite reimport after deleting a component auto sync is toggle

- fixed case when change in material was detected when no material were present
- also fixed not detecting material modifications

diff --git a/Engine/Source/Programs/Enterprise/Datasmith/DatasmithSolidworksExporter/Private/AssemblyDocument.cs b/Engine/Source/Programs/Enterprise/Datasmith/DatasmithSolidworksExporter/Private/AssemblyDocument.cs
--- a/Engine/Source/Programs/Enterprise/Datasmith/DatasmithSolidworksExporter/Private/AssemblyDocument.cs
+++ b/Engine/Source/Programs/Enterprise/Datasmith/DatasmithSolidworksExporter/Private/AssemblyDocument.cs
@@ -7,7 +7,6 @@ using System.Collections.Concurrent;
 using System.Collections.Generic;
 using System.Diagnostics;
 using System.Linq;
-using System.Threading.Tasks;
 using DatasmithSolidworks.Names;
 
 namespace DatasmithSolidworks
@@ -28,8 +27,7 @@ namespace DatasmithSolidworks
 		{
 			public Dictionary<string, FPartDocument> PartsMap = new Dictionary<string, FPartDocument>();
 
-			public ConcurrentDictionary<FComponentName, FObjectMaterials> ComponentsMaterialsMap =
-				new ConcurrentDictionary<FComponentName, FObjectMaterials>();
+			public ConcurrentDictionary<FComponentName, FObjectMaterials> ComponentsMaterialsMap = null;
 
 			public Dictionary<FComponentName, string> ComponentToPartMap = new Dictionary<FComponentName, string>();
 
@@ -122,7 +120,7 @@ namespace DatasmithSolidworks
 					Exporter.RemoveMesh(MeshName);
 				}
 
-				SyncState.ComponentsMaterialsMap.TryRemove(CompName, out FObjectMaterials _);
+				SyncState.ComponentsMaterialsMap?.TryRemove(CompName, out FObjectMaterials _);
 				SyncState.ComponentsTransformsMap.Remove(CompName);
 				Exporter.RemoveActor(ActorName);
 			}
@@ -140,6 +138,11 @@ namespace DatasmithSolidworks
 		public override void AddComponentMaterials(FComponentName ComponentName, FObjectMaterials Materials)
 		{
 			Addin.Instance.LogDebug($"AddComponentMaterials: {ComponentName} - {Materials}");
+
+			if (SyncState.ComponentsMaterialsMap == null)
+			{
+				SyncState.ComponentsMaterialsMap = new ConcurrentDictionary<FComponentName, FObjectMaterials>();
+			}
 			SyncState.ComponentsMaterialsMap[ComponentName] = Materials;
 		}
 
@@ -192,7 +195,7 @@ namespace DatasmithSolidworks
 			{
 				foreach (var MatKVP in ModifiedComponentsMaterials)
 				{
-					SyncState.ComponentsMaterialsMap[MatKVP.Key] = MatKVP.Value;
+					AddComponentMaterials(MatKVP.Key, MatKVP.Value);
 				}
 			}
 
@@ -239,85 +242,54 @@ namespace DatasmithSolidworks
 				FObjectMaterials.LoadAssemblyMaterials(this, AllExportedComponents,
 					swDisplayStateOpts_e.swThisDisplayState, null);
 
-			if (CurrentDocMaterialsMap == null && SyncState.ComponentsMaterialsMap == null)
-			{
-				return false;
-			}
-			else if (CurrentDocMaterialsMap == null && SyncState.ComponentsMaterialsMap != null)
-			{
-				foreach (var KVP in SyncState.ComponentsMaterialsMap)
-				{
-					SetComponentDirty(KVP.Key, EComponentDirtyState.Material);
-				}
 
-				return true;
-			}
-			else if (CurrentDocMaterialsMap != null && SyncState.ComponentsMaterialsMap == null)
-			{
-				foreach (var KVP in CurrentDocMaterialsMap)
-				{
-					SetComponentDirty(KVP.Key, EComponentDirtyState.Material);
-				}
+			HashSet<FComponentName> Components =  SyncState.ComponentsMaterialsMap == null ? new HashSet<FComponentName>() : new HashSet<FComponentName>(SyncState.ComponentsMaterialsMap.Keys);
+			HashSet<FComponentName> CurrentComponents =  CurrentDocMaterialsMap == null ? new HashSet<FComponentName>() : new HashSet<FComponentName>(CurrentDocMaterialsMap.Keys);
 
-				return true;
-			}
-			else
+			// Components which stayed in the materials map
+			HashSet<FComponentName> CommonComponents = new HashSet<FComponentName>(CurrentComponents.Intersect(Components));
+
+			// Check components that have material added or removed
+			bool bHasDirtyMaterials = false;
+			IEnumerable<FComponentName> ComponentsWithAddedOrRemovedMaterial = Components.Union(CurrentComponents).Except(CommonComponents);
+			foreach (FComponentName CompName in ComponentsWithAddedOrRemovedMaterial)
 			{
-				bool bHasDirtyMaterials = false;
+				bool bShouldSyncComponentMaterial = false;
 
-				if (CurrentDocMaterialsMap.Count != SyncState.ComponentsMaterialsMap.Count)
+				if (SyncState.ExportedComponentsMap.ContainsKey(CompName))
 				{
-					IEnumerable<FComponentName> Diff1 = CurrentDocMaterialsMap.Keys.Except(SyncState.ComponentsMaterialsMap.Keys);
-					IEnumerable<FComponentName> Diff2 = SyncState.ComponentsMaterialsMap.Keys.Except(CurrentDocMaterialsMap.Keys);
-
-					HashSet<FComponentName> DiffSet = new HashSet<FComponentName>();
-					DiffSet.UnionWith(Diff1);
-					DiffSet.UnionWith(Diff2);
-
-					// Components in the DiffSet have their materials changed
-					foreach (FComponentName CompName in DiffSet)
+					try
 					{
-						bool bShouldSyncComponentMaterial = false;
-
-						if (SyncState.ExportedComponentsMap.ContainsKey(CompName))
-						{
-							try
-							{
-								Component2 Comp = SyncState.ExportedComponentsMap[CompName];
-								bShouldSyncComponentMaterial = !Comp.IsSuppressed() && (Comp.Visible == (int)swComponentVisibilityState_e.swComponentVisible);
-							}
-							catch
-							{
-							}
-						}
-
-						if (bShouldSyncComponentMaterial)
-						{
-							bHasDirtyMaterials = true;
-							SetComponentDirty(CompName, EComponentDirtyState.Material);
-						}
+						Component2 Comp = SyncState.ExportedComponentsMap[CompName];
+						bShouldSyncComponentMaterial = !Comp.IsSuppressed() &&
+						                               (Comp.Visible == (int)swComponentVisibilityState_e
+							                               .swComponentVisible);
+					}
+					catch
+					{
+						// todo: needs a comment on how this happens
 					}
-
-					return bHasDirtyMaterials;
 				}
 
-				bool bHasDirtyComponents = false;
-
-				foreach (var KVP in SyncState.ComponentsMaterialsMap)
+				if (bShouldSyncComponentMaterial)
 				{
-					FObjectMaterials CurrentComponentMaterials;
-					if (CurrentDocMaterialsMap.TryGetValue(KVP.Key, out CurrentComponentMaterials))
-					{
-						if (!CurrentComponentMaterials.EqualMaterials(KVP.Value))
-						{
-							SetComponentDirty(KVP.Key, EComponentDirtyState.Material);
-							bHasDirtyComponents = true;
-						}
-					}
+					bHasDirtyMaterials = true;
+					SetComponentDirty(CompName, EComponentDirtyState.Material);
 				}
+			}
 
-				return bHasDirtyComponents;
+			// Check if components have their material modified
+			bool bHasDirtyComponents = false;
+			foreach (FComponentName ComponentName in CommonComponents)
+			{
+				if (!SyncState.ComponentsMaterialsMap[ComponentName].EqualMaterials(CurrentDocMaterialsMap[ComponentName]))
+				{
+					SetComponentDirty(ComponentName, EComponentDirtyState.Material);
+					bHasDirtyComponents = true;
+				}
 			}
+
+			return bHasDirtyComponents || bHasDirtyMaterials;
 		}
 
 		public float[] GetComponentDatasmithTransform(Component2 InComponent)
