[Solidworks] Fixed exporter not installing if only 2023 version is available

diff --git a/Engine/Source/Programs/Enterprise/Datasmith/DatasmithSolidworksExporter/RegistryKeyExistsCustomAction/CustomAction.cs b/Engine/Source/Programs/Enterprise/Datasmith/DatasmithSolidworksExporter/RegistryKeyExistsCustomAction/CustomAction.cs
--- a/Engine/Source/Programs/Enterprise/Datasmith/DatasmithSolidworksExporter/RegistryKeyExistsCustomAction/CustomAction.cs
+++ b/Engine/Source/Programs/Enterprise/Datasmith/DatasmithSolidworksExporter/RegistryKeyExistsCustomAction/CustomAction.cs
@@ -60,12 +60,14 @@ namespace RegistryKeyExistsCustomAction
 			var tools2020 = CheckSolidworksInstalledSub(session, "SOLIDWORKS 2020", "SOLIDWORKS2020INSTALLED");
 			var tools2021 = CheckSolidworksInstalledSub(session, "SOLIDWORKS 2021", "SOLIDWORKS2021INSTALLED");
 			var tools2022 = CheckSolidworksInstalledSub(session, "SOLIDWORKS 2022", "SOLIDWORKS2022INSTALLED");
-			if (!string.IsNullOrEmpty(tools2022))
-				session["SOLIDWORKSTOOLSPATH"] = tools2022;
-			else if (!string.IsNullOrEmpty(tools2021))
-				session["SOLIDWORKSTOOLSPATH"] = tools2021;
-			else if (!string.IsNullOrEmpty(tools2020))
-				session["SOLIDWORKSTOOLSPATH"] = tools2020;
+			var tools2023 = CheckSolidworksInstalledSub(session, "SOLIDWORKS 2023", "SOLIDWORKS2023INSTALLED");
+			// SOLIDWORKSTOOLSPATH is unused(not installing solidworkstools.dll anymore)
+			// if (!string.IsNullOrEmpty(tools2022))
+			// 	session["SOLIDWORKSTOOLSPATH"] = tools2022;
+			// else if (!string.IsNullOrEmpty(tools2021))
+			// 	session["SOLIDWORKSTOOLSPATH"] = tools2021;
+			// else if (!string.IsNullOrEmpty(tools2020))
+			// 	session["SOLIDWORKSTOOLSPATH"] = tools2020;
 
 			return ActionResult.Success;
 		}
