[Solidworks] Fixed The color of the component is being changed when pressing Datasmith sync

diff --git a/Engine/Source/Programs/Enterprise/Datasmith/DatasmithSolidworksExporter/Private/AnimationExtractor.cs b/Engine/Source/Programs/Enterprise/Datasmith/DatasmithSolidworksExporter/Private/AnimationExtractor.cs
--- a/Engine/Source/Programs/Enterprise/Datasmith/DatasmithSolidworksExporter/Private/AnimationExtractor.cs
+++ b/Engine/Source/Programs/Enterprise/Datasmith/DatasmithSolidworksExporter/Private/AnimationExtractor.cs
@@ -34,6 +34,9 @@ namespace DatasmithSolidworks
                 // for components
                 Doc.EnablePresentation = true;
             }
+
+            ModelDoc2 ModelDoc = Doc as ModelDoc2;
+            bool ModelTabActiveStored = ModelDoc.ModelViewManager.IsModelTabActive();
             
 			string[] StudyNames = MSManager.GetMotionStudyNames();
 			for (int Idx = 0; Idx < StudyCount; Idx++)
@@ -48,7 +51,14 @@ namespace DatasmithSolidworks
 					}
 				}
 			}
-			
+
+			if (ModelTabActiveStored)
+			{
+				// Switch back to Model Tab after activating Motion Study tabs for least user surprise.
+				// Motion Study model state was quite probably very different from default Model view
+				ModelDoc.ModelViewManager.ActivateModelTab();
+			}			
+
 			if (!bWasInPresentationMode)
 			{
                 Doc.EnablePresentation = false; // Return back to normal mode
