Fixed warnings and compilation failure
  Solution: Reworked the header files to make sure all ArchiCAD headers are in one header file. That header file properly included.

diff --git a/Engine/Source/Programs/Enterprise/Datasmith/DatasmithARCHICADExporter/Private/AddOn.cpp b/Engine/Source/Programs/Enterprise/Datasmith/DatasmithARCHICADExporter/Private/AddOn.cpp
--- a/Engine/Source/Programs/Enterprise/Datasmith/DatasmithARCHICADExporter/Private/AddOn.cpp
+++ b/Engine/Source/Programs/Enterprise/Datasmith/DatasmithARCHICADExporter/Private/AddOn.cpp
@@ -4,25 +4,15 @@
 
 #include <stddef.h>
 
-#include "Utils/WarningsDisabler.h"
+#include "Utils/APIEnvir.h"
 #include "Utils/TaskMgr.h"
 #include "Utils/TaskCalledFromEventLoop.h"
 
-DISABLE_SDK_WARNINGS_START
-
 #include "IDatasmithSceneElements.h"
 #include "DatasmithSceneFactory.h"
 
 #undef PI
 
-#include "Sight.hpp"
-#if AC_VERSION < 26
-#include "AttributeReader.hpp"
-#endif
-#include "Model.hpp"
-
-DISABLE_SDK_WARNINGS_END
-
 #include "Utils/DebugTools.h"
 #include "LoadDatasmithDlls.h"
 #include "Export.h"
diff --git a/Engine/Source/Programs/Enterprise/Datasmith/DatasmithARCHICADExporter/Private/Commander.cpp b/Engine/Source/Programs/Enterprise/Datasmith/DatasmithARCHICADExporter/Private/Commander.cpp
--- a/Engine/Source/Programs/Enterprise/Datasmith/DatasmithARCHICADExporter/Private/Commander.cpp
+++ b/Engine/Source/Programs/Enterprise/Datasmith/DatasmithARCHICADExporter/Private/Commander.cpp
@@ -23,17 +23,6 @@
 
 #include "exp.h"
 
-DISABLE_SDK_WARNINGS_START
-#include "Sight.hpp"
-DISABLE_SDK_WARNINGS_END
-
-#if AC_VERSION < 26
-#include "AttributeReader.hpp"
-#else
-#include "IAttributeReader.hpp"
-#endif
-#include "Model.hpp"
-
 BEGIN_NAMESPACE_UE_AC
 
 void FCommander::DoSnapshot()
diff --git a/Engine/Source/Programs/Enterprise/Datasmith/DatasmithARCHICADExporter/Private/ConnectionWindow.cpp b/Engine/Source/Programs/Enterprise/Datasmith/DatasmithARCHICADExporter/Private/ConnectionWindow.cpp
--- a/Engine/Source/Programs/Enterprise/Datasmith/DatasmithARCHICADExporter/Private/ConnectionWindow.cpp
+++ b/Engine/Source/Programs/Enterprise/Datasmith/DatasmithARCHICADExporter/Private/ConnectionWindow.cpp
@@ -8,19 +8,11 @@
 #include "Utils/ShellOpenDocument.h"
 #include "Utils/TaskCalledFromEventLoop.h"
 
-DISABLE_SDK_WARNINGS_START
-
-#include "DGDialog.hpp"
-#include "DGModule.hpp"
-#include "DGFileDlg.hpp"
-
 #undef TicksPerSecond
 
 #include "DirectLinkEndpoint.h"
 #include "DatasmithDirectLink.h"
 
-DISABLE_SDK_WARNINGS_END
-
 BEGIN_NAMESPACE_UE_AC
 
 static class FConnectionDialog* ConnectionDialog = nullptr;
diff --git a/Engine/Source/Programs/Enterprise/Datasmith/DatasmithARCHICADExporter/Private/Export.cpp b/Engine/Source/Programs/Enterprise/Datasmith/DatasmithARCHICADExporter/Private/Export.cpp
--- a/Engine/Source/Programs/Enterprise/Datasmith/DatasmithARCHICADExporter/Private/Export.cpp
+++ b/Engine/Source/Programs/Enterprise/Datasmith/DatasmithARCHICADExporter/Private/Export.cpp
@@ -6,17 +6,6 @@
 #include "Utils/AutoChangeDatabase.h"
 #include "Exporter.h"
 
-#include "exp.h"
-#include "Sight.hpp"
-DISABLE_SDK_WARNINGS_START
-#if AC_VERSION < 26
-#include "AttributeReader.hpp"
-#else
-#include "IAttributeReader.hpp"
-#endif
-DISABLE_SDK_WARNINGS_END
-#include "FileSystem.hpp"
-
 BEGIN_NAMESPACE_UE_AC
 
 enum : GSType
diff --git a/Engine/Source/Programs/Enterprise/Datasmith/DatasmithARCHICADExporter/Private/Exporter.cpp b/Engine/Source/Programs/Enterprise/Datasmith/DatasmithARCHICADExporter/Private/Exporter.cpp
--- a/Engine/Source/Programs/Enterprise/Datasmith/DatasmithARCHICADExporter/Private/Exporter.cpp
+++ b/Engine/Source/Programs/Enterprise/Datasmith/DatasmithARCHICADExporter/Private/Exporter.cpp
@@ -1,16 +1,12 @@
 // Copyright Epic Games, Inc. All Rights Reserved.
 
-#include <stddef.h>
-
-#include "Utils/WarningsDisabler.h"
+#include "Utils/APIEnvir.h"
 
-DISABLE_SDK_WARNINGS_START
+#include <stddef.h>
 
 #include "DatasmithSceneExporter.h"
 #include "DatasmithSceneXmlWriter.h"
 
-DISABLE_SDK_WARNINGS_END
-
 #include "Exporter.h"
 #include "ResourcesIDs.h"
 #include "Utils/TimeStat.h"
@@ -19,16 +15,12 @@ DISABLE_SDK_WARNINGS_END
 	#undef TicksPerSecond
 #endif
 
-DISABLE_SDK_WARNINGS_START
-
 #include "FileManager.h"
 #include "Paths.h"
 
 #include "FileSystem.hpp"
 #include "DGFileDialog.hpp"
 
-DISABLE_SDK_WARNINGS_END
-
 BEGIN_NAMESPACE_UE_AC
 
 // Constructor
diff --git a/Engine/Source/Programs/Enterprise/Datasmith/DatasmithARCHICADExporter/Private/LoadDatasmithDlls.cpp b/Engine/Source/Programs/Enterprise/Datasmith/DatasmithARCHICADExporter/Private/LoadDatasmithDlls.cpp
--- a/Engine/Source/Programs/Enterprise/Datasmith/DatasmithARCHICADExporter/Private/LoadDatasmithDlls.cpp
+++ b/Engine/Source/Programs/Enterprise/Datasmith/DatasmithARCHICADExporter/Private/LoadDatasmithDlls.cpp
@@ -3,14 +3,10 @@
 #include "LoadDatasmithDlls.h"
 
 #include <stddef.h>
-#include "Utils/WarningsDisabler.h"
-
-DISABLE_SDK_WARNINGS_START
+#include "Utils/APIEnvir.h"
 
 #include "IDatasmithSceneElements.h"
 
-DISABLE_SDK_WARNINGS_END
-
 #include "DatasmithSceneFactory.h"
 #include "DatasmithExporterManager.h"
 #include "DatasmithDirectLink.h"
@@ -30,14 +26,6 @@ DISABLE_SDK_WARNINGS_END
 	#pragma warning(pop)
 #endif
 
-#include "Utils/APIEnvir.h"
-
-DISABLE_SDK_WARNINGS_START
-#include "ACAPinc.h"
-#include "Md5.hpp"
-#include "Lock.hpp"
-DISABLE_SDK_WARNINGS_END
-
 #include "Utils/LocalizeTools.h"
 #include "Utils/AddonTools.h"
 
diff --git a/Engine/Source/Programs/Enterprise/Datasmith/DatasmithARCHICADExporter/Private/Palette.cpp b/Engine/Source/Programs/Enterprise/Datasmith/DatasmithARCHICADExporter/Private/Palette.cpp
--- a/Engine/Source/Programs/Enterprise/Datasmith/DatasmithARCHICADExporter/Private/Palette.cpp
+++ b/Engine/Source/Programs/Enterprise/Datasmith/DatasmithARCHICADExporter/Private/Palette.cpp
@@ -2,6 +2,7 @@
 
 #include "Palette.h"
 #include "Synchronizer.h"
+#include "Utils/APIEnvir.h"
 #include "Utils/TAssValueName.h"
 #include "ResourcesIDs.h"
 #include "Utils/Error.h"
@@ -9,10 +10,6 @@
 #include "Commander.h"
 #include "Menus.h"
 
-DISABLE_SDK_WARNINGS_START
-#include "DGModule.hpp"
-DISABLE_SDK_WARNINGS_END
-
 #define PALETTE_4_ALL_VIEW 1
 #define TRACE_PALETTE 0
 
diff --git a/Engine/Source/Programs/Enterprise/Datasmith/DatasmithARCHICADExporter/Private/ReportWindow.cpp b/Engine/Source/Programs/Enterprise/Datasmith/DatasmithARCHICADExporter/Private/ReportWindow.cpp
--- a/Engine/Source/Programs/Enterprise/Datasmith/DatasmithARCHICADExporter/Private/ReportWindow.cpp
+++ b/Engine/Source/Programs/Enterprise/Datasmith/DatasmithARCHICADExporter/Private/ReportWindow.cpp
@@ -9,11 +9,6 @@
 #include "Utils/Error.h"
 #include "Utils/TaskCalledFromEventLoop.h"
 
-DISABLE_SDK_WARNINGS_START
-#include "DGDialog.hpp"
-#include "DGModule.hpp"
-DISABLE_SDK_WARNINGS_END
-
 BEGIN_NAMESPACE_UE_AC
 
 class FReportDialog : public DG::Palette,
diff --git a/Engine/Source/Programs/Enterprise/Datasmith/DatasmithARCHICADExporter/Private/SyncData.cpp b/Engine/Source/Programs/Enterprise/Datasmith/DatasmithARCHICADExporter/Private/SyncData.cpp
--- a/Engine/Source/Programs/Enterprise/Datasmith/DatasmithARCHICADExporter/Private/SyncData.cpp
+++ b/Engine/Source/Programs/Enterprise/Datasmith/DatasmithARCHICADExporter/Private/SyncData.cpp
@@ -14,11 +14,6 @@
 #include "TexturesCache.h"
 #include "Utils/TaskMgr.h"
 
-DISABLE_SDK_WARNINGS_START
-#include "Transformation.hpp"
-#include "Line3D.hpp"
-DISABLE_SDK_WARNINGS_END
-
 #include <stdexcept>
 
 BEGIN_NAMESPACE_UE_AC
@@ -189,8 +184,6 @@ bool FSyncData::CheckAllCutOut()
 	return true;
 }
 
-#pragma mark -
-
 // Guid given to the scene element.
 const GS::Guid FSyncData::FScene::SceneGUID("CBDEFBEF-0D4E-4162-8C4C-64AC34CEB4E6");
 
@@ -375,8 +368,6 @@ void FSyncData::FScene::RemoveChildActor(const TSharedPtr< IDatasmithActorElemen
 	SceneElement->RemoveActor(InActor, EDatasmithActorRemovalRule::RemoveChildren);
 }
 
-#pragma mark -
-
 FSyncData::FActor::FActor(const GS::Guid& InGuid)
 	: FSyncData(InGuid)
 {
@@ -473,8 +464,6 @@ void FSyncData::FActor::ReplaceMetaData(IDatasmithScene&							   IOScene,
 	IOScene.AddMetaData(MetaData);
 }
 
-#pragma mark -
-
 // Guid used to synthetize layer guid
 const GS::Guid FSyncData::FLayer::LayerGUID("97D32F90-A33E-0000-8305-D1A7D3FCED66");
 
@@ -546,8 +535,6 @@ void FSyncData::FLayer::Process(FProcessInfo* /* IOProcessInfo */)
 	}
 }
 
-#pragma mark -
-
 inline Geometry::Transformation3D Convert(const ModelerAPI::Transformation& InMatrix)
 {
 	Geometry::Matrix33 M33;
@@ -1053,8 +1040,6 @@ bool FSyncData::FElement::UpdateMetaData(IDatasmithScene* IOScene)
 	return MetaDataExporter.SetOrUpdate(&MetaData, IOScene);
 }
 
-#pragma mark -
-
 void FSyncData::FCameraSet::Process(FProcessInfo* /* IOProcessInfo */)
 {
 	if (!ActorElement.IsValid())
@@ -1077,8 +1062,6 @@ void FSyncData::FCameraSet::Process(FProcessInfo* /* IOProcessInfo */)
 	}
 }
 
-#pragma mark -
-
 // Guid given to the current view.
 const GS::Guid FSyncData::FCamera::CurrentViewGUID("B2BD9C50-60EB-4E64-902B-D1574FADEC45");
 
@@ -1177,8 +1160,6 @@ void FSyncData::FCamera::InitWithCameraElement()
 		FGeometryUtil::GetCameraFocalLength(CameraElement.GetSensorWidth(), camPars.viewCone * RADDEG));
 }
 
-#pragma mark -
-
 void FSyncData::FLight::Process(FProcessInfo* IOProcessInfo)
 {
 	if (!ActorElement.IsValid())
@@ -1449,8 +1430,6 @@ bool FSyncData::FLight::FLightData::operator!=(const FLightData& InOther) const
 		   Rotation != InOther.Rotation;
 }
 
-#pragma mark -
-
 const GS::Guid FSyncData::FHotLinksRoot::HotLinksRootGUID("C4BFD876-FDE9-4CCF-8899-12023968DC0D");
 
 void FSyncData::FHotLinksRoot::Process(FProcessInfo* /* IOProcessInfo */)
@@ -1628,8 +1607,6 @@ void FSyncData::FHotLinkInstance::Process(FProcessInfo* IOProcessInfo)
 	}
 }
 
-#pragma mark -
-
 // Start the process with this root observer
 void FSyncData::FInterator::Start(FSyncData* Root)
 {
@@ -1701,8 +1678,6 @@ FSyncData::FChildsArray::SizeType FSyncData::FInterator::GetCurrentIndex()
 	return Stack.Num() > 1 ? Stack[Stack.Num() - 2].ChildIndex : 0;
 }
 
-#pragma mark -
-
 // Start the process with this root observer
 void FSyncData::FProcessMetadata::Start(FSyncData* Root)
 {
@@ -1728,8 +1703,6 @@ FSyncData::FInterator::EProcessControl FSyncData::FProcessMetadata::Process(FSyn
 	return FInterator::kContinue;
 }
 
-#pragma mark -
-
 // Constructor
 FSyncData::FAttachObservers::FAttachObservers() {}
 
diff --git a/Engine/Source/Programs/Enterprise/Datasmith/DatasmithARCHICADExporter/Private/SyncDatabase.cpp b/Engine/Source/Programs/Enterprise/Datasmith/DatasmithARCHICADExporter/Private/SyncDatabase.cpp
--- a/Engine/Source/Programs/Enterprise/Datasmith/DatasmithARCHICADExporter/Private/SyncDatabase.cpp
+++ b/Engine/Source/Programs/Enterprise/Datasmith/DatasmithARCHICADExporter/Private/SyncDatabase.cpp
@@ -17,8 +17,6 @@
 
 #undef TicksPerSecond
 
-DISABLE_SDK_WARNINGS_START
-
 #include "DatasmithUtils.h"
 #include "IDirectLinkUI.h"
 #include "IDatasmithExporterUIModule.h"
@@ -28,8 +26,6 @@ DISABLE_SDK_WARNINGS_START
 	#include "Misc/ConfigCacheIni.h"
 #endif
 
-DISABLE_SDK_WARNINGS_END
-
 BEGIN_NAMESPACE_UE_AC
 
 // Control access on this object (for queue operations)
diff --git a/Engine/Source/Programs/Enterprise/Datasmith/DatasmithARCHICADExporter/Private/SyncDatabase.h b/Engine/Source/Programs/Enterprise/Datasmith/DatasmithARCHICADExporter/Private/SyncDatabase.h
--- a/Engine/Source/Programs/Enterprise/Datasmith/DatasmithARCHICADExporter/Private/SyncDatabase.h
+++ b/Engine/Source/Programs/Enterprise/Datasmith/DatasmithARCHICADExporter/Private/SyncDatabase.h
@@ -8,10 +8,8 @@
 #include "Lock.hpp"
 #include "ModelElement.hpp"
 
-DISABLE_SDK_WARNINGS_START
 #include "Map.h"
 #include "List.h"
-DISABLE_SDK_WARNINGS_END
 
 BEGIN_NAMESPACE_UE_AC
 
diff --git a/Engine/Source/Programs/Enterprise/Datasmith/DatasmithARCHICADExporter/Private/Synchronizer.cpp b/Engine/Source/Programs/Enterprise/Datasmith/DatasmithARCHICADExporter/Private/Synchronizer.cpp
--- a/Engine/Source/Programs/Enterprise/Datasmith/DatasmithARCHICADExporter/Private/Synchronizer.cpp
+++ b/Engine/Source/Programs/Enterprise/Datasmith/DatasmithARCHICADExporter/Private/Synchronizer.cpp
@@ -18,14 +18,10 @@
 	#undef TicksPerSecond
 #endif
 
-DISABLE_SDK_WARNINGS_START
-
 #include "FileManager.h"
 #include "Paths.h"
 #include "Version.h"
 
-DISABLE_SDK_WARNINGS_END
-
 BEGIN_NAMESPACE_UE_AC
 
 // Do Direct Link snapshot update on a thread
diff --git a/Engine/Source/Programs/Enterprise/Datasmith/DatasmithARCHICADExporter/Private/TexturesCache.cpp b/Engine/Source/Programs/Enterprise/Datasmith/DatasmithARCHICADExporter/Private/TexturesCache.cpp
--- a/Engine/Source/Programs/Enterprise/Datasmith/DatasmithARCHICADExporter/Private/TexturesCache.cpp
+++ b/Engine/Source/Programs/Enterprise/Datasmith/DatasmithARCHICADExporter/Private/TexturesCache.cpp
@@ -159,7 +159,7 @@ class FSearchAndCopyFile
 	}
 
 	// Search in the folder tree starting at parent folder.
-	bool DoSearchIn(const IO::Folder& Parent) { return Parent.Enumerate(EnumCallBack, this); }
+	bool DoSearchIn(const IO::Folder& Parent) { return Parent.Enumerate(EnumCallBack, this) == NoError; }
 
 	// Callback function called for each element of the parent specified
 	bool EnumCallBack(const IO::Folder& Parent, const IO::Name& EntryName, bool bIsFolder);
diff --git a/Engine/Source/Programs/Enterprise/Datasmith/DatasmithARCHICADExporter/Private/Utils/3DElement2String.cpp b/Engine/Source/Programs/Enterprise/Datasmith/DatasmithARCHICADExporter/Private/Utils/3DElement2String.cpp
--- a/Engine/Source/Programs/Enterprise/Datasmith/DatasmithARCHICADExporter/Private/Utils/3DElement2String.cpp
+++ b/Engine/Source/Programs/Enterprise/Datasmith/DatasmithARCHICADExporter/Private/Utils/3DElement2String.cpp
@@ -1,22 +1,12 @@
 // Copyright Epic Games, Inc. All Rights Reserved.
 
+#include "APIEnvir.h"
+
 #include "3DElement2String.h"
 #include "ElementTools.h"
 #include "Element2String.h"
 #include "TAssValueName.h"
 
-DISABLE_SDK_WARNINGS_START
-
-#include "ModelElement.hpp"
-#include "ModelMeshBody.hpp"
-#include "ConvexPolygon.hpp"
-
-#include "Transformation.hpp"
-#include "Parameter.hpp"
-#include "Light.hpp"
-
-DISABLE_SDK_WARNINGS_END
-
 BEGIN_NAMESPACE_UE_AC
 
 static void AddIfNotZero(utf8_string* IOString, const char* InFmt, int32 InCount)
diff --git a/Engine/Source/Programs/Enterprise/Datasmith/DatasmithARCHICADExporter/Private/Utils/APIEnvir.h b/Engine/Source/Programs/Enterprise/Datasmith/DatasmithARCHICADExporter/Private/Utils/APIEnvir.h
--- a/Engine/Source/Programs/Enterprise/Datasmith/DatasmithARCHICADExporter/Private/Utils/APIEnvir.h
+++ b/Engine/Source/Programs/Enterprise/Datasmith/DatasmithARCHICADExporter/Private/Utils/APIEnvir.h
@@ -4,22 +4,84 @@
 
 #include <stddef.h>
 
-#include "WarningsDisabler.h"
-
-DISABLE_SDK_WARNINGS_START
-
 #include "IDatasmithSceneElements.h"
 #include "DatasmithSceneFactory.h"
 
 #undef PI
 
+// clang-format off
+
+THIRD_PARTY_INCLUDES_START
+
+#ifdef __clang__
+
+#pragma clang diagnostic push
+#pragma clang diagnostic ignored "-Wcomma"
+#pragma clang diagnostic ignored "-Wdefaulted-function-deleted"
+
+#if (__clang_major__ > 12) || (__clang_major__ == 12 && __clang_minor__ == 0 && __clang_patchlevel__ > 4)
+#pragma clang diagnostic ignored "-Wnon-c-typedef-for-linkage"
+#endif
+
+#elif PLATFORM_WINDOWS
+
+__pragma(warning(push))
+__pragma(warning(disable: 4005)) // 'TEXT': macro redefinition
+__pragma(warning(disable: 5040)) // Support\Modules\GSRoot\vaarray.hpp: dynamic exception specifications are valid only in C++14 and earlier
+
+#else
+
+_Pragma("clang diagnostic ignored \"-Wcomma\"")
+_Pragma("clang diagnostic ignored \"-Wdefaulted-function-deleted\"")
+
+#endif
+
+// clang-format on
+
+#include "ACAPinc.h"
+#include "Md5.hpp"
+
 #include "Sight.hpp"
 #if AC_VERSION < 26
 #include "AttributeReader.hpp"
+#else
+#include "IAttributeReader.hpp"
 #endif
 #include "Model.hpp"
 
-DISABLE_SDK_WARNINGS_END
+#include "ModelElement.hpp"
+#include "ModelMeshBody.hpp"
+#include "ConvexPolygon.hpp"
+
+#include "Transformation.hpp"
+#include "Parameter.hpp"
+#include "Light.hpp"
+
+#include "exp.h"
+#include "FileSystem.hpp"
+
+#include "Lock.hpp"
+
+#include "DGModule.hpp"
+
+#include "DGDialog.hpp"
+
+#include "Transformation.hpp"
+#include "Line3D.hpp"
+
+#include "DGFileDlg.hpp"
+
+// clang-format off
+#ifdef __clang__
+#pragma clang diagnostic pop
+#elif PLATFORM_WINDOWS
+__pragma(warning(pop))
+#else
+_Pragma("clang diagnostic pop")
+#endif
+// clang-format on
+
+THIRD_PARTY_INCLUDES_END
 
 #if defined(_MSC_VER)
 	#if !defined(WINDOWS)
diff --git a/Engine/Source/Programs/Enterprise/Datasmith/DatasmithARCHICADExporter/Private/Utils/AddonTools.h b/Engine/Source/Programs/Enterprise/Datasmith/DatasmithARCHICADExporter/Private/Utils/AddonTools.h
--- a/Engine/Source/Programs/Enterprise/Datasmith/DatasmithARCHICADExporter/Private/Utils/AddonTools.h
+++ b/Engine/Source/Programs/Enterprise/Datasmith/DatasmithARCHICADExporter/Private/Utils/AddonTools.h
@@ -2,7 +2,7 @@
 
 #pragma once
 
-#include "WarningsDisabler.h"
+#include "APIEnvir.h"
 
 #if PLATFORM_WINDOWS
 	#pragma warning(push)
@@ -17,12 +17,6 @@
 	#pragma warning(pop)
 #endif
 
-DISABLE_SDK_WARNINGS_START
-#include "APIEnvir.h"
-#include "ACAPinc.h"
-#include "Md5.hpp"
-DISABLE_SDK_WARNINGS_END
-
 #include "LocalizeTools.h"
 #include "DebugTools.h"
 
diff --git a/Engine/Source/Programs/Enterprise/Datasmith/DatasmithARCHICADExporter/Private/Utils/WarningsDisabler.h b/Engine/Source/Programs/Enterprise/Datasmith/DatasmithARCHICADExporter/Private/Utils/WarningsDisabler.h
deleted file mode 100644
diff --git a/Engine/Source/Programs/Enterprise/Datasmith/DatasmithARCHICADExporter/Private/Utils/WinOS.cpp b/Engine/Source/Programs/Enterprise/Datasmith/DatasmithARCHICADExporter/Private/Utils/WinOS.cpp
--- a/Engine/Source/Programs/Enterprise/Datasmith/DatasmithARCHICADExporter/Private/Utils/WinOS.cpp
+++ b/Engine/Source/Programs/Enterprise/Datasmith/DatasmithARCHICADExporter/Private/Utils/WinOS.cpp
@@ -19,7 +19,7 @@ void ThrowWinError(DWORD InWinErr, const utf8_t* InFile, int InLineNo)
 	}
 
 	char FormattedMessage[1024];
-	snprintf(FormattedMessage, sizeof(FormattedMessage), "Error %d=\"%s\" at \"%s:%d\"", InWinErr,
+	snprintf(FormattedMessage, sizeof(FormattedMessage), "Error %ld=\"%s\" at \"%s:%d\"", InWinErr,
 			 Utf16ToUtf8(WinMsg).c_str(), InFile, InLineNo);
 
 	throw std::runtime_error(FormattedMessage);
