[Solidworks] Unused materials not left in synced/exported scene

- fixed export of materials that are not present in the active configuration when exporting configuration

diff --git a/Engine/Source/Programs/Enterprise/Datasmith/DatasmithSolidworksExporter/Private/AssemblyDocument.cs b/Engine/Source/Programs/Enterprise/Datasmith/DatasmithSolidworksExporter/Private/AssemblyDocument.cs
--- a/Engine/Source/Programs/Enterprise/Datasmith/DatasmithSolidworksExporter/Private/AssemblyDocument.cs
+++ b/Engine/Source/Programs/Enterprise/Datasmith/DatasmithSolidworksExporter/Private/AssemblyDocument.cs
@@ -199,11 +199,11 @@ namespace DatasmithSolidworks
 				}
 			}
 
-			Exporter.ExportMaterials(ExportedMaterialsMap);
-
 			// Export meshes
 			SetExportStatus($"Component Meshes");
-			ProcessConfigurationMeshes(MeshesConfiguration, null);
+			List<FDatasmithExporter.FMeshExportInfo> CreatedMeshes = ProcessConfigurationMeshes(MeshesConfiguration, null);
+			ExportMaterials();
+			AssignMaterialsToDatasmithMeshes(CreatedMeshes);
 
 			SyncState.DirtyComponents.Clear();
 
diff --git a/Engine/Source/Programs/Enterprise/Datasmith/DatasmithSolidworksExporter/Private/ConfigurationExporter.cs b/Engine/Source/Programs/Enterprise/Datasmith/DatasmithSolidworksExporter/Private/ConfigurationExporter.cs
--- a/Engine/Source/Programs/Enterprise/Datasmith/DatasmithSolidworksExporter/Private/ConfigurationExporter.cs
+++ b/Engine/Source/Programs/Enterprise/Datasmith/DatasmithSolidworksExporter/Private/ConfigurationExporter.cs
@@ -354,15 +354,13 @@ namespace DatasmithSolidworks
 					}
 				}
 				
-				// Export materials before exporting meshes(currently mesh export code is tied to FDatasmithExporter.ExportedMaterialsMap
-				// todo: separate material export and mesh export(see comment above)?
-				InDoc.Exporter.ExportMaterials(InDoc.ExportedMaterialsMap);
-
 				// Export meshes
 				InDoc.SetExportStatus($"Component Meshes");
 
 				string MeshSuffix = bIsActiveConfiguration ? null : CfgName; // Suffix for main configuration is absent(so in default or single configuration export we don't have unnecessary long names)
-				InDoc.ProcessConfigurationMeshes(MeshesConfiguration, MeshSuffix);
+				List<FDatasmithExporter.FMeshExportInfo> CreatedMeshes = InDoc.ProcessConfigurationMeshes(MeshesConfiguration, MeshSuffix);
+				InDoc.ExportMaterials();
+				InDoc.AssignMaterialsToDatasmithMeshes(CreatedMeshes);
 
 				// Combine separate scene trees into the single one with configuration-specific data
 				FConfigurationTree.Merge(CombinedTree, ConfigNode, CfgName, bIsActiveConfiguration);
diff --git a/Engine/Source/Programs/Enterprise/Datasmith/DatasmithSolidworksExporter/Private/DatasmithExporter.cs b/Engine/Source/Programs/Enterprise/Datasmith/DatasmithSolidworksExporter/Private/DatasmithExporter.cs
--- a/Engine/Source/Programs/Enterprise/Datasmith/DatasmithSolidworksExporter/Private/DatasmithExporter.cs
+++ b/Engine/Source/Programs/Enterprise/Datasmith/DatasmithSolidworksExporter/Private/DatasmithExporter.cs
@@ -156,15 +156,17 @@ namespace DatasmithSolidworks
 	public class FDatasmithExporter
 	{
 		private Dictionary<FActorName, Tuple<EActorType, FDatasmithFacadeActor>> ExportedActorsMap = new Dictionary<FActorName, Tuple<EActorType, FDatasmithFacadeActor>>();
-		private ConcurrentDictionary<FMeshName, FDatasmithFacadeMeshElement> ExportedMeshesMap = new ConcurrentDictionary<FMeshName, FDatasmithFacadeMeshElement>();
-		private ConcurrentDictionary<int, FDatasmithFacadeMaterialInstance> ExportedMaterialsMap = new ConcurrentDictionary<int, FDatasmithFacadeMaterialInstance>();
+		private ConcurrentDictionary<FMeshName, FMeshExportInfo> ExportedMeshesMap = new ConcurrentDictionary<FMeshName, FMeshExportInfo>();
+		private ConcurrentDictionary<int, FDatasmithFacadeMaterialInstance> ExportedMaterialInstancesMap = new ConcurrentDictionary<int, FDatasmithFacadeMaterialInstance>();
 		private ConcurrentDictionary<string, FDatasmithFacadeTexture> ExportedTexturesMap = new ConcurrentDictionary<string, FDatasmithFacadeTexture>();
-		private Dictionary<string, FDatasmithFacadeActorBinding> ExportedActorBindingsMap = new Dictionary<string, FDatasmithFacadeActorBinding>();
-		private Dictionary<string, FDatasmithFacadeVariant> ExportedVariantsMap = new Dictionary<string, FDatasmithFacadeVariant>();
+
+		// Number of meshes/variants using this material to allow exporting only used materials
+		private Dictionary<int, int> MaterialUsers = new Dictionary<int, int>();
+
 
 		private FDatasmithFacadeScene DatasmithScene = null;
 
-		public struct FMeshExportInfo
+		public class FMeshExportInfo
 		{
 			public FComponentName ComponentName;
 			public FMeshName MeshName;
@@ -172,6 +174,7 @@ namespace DatasmithSolidworks
 			public FMeshData MeshData;
 
 			public FDatasmithFacadeMeshElement MeshElement;
+			public HashSet<int> Materials;
 
 			public override string ToString()
 			{
@@ -259,13 +262,21 @@ namespace DatasmithSolidworks
 			return Actor;
 		}
 
-		public void AddMesh(FMeshName MeshName, FDatasmithFacadeMeshElement NewMesh)
+		private void AddMesh(FMeshExportInfo Info)
 		{
-			Debug.Assert(MeshName.IsValid());
-			Debug.Assert(NewMesh != null);
-			RemoveMesh(MeshName);
-			ExportedMeshesMap.TryAdd(MeshName, NewMesh);
-			DatasmithScene.AddMesh(NewMesh);
+			Debug.Assert(Info.MeshName.IsValid());
+			Debug.Assert(Info.MeshElement != null);
+
+			RemoveMesh(Info.MeshName);
+
+			ExportedMeshesMap.TryAdd(Info.MeshName, Info);
+			DatasmithScene.AddMesh(Info.MeshElement);
+
+			// Increase user count on materials
+			foreach (int MaterialId in Info.Materials)
+			{
+				AddMaterialUser(MaterialId);
+			}
 		}
 
 		public void RemoveMesh(FMeshName MeshName)
@@ -274,9 +285,47 @@ namespace DatasmithSolidworks
             {
                 return;
             }
-			if (ExportedMeshesMap.TryRemove(MeshName, out FDatasmithFacadeMeshElement OldMesh))
+			if (ExportedMeshesMap.TryRemove(MeshName, out FMeshExportInfo Info))
 			{
-				DatasmithScene.RemoveMesh(OldMesh);
+				DatasmithScene.RemoveMesh(Info.MeshElement);
+
+				// Decrease user count on materials
+				foreach (int MaterialId in Info.Materials)
+				{
+					RemoveMaterialUser(MaterialId);
+				}
+
+			}
+		}
+
+		private void AddMaterialUser(int MaterialId)
+		{
+			if (MaterialUsers.TryGetValue(MaterialId, out int Count))
+			{
+				MaterialUsers[MaterialId] = Count + 1;
+			}
+			else
+			{
+				MaterialUsers.Add(MaterialId, 1);
+			}
+		}
+
+		private void RemoveMaterialUser(int MaterialId)
+		{
+			if (MaterialUsers.TryGetValue(MaterialId, out int Count))
+			{
+				if (Count == 1)
+				{
+					MaterialUsers.Remove(MaterialId);
+				}
+				else
+				{
+					MaterialUsers[MaterialId] = Count - 1;
+				}
+			}
+			else
+			{
+				Debug.Assert(false); // Bug
 			}
 		}
 
@@ -351,21 +400,24 @@ namespace DatasmithSolidworks
 			}
 		}
 
-		public bool ExportMesh(FMeshName InMeshName, FMeshData InData, out FDatasmithFacadeMeshElement OutMeshElement)
+		public bool ExportMesh(FMeshExportInfo Info)
 		{
-			string MeshNameInitial = InMeshName.GetString(); // Make Datasmith Mesh string to build Mesh/MeshElement from, Datasmith will modify it with FDatasmithUtils::SanitizeObjectName
+			FMeshName MeshName = Info.MeshName;
+			FMeshData MeshData = Info.MeshData;
 
-			OutMeshElement = null;
+			string MeshNameInitial = MeshName.GetString(); // Make Datasmith Mesh string to build Mesh/MeshElement from, Datasmith will modify it with FDatasmithUtils::SanitizeObjectName
 
-			if (InData.Vertices == null || InData.Normals == null || InData.TexCoords == null || InData.Triangles == null)
+			Info.MeshElement = null;
+
+			if (MeshData.Vertices == null || MeshData.Normals == null || MeshData.TexCoords == null || MeshData.Triangles == null)
 			{
-				Addin.Instance.LogDebugThread($"  skipping(some attributes are null) - Vertices:{InData.Vertices == null}, Normals:{InData.Normals == null}, TexCoords: {InData.TexCoords == null}, Triangles:{InData.Triangles}");
+				Addin.Instance.LogDebugThread($"  skipping(some attributes are null) - Vertices:{MeshData.Vertices == null}, Normals:{MeshData.Normals == null}, TexCoords: {MeshData.TexCoords == null}, Triangles:{MeshData.Triangles}");
 				return false;
 			}
 
-			if (InData.Vertices.Length == 0 || InData.Normals.Length == 0 || InData.TexCoords.Length == 0 || InData.Triangles.Length == 0)
+			if (MeshData.Vertices.Length == 0 || MeshData.Normals.Length == 0 || MeshData.TexCoords.Length == 0 || MeshData.Triangles.Length == 0)
 			{
-				Addin.Instance.LogDebugThread($"  skipping - Vertices: {InData.Vertices.Length}, Normals: {InData.Normals.Length}, TexCoords: {InData.TexCoords.Length} Triangles: {InData.Triangles.Length}");
+				Addin.Instance.LogDebugThread($"  skipping - Vertices: {MeshData.Vertices.Length}, Normals: {MeshData.Normals.Length}, TexCoords: {MeshData.TexCoords.Length} Triangles: {MeshData.Triangles.Length}");
 				return false;
 			}
 
@@ -375,26 +427,26 @@ namespace DatasmithSolidworks
 
 				FDatasmithFacadeMeshElement MeshElement = new FDatasmithFacadeMeshElement(MeshNameInitial);
 
-				Mesh.SetVerticesCount(InData.Vertices.Length);
-				Mesh.SetFacesCount(InData.Triangles.Length);
+				Mesh.SetVerticesCount(MeshData.Vertices.Length);
+				Mesh.SetFacesCount(MeshData.Triangles.Length);
 
-				for (int i = 0; i < InData.Vertices.Length; i++)
+				for (int i = 0; i < MeshData.Vertices.Length; i++)
 				{
-					Mesh.SetVertex(i, InData.Vertices[i].X, InData.Vertices[i].Y, InData.Vertices[i].Z);
+					Mesh.SetVertex(i, MeshData.Vertices[i].X, MeshData.Vertices[i].Y, MeshData.Vertices[i].Z);
 				}
 
-				for (int i = 0; i < InData.Normals.Length; i++)
+				for (int i = 0; i < MeshData.Normals.Length; i++)
 				{
-					Mesh.SetNormal(i, InData.Normals[i].X, InData.Normals[i].Y, InData.Normals[i].Z);
+					Mesh.SetNormal(i, MeshData.Normals[i].X, MeshData.Normals[i].Y, MeshData.Normals[i].Z);
 				}
 
-				if (InData.TexCoords != null)
+				if (MeshData.TexCoords != null)
 				{
 					Mesh.SetUVChannelsCount(1);
-					Mesh.SetUVCount(0, InData.TexCoords.Length);
-					for (int i = 0; i < InData.TexCoords.Length; i++)
+					Mesh.SetUVCount(0, MeshData.TexCoords.Length);
+					for (int i = 0; i < MeshData.TexCoords.Length; i++)
 					{
-						Mesh.SetUV(0, i, InData.TexCoords[i].X, InData.TexCoords[i].Y);
+						Mesh.SetUV(0, i, MeshData.TexCoords[i].X, MeshData.TexCoords[i].Y);
 					}
 
 				}
@@ -403,35 +455,20 @@ namespace DatasmithSolidworks
 
 				HashSet<int> MeshMaterialIds = new HashSet<int>();
 
-				for (int TriIndex = 0; TriIndex < InData.Triangles.Length; TriIndex++)
+				for (int TriIndex = 0; TriIndex < MeshData.Triangles.Length; TriIndex++)
 				{
-					FTriangle Triangle = InData.Triangles[TriIndex];
+					FTriangle Triangle = MeshData.Triangles[TriIndex];
 					int MatID = 0;
 
 					MeshMaterialIds.Add(Triangle.MaterialID);
-					if (Triangle.MaterialID >= 1)
+
+					if (Triangle.MaterialID > 0)
 					{
 						if (!MeshAddedMaterials.Contains(Triangle.MaterialID))
 						{
-							Addin.Instance.LogDebugThread($"  set material if it's exported: {Triangle.MaterialID}");
-
-							FDatasmithFacadeMaterialInstance Material = null;
-							ExportedMaterialsMap.TryGetValue(Triangle.MaterialID, out Material);
-
-							if (Material != null)
-							{
-								Addin.Instance.LogDebugThread(
-									$" SetMaterial({Material.GetName()}, SlotId: {Triangle.MaterialID})");
-
-								MeshAddedMaterials.Add(Triangle.MaterialID);
-								MeshElement.SetMaterial(Material.GetName(), Triangle.MaterialID);
-								MatID = Triangle.MaterialID;
-							}
-						}
-						else
-						{
-							MatID = Triangle.MaterialID;
+							MeshAddedMaterials.Add(Triangle.MaterialID);
 						}
+						MatID = Triangle.MaterialID;
 					}
 
 					Mesh.SetFace(TriIndex, Triangle[0], Triangle[1], Triangle[2], MatID);
@@ -445,7 +482,8 @@ namespace DatasmithSolidworks
 				}
 
 				DatasmithScene.ExportDatasmithMesh(MeshElement, Mesh);
-				OutMeshElement = MeshElement;
+				Info.MeshElement = MeshElement;
+				Info.Materials = MeshAddedMaterials;
 			}
 
 			return true;
@@ -525,7 +563,8 @@ namespace DatasmithSolidworks
 			return NewBinding;
 		}
 
-		private void ExportMaterialVariants(List<Tuple<FConfigurationData, FDatasmithFacadeVariant>> InVariants)
+		private void ExportMaterialVariants(List<Tuple<FConfigurationData, FDatasmithFacadeVariant>> InVariants,
+			Dictionary<FDatasmithFacadeActorBinding, int> MaterialBindings)
 		{
 			foreach (Tuple<FConfigurationData, FDatasmithFacadeVariant> KVP in InVariants)
 			{
@@ -574,10 +613,13 @@ namespace DatasmithSolidworks
 						foreach (var MatID in UniqueMaterialsSet)
 						{
 							FMaterial Material = Materials.GetMaterial(MatID);
-							FDatasmithFacadeMaterialInstance DatasmithMaterial = null;
-							if (Material != null && ExportedMaterialsMap.TryGetValue(Material.ID, out DatasmithMaterial))
+							if (Material != null)
 							{
-								Binding.AddMaterialCapture(DatasmithMaterial);
+								MaterialBindings[Binding] = Material.ID;
+
+								// todo: Sync for variants needs to track usage of materials between variants/configurations
+								// i.e. removal of material user too
+								AddMaterialUser(Material.ID);
 							}
 						}
 					}
@@ -623,7 +665,8 @@ namespace DatasmithSolidworks
 			}
 		}
 
-		public void ExportLevelVariantSets(List<FConfigurationData> InConfigs)
+		public void ExportLevelVariantSets(List<FConfigurationData> InConfigs,
+			Dictionary<FDatasmithFacadeActorBinding, int> MaterialBindings)
 		{
 			if (InConfigs == null)
 			{
@@ -704,7 +747,7 @@ namespace DatasmithSolidworks
 				// todo: visibility variants may drop ComponentName from FConfigurationData
 				// at this point only actor names may stay
 				ExportActorVisibilityVariants(ConfigurationVariants);
-				ExportMaterialVariants(ConfigurationVariants);
+				ExportMaterialVariants(ConfigurationVariants, MaterialBindings);
 				ExportTransformVariants(ConfigurationVariants);
 
 				// Geometry variants
@@ -727,7 +770,7 @@ namespace DatasmithSolidworks
 
 			if (DisplayStateVariants != null)
 			{
-				ExportMaterialVariants(DisplayStateVariants);
+				ExportMaterialVariants(DisplayStateVariants, MaterialBindings);
 			}
 		}
 
@@ -735,49 +778,77 @@ namespace DatasmithSolidworks
 		{
 			Addin.Instance.LogDebug($"ExportMaterials: \n  {string.Join("  \n", InMaterialsMap.Select(KVP => $"{KVP.Key}: {KVP.Value}" ))}");
 
-			ConcurrentBag<FDatasmithFacadeTexture> CreatedTextures = new ConcurrentBag<FDatasmithFacadeTexture>();
-			ConcurrentBag<FDatasmithFacadeMaterialInstance> CreatedMaterials = new ConcurrentBag<FDatasmithFacadeMaterialInstance>();
+			List<FDatasmithFacadeTexture> CreatedTextures = new List<FDatasmithFacadeTexture>();
+			List<FDatasmithFacadeMaterialInstance> CreatedMaterials = new List<FDatasmithFacadeMaterialInstance>();
+			List<FDatasmithFacadeMaterialInstance> UnusedMaterials = new List<FDatasmithFacadeMaterialInstance>();
+
 			foreach(KeyValuePair<int, FMaterial> MatKVP in InMaterialsMap)
 			{
-				List<FDatasmithFacadeTexture> NewMaterialTextures = null;
-				FDatasmithFacadeMaterialInstance NewMaterial = null;
-				if (CreateAndCacheMaterial(MatKVP.Value, out NewMaterialTextures, out NewMaterial))
-				{
-					CreatedMaterials.Add(NewMaterial);
+				int MaterialId = MatKVP.Key;
+				FMaterial Material = MatKVP.Value;
+
+				Addin.Instance.LogDebug($"  Material: {MaterialId}, {Material.Name}");
 
-					foreach (FDatasmithFacadeTexture Texture in NewMaterialTextures)
+				if (!MaterialUsers.ContainsKey(MaterialId))
+				{
+					Addin.Instance.LogDebug($"    skipping, no users");
+					if (ExportedMaterialInstancesMap.ContainsKey(MaterialId))
 					{
-						CreatedTextures.Add(Texture);
+						Addin.Instance.LogDebug($"      but should remove exported");
+						if (ExportedMaterialInstancesMap.TryRemove(MaterialId, out FDatasmithFacadeMaterialInstance MaterialInstance))
+						{
+							UnusedMaterials.Add(MaterialInstance);
+						}
 					}
+					continue;
+				}
+
+				if (ExportedMaterialInstancesMap.ContainsKey(MaterialId))
+				{
+					Addin.Instance.LogDebug($"    skipping, already exported");
+					continue;
 				}
+
+				ConvertMaterialToDatasmith(Material, out List<FDatasmithFacadeTexture> NewMaterialTextures, out FDatasmithFacadeMaterialInstance NewDatasmithMaterial);
+
+				ExportedMaterialInstancesMap[MaterialId] = NewDatasmithMaterial;
+
+				CreatedMaterials.Add(NewDatasmithMaterial);
+
+				foreach (FDatasmithFacadeTexture Texture in NewMaterialTextures)
+				{
+					CreatedTextures.Add(Texture);
+				}
+			}
+
+			Addin.Instance.LogDebug($"  Remove Materials");
+			foreach (FDatasmithFacadeMaterialInstance Mat in UnusedMaterials)
+			{
+				Addin.Instance.LogDebug($"    Material: {Mat.GetName()}: {Mat.GetLabel()}");
+				DatasmithScene.RemoveMaterial(Mat);
 			}
 
 			Addin.Instance.LogDebug($"  AddMaterials");
-			// Adding stuff to a datasmith scene cannot be multithreaded!
 			foreach (FDatasmithFacadeMaterialInstance Mat in CreatedMaterials)
 			{
-				Addin.Instance.LogDebug($"  AddMaterial: {Mat.GetName()}: {Mat.GetLabel()}");
+				Addin.Instance.LogDebug($"    Material: {Mat.GetName()}: {Mat.GetLabel()}");
 				DatasmithScene.AddMaterial(Mat);
 			}
+
 			foreach (FDatasmithFacadeTexture Texture in CreatedTextures)
 			{
 				DatasmithScene.AddTexture(Texture);
 			}
+
 		}
 
-		private bool CreateAndCacheMaterial(FMaterial InMaterial, out List<FDatasmithFacadeTexture> OutCreatedTextures, out FDatasmithFacadeMaterialInstance OutCreatedMaterial)
+		private void ConvertMaterialToDatasmith(FMaterial InMaterial, out List<FDatasmithFacadeTexture> OutCreatedTextures, out FDatasmithFacadeMaterialInstance OutCreatedMaterial)
 		{
-			Addin.Instance.LogDebug($"CreateAndCacheMaterial({InMaterial.Name}: {InMaterial.ID})");
+			Addin.Instance.LogDebug($"ConvertMaterialToDatasmith({InMaterial.Name}: {InMaterial.ID})");
 
 			OutCreatedTextures = null;
 			OutCreatedMaterial = null;
 
-			if (ExportedMaterialsMap.ContainsKey(InMaterial.ID))
-			{
-				Addin.Instance.LogDebug($"  skipping, already present in map");
-				return false;
-			}
-
 			Addin.Instance.LogDebug($"  making, ShaderName: '{InMaterial.ShaderName}'");
 
 			FMaterial.EMaterialType Type = FMaterial.GetMaterialType(InMaterial.ShaderName);
@@ -870,10 +941,6 @@ namespace DatasmithSolidworks
 					OutCreatedTextures.Add(NormalMap);
 				}
 			}
-
-			ExportedMaterialsMap[InMaterial.ID] = MaterialInstance;
-
-			return true;
 		}
 
 		public FActorName GetComponentActorName(FComponentName ComponentName)
@@ -1049,9 +1116,9 @@ namespace DatasmithSolidworks
 		//   The way to make element name ahead of constructing and element(e.g. to properly make unique names for every case) is to expose FDatasmithUtils::SanitizeObjectName
 		public string GetDatasmithMeshName(FMeshName MeshName)  
 		{
-			if (ExportedMeshesMap.TryGetValue(MeshName, out FDatasmithFacadeMeshElement MeshElement))
+			if (ExportedMeshesMap.TryGetValue(MeshName, out FMeshExportInfo Info))
 			{
-				return MeshElement.GetName();
+				return Info.MeshElement.GetName();
 			}
 
 			return null;
@@ -1066,7 +1133,7 @@ namespace DatasmithSolidworks
 
 			Parallel.ForEach(MeshExportInfos, Info =>
 			{
-				if (ExportMesh(Info.MeshName, Info.MeshData, out Info.MeshElement))
+				if (ExportMesh(Info))
 				{
 					CreatedMeshes.Enqueue(Info);
 				}
@@ -1076,8 +1143,8 @@ namespace DatasmithSolidworks
 			foreach (FMeshExportInfo Info in CreatedMeshes)
 			{
 				Addin.Instance.LogDebug($"  AddMesh(DatasmithMeshName: {Info}");
-				// Info.MeshElement should be valid here
-				AddMesh(Info.MeshName, Info.MeshElement);
+
+				AddMesh(Info);
 
 				if (Info.ActorName.IsValid())
 				{
@@ -1092,6 +1159,36 @@ namespace DatasmithSolidworks
 			}
 
 			OutCreatedMeshes = CreatedMeshes.ToList();
+
+		}
+
+		public void AssignMaterialsToDatasmithMeshes(List<FMeshExportInfo> MeshExportInfos)
+		{
+			Addin.Instance.LogDebug($"AssignMaterialsToDatasmithMeshes: {string.Join(", ", MeshExportInfos)}");
+
+			foreach (FMeshExportInfo Info in MeshExportInfos)
+			{
+				Addin.Instance.LogDebug($"  {Info}");
+
+				foreach (int MaterialId in Info.Materials)
+				{
+					Addin.Instance.LogDebug($"    Material: {MaterialId}");
+					if (GetDatasmithMaterial(MaterialId, out FDatasmithFacadeMaterialInstance Material))
+					{
+						Addin.Instance.LogDebug($"      SetMaterial({Material.GetName()}, SlotId: {MaterialId})");
+						Info.MeshElement.SetMaterial(Material.GetName(), MaterialId);
+					}
+					else
+					{
+						Addin.Instance.LogDebug($"      NOT FOUND among exported materials!");
+					}
+				}
+			}
+		}
+
+		public bool GetDatasmithMaterial(int MaterialId, out FDatasmithFacadeMaterialInstance Material)
+		{
+			return ExportedMaterialInstancesMap.TryGetValue(MaterialId, out Material);
 		}
 	}
 }
\ No newline at end of file
diff --git a/Engine/Source/Programs/Enterprise/Datasmith/DatasmithSolidworksExporter/Private/Document.cs b/Engine/Source/Programs/Enterprise/Datasmith/DatasmithSolidworksExporter/Private/Document.cs
--- a/Engine/Source/Programs/Enterprise/Datasmith/DatasmithSolidworksExporter/Private/Document.cs
+++ b/Engine/Source/Programs/Enterprise/Datasmith/DatasmithSolidworksExporter/Private/Document.cs
@@ -11,6 +11,7 @@ using System.Threading;
 using System.Threading.Tasks;
 using DatasmithSolidworks.Names;
 using static DatasmithSolidworks.FAssemblyDocumentTracker;
+using System.Windows.Forms;
 
 namespace DatasmithSolidworks
 {
@@ -30,10 +31,11 @@ namespace DatasmithSolidworks
 		{
 			DocumentTracker = InDocumentTracker;
 
-			MaterialCheckerThread = new Thread(CheckForMaterialUpdatesProc);
 			MaterialCheckerEvent = new ManualResetEvent(false);
-			MaterialCheckerThread.Start();
 			MaterialCheckerEvent.Set();
+
+			MaterialCheckerThread = new Thread(CheckForMaterialUpdatesProc);
+			MaterialCheckerThread.Start();
 		}
 
 		public void Destroy()
@@ -165,14 +167,34 @@ namespace DatasmithSolidworks
 				}
 			}
 
-			Exporter.ExportMaterials(ExportedMaterialsMap);
+			Dictionary<FDatasmithFacadeActorBinding, int> MaterialBindings = null;
 
 			if (Configs != null)
 			{
-				Exporter.ExportLevelVariantSets(Configs);
+				MaterialBindings = new Dictionary<FDatasmithFacadeActorBinding, int>();
+				Exporter.ExportLevelVariantSets(Configs, MaterialBindings);
+			}
+
+			ExportMaterials();
+
+			// Assign exported datasmith material instances to bindings
+			foreach (KeyValuePair<FDatasmithFacadeActorBinding, int> KVP in MaterialBindings)
+			{
+				FDatasmithFacadeActorBinding Binding = KVP.Key;
+				int MaterialId = KVP.Value;
+
+				if (Exporter.GetDatasmithMaterial(MaterialId, out FDatasmithFacadeMaterialInstance DatasmithMaterial))
+				{
+					Binding.AddMaterialCapture(DatasmithMaterial);
+				}					
 			}
 		}
 
+		public void ExportMaterials()
+		{
+			Exporter.ExportMaterials(ExportedMaterialsMap);
+		}
+
 		private void ExportLights()
 		{
 			List<FLight> Lights  = FLightExporter.ExportLights(SwDoc);
@@ -277,7 +299,7 @@ namespace DatasmithSolidworks
 		// Extracts, exports meshes used for the assembly configuration, assigns them to actors
 		// todo: separate just Datasmith Mesh export to its own thread, all other code should be in single thread
 		// Datasmith actor assignment doesn't need threading(and not supposed to be thread-safe?) and Solidworks multithreading is supposed to be slower(SW does all the work in main thread)
-		public void ProcessConfigurationMeshes(FMeshes.FConfiguration MeshesConfiguration, string MeshSuffix)
+		public List<FDatasmithExporter.FMeshExportInfo> ProcessConfigurationMeshes(FMeshes.FConfiguration MeshesConfiguration, string MeshSuffix)
 		{
 			
 			List<FDatasmithExporter.FMeshExportInfo> MeshExportInfos = new List<FDatasmithExporter.FMeshExportInfo>();
@@ -311,6 +333,13 @@ namespace DatasmithSolidworks
 			{
 				AddMeshForComponent(Info.ComponentName, Info.MeshName);  // Register that this mesh was used for the component
 			}
+
+			return CreatedMeshes;
+		}
+
+		public void AssignMaterialsToDatasmithMeshes(List<FDatasmithExporter.FMeshExportInfo> CreatedMeshes)
+		{
+			Exporter.AssignMaterialsToDatasmithMeshes(CreatedMeshes);
 		}
 	};
 
diff --git a/Engine/Source/Programs/Enterprise/Datasmith/DatasmithSolidworksExporter/Private/PartDocument.cs b/Engine/Source/Programs/Enterprise/Datasmith/DatasmithSolidworksExporter/Private/PartDocument.cs
--- a/Engine/Source/Programs/Enterprise/Datasmith/DatasmithSolidworksExporter/Private/PartDocument.cs
+++ b/Engine/Source/Programs/Enterprise/Datasmith/DatasmithSolidworksExporter/Private/PartDocument.cs
@@ -114,8 +114,6 @@ namespace DatasmithSolidworks
 
 			SetExportStatus($"{PartName} Materials");
 
-			Exporter.ExportMaterials(ExportedMaterialsMap);
-
 			SetExportStatus($"{PartName} Meshes");
 
 			ConfigurationManager ConfigManager = (SwPartDoc as ModelDoc2).ConfigurationManager;
@@ -154,6 +152,8 @@ namespace DatasmithSolidworks
 					};
 
 					Exporter.ExportMeshes(MeshExportInfos, out List<FMeshExportInfo> OutCreatedMeshes);
+					ExportMaterials();
+					Exporter.AssignMaterialsToDatasmithMeshes(OutCreatedMeshes);
 					bool bHasMesh = OutCreatedMeshes.Count > 0;
 
 					FDatasmithActorExportInfo ExportInfo = new FDatasmithActorExportInfo();
@@ -195,6 +195,8 @@ namespace DatasmithSolidworks
 					};
 
 					Exporter.ExportMeshes(MeshExportInfos, out List<FMeshExportInfo> OutCreatedMeshes);
+					ExportMaterials();
+					Exporter.AssignMaterialsToDatasmithMeshes(OutCreatedMeshes);
 					bHasMesh = OutCreatedMeshes.Count > 0;
 				}
 
@@ -206,6 +208,8 @@ namespace DatasmithSolidworks
 				ExportInfo.Type = bHasMesh ? EActorType.MeshActor : EActorType.SimpleActor;
 				Exporter.ExportOrUpdateActor(ExportInfo);
 			}
+
+			
 		}
 
 		public override bool HasMaterialUpdates()
