[Solidworks] Fixed Datasmith plugin preventing proper shutdown of Solidworks

- repro: run SW, make Datasmith addin enabled, open a model and Sync it, then exit SW. Expected: Solidwokrs process exits(check in TaskManager as it can have its window destroyed but not process itself).

diff --git a/Engine/Source/Programs/Enterprise/Datasmith/DatasmithSolidworksExporter/Private/Addin.cs b/Engine/Source/Programs/Enterprise/Datasmith/DatasmithSolidworksExporter/Private/Addin.cs
--- a/Engine/Source/Programs/Enterprise/Datasmith/DatasmithSolidworksExporter/Private/Addin.cs
+++ b/Engine/Source/Programs/Enterprise/Datasmith/DatasmithSolidworksExporter/Private/Addin.cs
@@ -202,6 +202,13 @@ namespace DatasmithSolidworks
 			return true;
 		}
 
+		private static int OnAppDestroy()
+		{
+			// Calling this only on App exit(see comment below)
+			FDatasmithFacadeDirectLink.Shutdown();
+			return 0;
+		}
+
 		public bool DisconnectFromSW()
 		{
 			// Disabled Shutdown as Initing again crashes if the plugin is re-enabled in SW
@@ -325,6 +332,7 @@ namespace DatasmithSolidworks
 				AppEvents.FileCloseNotify += new DSldWorksEvents_FileCloseNotifyEventHandler(OnFileClose);
 				AppEvents.CommandCloseNotify += new DSldWorksEvents_CommandCloseNotifyEventHandler(OnCommandClose);
 				AppEvents.OnIdleNotify += new DSldWorksEvents_OnIdleNotifyEventHandler(OnIdle);
+				AppEvents.DestroyNotify += OnAppDestroy;
 			}
 			catch {}
 		}
@@ -340,6 +348,8 @@ namespace DatasmithSolidworks
 				AppEvents.FileCloseNotify -= new DSldWorksEvents_FileCloseNotifyEventHandler(OnFileClose);
 				AppEvents.CommandCloseNotify -= new DSldWorksEvents_CommandCloseNotifyEventHandler(OnCommandClose);
 				AppEvents.OnIdleNotify -= new DSldWorksEvents_OnIdleNotifyEventHandler(OnIdle);
+				// AppEvents.DestroyNotify - don't teach destroy handler as Detach is called in DisconnectFromSW(which also executed on addin disable)
+				// and we need this handler to shutdown DirectLink(with all the Unreal engine) that we can do only once
 			}
 			catch {}
 		}
