[Solidworks] Fixed crash exporting empty Part component

diff --git a/Engine/Source/Programs/Enterprise/Datasmith/DatasmithSolidworksExporter/Private/ConfigurationExporter.cs b/Engine/Source/Programs/Enterprise/Datasmith/DatasmithSolidworksExporter/Private/ConfigurationExporter.cs
--- a/Engine/Source/Programs/Enterprise/Datasmith/DatasmithSolidworksExporter/Private/ConfigurationExporter.cs
+++ b/Engine/Source/Programs/Enterprise/Datasmith/DatasmithSolidworksExporter/Private/ConfigurationExporter.cs
@@ -99,6 +99,11 @@ namespace DatasmithSolidworks
 					return false;
 				}
 
+				if (ReferenceEquals(MeshData, Other.MeshData))
+				{
+					return true;
+				}
+				
 				if (!MeshData.Vertices.SequenceEqual(Other.MeshData.Vertices))
 				{
 					return false;
