[3ds Max Exporter] Fixed public headers include to fix compilation

diff --git a/Engine/Source/Programs/Enterprise/Datasmith/DatasmithMaxExporter/DatasmithMax2017.Build.cs b/Engine/Source/Programs/Enterprise/Datasmith/DatasmithMaxExporter/DatasmithMax2017.Build.cs
--- a/Engine/Source/Programs/Enterprise/Datasmith/DatasmithMaxExporter/DatasmithMax2017.Build.cs
+++ b/Engine/Source/Programs/Enterprise/Datasmith/DatasmithMaxExporter/DatasmithMax2017.Build.cs
@@ -80,6 +80,8 @@ namespace UnrealBuildTool.Rules
 			{
 				PrivateIncludePaths.Add(ItooInterfaceLocation);
 			}
+
+			PrivateIncludePaths.Add(ModuleDirectory);
 		}
 
 		public abstract string GetMaxVersion();
diff --git a/Engine/Source/Programs/Enterprise/Datasmith/DatasmithMaxExporter/Private/UI/Datasmith3dsMaxUI.cpp b/Engine/Source/Programs/Enterprise/Datasmith/DatasmithMaxExporter/Private/UI/Datasmith3dsMaxUI.cpp
--- a/Engine/Source/Programs/Enterprise/Datasmith/DatasmithMaxExporter/Private/UI/Datasmith3dsMaxUI.cpp
+++ b/Engine/Source/Programs/Enterprise/Datasmith/DatasmithMaxExporter/Private/UI/Datasmith3dsMaxUI.cpp
@@ -7,14 +7,16 @@
 #include "Containers/Queue.h"
 #include "CoreGlobals.h"
 #include "Delegates/Delegate.h"
-#include "Framework/Application/SWindowTitleBar.h"
-#include "Framework/Application/SlateApplication.h"
 #include "HAL/Platform.h"
 #include "Misc/ConfigCacheIni.h"
 #include "Misc/ScopeLock.h"
-#include "SlateCore/Public/Widgets/DeclarativeSyntaxSupport.h"
-#include "Styling/CoreStyle.h"
 #include "Templates/SharedPointer.h"
+
+// Slate
+#include "Framework/Application/SWindowTitleBar.h"
+#include "Framework/Application/SlateApplication.h"
+#include "Styling/CoreStyle.h"
+#include "Widgets/DeclarativeSyntaxSupport.h"
 #include "Widgets/Input/SButton.h"
 #include "Widgets/Layout/SBorder.h"
 #include "Widgets/Layout/SScrollBox.h"
