Due to a crash, bProcessIsRunning is never set to false, so the parallel tasks "CheckDuration" and "CheckMemory" never end.
To fix it, bProcessIsRunning become static so it can be set to false without any instance of FDatasmithCADWorkerImpl

diff --git a/Engine/Source/Programs/Enterprise/Datasmith/DatasmithCADWorker/Private/DatasmithCADWorker.cpp b/Engine/Source/Programs/Enterprise/Datasmith/DatasmithCADWorker/Private/DatasmithCADWorker.cpp
--- a/Engine/Source/Programs/Enterprise/Datasmith/DatasmithCADWorker/Private/DatasmithCADWorker.cpp
+++ b/Engine/Source/Programs/Enterprise/Datasmith/DatasmithCADWorker/Private/DatasmithCADWorker.cpp
@@ -88,6 +88,8 @@ int32 Main(int32 Argc, TCHAR * Argv[])
 	FDatasmithCADWorkerImpl Worker(FCString::Atoi(*ServerPID), FCString::Atoi(*ServerPort), EnginePluginsPath, CacheDirectory);
 	Worker.Run();
 
+	FDatasmithCADWorkerImpl::bProcessIsRunning = false;
+
 	return EXIT_SUCCESS;
 #endif // USE_TECHSOFT_SDK
 }
@@ -111,6 +113,7 @@ INT32_MAIN_INT32_ARGC_TCHAR_ARGV()
 	}
 	__except (Filter(GetExceptionCode(), GetExceptionInformation()))
 	{
+		FDatasmithCADWorkerImpl::bProcessIsRunning = false;
 		Result = EXIT_FAILURE;
 	}
 
diff --git a/Engine/Source/Programs/Enterprise/Datasmith/DatasmithCADWorker/Private/DatasmithCADWorkerImpl.cpp b/Engine/Source/Programs/Enterprise/Datasmith/DatasmithCADWorker/Private/DatasmithCADWorkerImpl.cpp
--- a/Engine/Source/Programs/Enterprise/Datasmith/DatasmithCADWorker/Private/DatasmithCADWorkerImpl.cpp
+++ b/Engine/Source/Programs/Enterprise/Datasmith/DatasmithCADWorker/Private/DatasmithCADWorkerImpl.cpp
@@ -18,6 +18,8 @@
 
 using namespace DatasmithDispatcher;
 
+std::atomic<bool> FDatasmithCADWorkerImpl::bProcessIsRunning = false;
+std::atomic<bool> FDatasmithCADWorkerImpl::bRequestRestart = false;
 
 FDatasmithCADWorkerImpl::FDatasmithCADWorkerImpl(int32 InServerPID, int32 InServerPort, const FString& InEnginePluginsPath, const FString& InCachePath)
 	: ServerPID(InServerPID)
@@ -25,7 +27,6 @@ FDatasmithCADWorkerImpl::FDatasmithCADWorkerImpl(int32 InServerPID, int32 InServ
 	, EnginePluginsPath(InEnginePluginsPath)
 	, CachePath(InCachePath)
 	, PingStartCycle(0)
-	, bRequestRestart(false)
 {
 }
 
@@ -162,8 +163,8 @@ void FDatasmithCADWorkerImpl::ProcessCommand(const FRunTaskCommand& RunTaskComma
 	bProcessIsRunning = true;
 	int64 MaxDuration = DefineMaximumAllowedDuration(FileToProcess);
 	TArray<UE::Tasks::FTask> Checkers;
-	Checkers.Emplace(UE::Tasks::Launch(TEXT("TimeChecker"), [&]() { CheckDuration(FileToProcess, MaxDuration); }));
-	Checkers.Emplace(UE::Tasks::Launch(TEXT("MemoryChecker"), [&]() { CheckMemory(); }));
+	Checkers.Emplace(UE::Tasks::Launch(TEXT("TimeChecker"), [&FileToProcess, &MaxDuration]() { CheckDuration(FileToProcess, MaxDuration); }));
+	Checkers.Emplace(UE::Tasks::Launch(TEXT("MemoryChecker"), []() { CheckMemory(); }));
 
 	CADLibrary::FCADFileReader FileReader(ImportParameters, FileToProcess, EnginePluginsPath, CachePath);
 	CompletedTask.ProcessResult = FileReader.ProcessFile();
@@ -212,7 +213,7 @@ void FDatasmithCADWorkerImpl::ProcessCommand(const FRunTaskCommand& RunTaskComma
 
 void FDatasmithCADWorkerImpl::CheckDuration(const CADLibrary::FFileDescriptor& FileToProcess, const int64 MaxDuration)
 {
-	if (!ImportParameters.bGEnableTimeControl)
+	if (!CADLibrary::FImportParameters::bGEnableTimeControl)
 	{
 		return;
 	}
diff --git a/Engine/Source/Programs/Enterprise/Datasmith/DatasmithCADWorker/Private/DatasmithCADWorkerImpl.h b/Engine/Source/Programs/Enterprise/Datasmith/DatasmithCADWorker/Private/DatasmithCADWorkerImpl.h
--- a/Engine/Source/Programs/Enterprise/Datasmith/DatasmithCADWorker/Private/DatasmithCADWorkerImpl.h
+++ b/Engine/Source/Programs/Enterprise/Datasmith/DatasmithCADWorker/Private/DatasmithCADWorkerImpl.h
@@ -31,8 +31,8 @@ private:
 	/**
 	 * @param MaxDuration is defined in seconds.
 	 */
-	void CheckDuration(const CADLibrary::FFileDescriptor& FileToProcess, const int64 MaxDuration); 
-	void CheckMemory();
+	static void CheckDuration(const CADLibrary::FFileDescriptor& FileToProcess, const int64 MaxDuration);
+	static void CheckMemory();
 
 private:
 	DatasmithDispatcher::FNetworkClientNode NetworkInterface;
@@ -45,6 +45,7 @@ private:
 	CADLibrary::FImportParameters ImportParameters;
 	uint64 PingStartCycle;
 
-	std::atomic<bool> bProcessIsRunning;
-	std::atomic<bool> bRequestRestart;
+public:
+	static std::atomic<bool> bProcessIsRunning;
+	static std::atomic<bool> bRequestRestart;
 };
