[SketchUp] Add support for SDK 2023

- compiles with SDK 2023
- Windows installer includes 2023
- Mac OS compiles and zips

diff --git a/Engine/Source/Programs/Enterprise/Datasmith/DatasmithSketchUpRubyExporter/Build/build_mac.sh b/Engine/Source/Programs/Enterprise/Datasmith/DatasmithSketchUpRubyExporter/Build/build_mac.sh
--- a/Engine/Source/Programs/Enterprise/Datasmith/DatasmithSketchUpRubyExporter/Build/build_mac.sh
+++ b/Engine/Source/Programs/Enterprise/Datasmith/DatasmithSketchUpRubyExporter/Build/build_mac.sh
@@ -92,6 +92,8 @@ BuildSketchUpPlugin 2019 SDK_Mac_2019-3-252 "x86_64" "$Intermediatex86_64Path"
 BuildSketchUpPlugin 2020 SDK_Mac_2020-2-171 "x86_64" "$Intermediatex86_64Path"
 BuildSketchUpPlugin 2021 SDK_Mac_2021-0-338 "x86_64" "$Intermediatex86_64Path"
 BuildSketchUpPlugin 2022 SDK_Mac_2022-0-353 "x86_64 arm64" "$IntermediateUniversalPath"
+BuildSketchUpPlugin 2023 SDK_Mac_2023-0-366 "x86_64 arm64" "$IntermediateUniversalPath"
+
 
 # install_name_tool -change @rpath/DatasmithSDK.dylib @loader_path/Dylibs/DatasmithSDK.dylib DatasmithSketchUp.bundle 
 # install_name_tool -change @loader_path/libfreeimage-3.18.0.dylib @loader_path/Dylibs/libfreeimage-3.18.0.dylib DatasmithSketchUp.bundle
diff --git a/Engine/Source/Programs/Enterprise/Datasmith/DatasmithSketchUpRubyExporter/DatasmithSketchUpRuby2023.Build.cs b/Engine/Source/Programs/Enterprise/Datasmith/DatasmithSketchUpRubyExporter/DatasmithSketchUpRuby2023.Build.cs
new file mode 100644
--- /dev/null
+++ b/Engine/Source/Programs/Enterprise/Datasmith/DatasmithSketchUpRubyExporter/DatasmithSketchUpRuby2023.Build.cs
@@ -0,0 +1,31 @@
+// Copyright Epic Games, Inc. All Rights Reserved.
+
+using System.IO;
+
+namespace UnrealBuildTool.Rules
+{
+
+	[SupportedPlatforms("Win64", "Mac")]
+	public class DatasmithSketchUpRuby2023 : DatasmithSketchUpRubyBase
+	{
+		public DatasmithSketchUpRuby2023(ReadOnlyTargetRules Target)
+			: base(Target)
+		{
+			PrivateDefinitions.Add("SKP_SDK_2023");
+		}
+
+		public override string GetSketchUpSDKFolder()
+		{
+			return "SDK_WIN_x64_2023-0-367";
+		}
+
+		public override string GetSketchUpEnvVar()
+		{
+			return "SKP_SDK_2023";
+		}
+		public override string GetRubyLibName()
+		{
+			return "x64-msvcrt-ruby270.lib";
+		}
+	}
+}
diff --git a/Engine/Source/Programs/Enterprise/Datasmith/DatasmithSketchUpRubyExporter/DatasmithSketchUpRuby2023.Target.cs b/Engine/Source/Programs/Enterprise/Datasmith/DatasmithSketchUpRubyExporter/DatasmithSketchUpRuby2023.Target.cs
new file mode 100644
--- /dev/null
+++ b/Engine/Source/Programs/Enterprise/Datasmith/DatasmithSketchUpRubyExporter/DatasmithSketchUpRuby2023.Target.cs
@@ -0,0 +1,16 @@
+// Copyright Epic Games, Inc. All Rights Reserved.
+
+using UnrealBuildTool;
+
+[SupportedPlatforms("Win64")]
+public class DatasmithSketchUpRuby2023Target : DatasmithSketchUpRubyBaseTarget
+{
+	public DatasmithSketchUpRuby2023Target(TargetInfo Target)
+		: base(Target)
+	{
+		LaunchModuleName = "DatasmithSketchUpRuby2023";
+		ExeBinariesSubFolder = @"SketchUpRuby/2023";
+
+		AddCopyPostBuildStep(Target);
+	}
+}
