[Solidworks] Instancing for same part meshes

- identical meshes exported as the same Datasmith static mesh

diff --git a/Engine/Source/Programs/Enterprise/Datasmith/DatasmithSolidworksExporter/Private/Addin.cs b/Engine/Source/Programs/Enterprise/Datasmith/DatasmithSolidworksExporter/Private/Addin.cs
--- a/Engine/Source/Programs/Enterprise/Datasmith/DatasmithSolidworksExporter/Private/Addin.cs
+++ b/Engine/Source/Programs/Enterprise/Datasmith/DatasmithSolidworksExporter/Private/Addin.cs
@@ -627,29 +627,41 @@ namespace DatasmithSolidworks
 		}
 
 		[Conditional("DatasmithSolidworksDebugOutput")]
-		public void LogDebug(string Message)  
+		public static void LogDebug(string Message)
 		{
-			DebugLog.LogDebug(Message);
+			Instance.DebugLog.LogDebug(Message);
 		}
 
 		[Conditional("DatasmithSolidworksDebugOutput")]
-		public void LogDebugThread(string Message)  
+		public static void LogDebugThread(string Message)  
 		{
-			DebugLog.LogDebugThread(Message);
+			Instance.DebugLog.LogDebugThread(Message);
 		}
 
 		[Conditional("DatasmithSolidworksDebugOutput")]
-		public void LogIndent()  
+		public static void LogIndent()  
 		{
-			DebugLog.Indent();
+			Instance.DebugLog.Indent();
 		}
 
 		[Conditional("DatasmithSolidworksDebugOutput")]
-		public void LogDedent()  
+		public static void LogDedent()  
 		{
-			DebugLog.Dedent();
+			Instance.DebugLog.Dedent();
 		}
 
 		#endregion
+
+		public static V FindOrAdd<K, V>(Dictionary<K, V> Map, K Key) 
+			where V : new()
+		{
+			if (!Map.TryGetValue(Key, out V Value))
+			{
+				Value = new V();
+				Map.Add(Key, Value);
+			}
+			return Value;
+		}
+
 	}
 }
diff --git a/Engine/Source/Programs/Enterprise/Datasmith/DatasmithSolidworksExporter/Private/AssemblyDocument.cs b/Engine/Source/Programs/Enterprise/Datasmith/DatasmithSolidworksExporter/Private/AssemblyDocument.cs
--- a/Engine/Source/Programs/Enterprise/Datasmith/DatasmithSolidworksExporter/Private/AssemblyDocument.cs
+++ b/Engine/Source/Programs/Enterprise/Datasmith/DatasmithSolidworksExporter/Private/AssemblyDocument.cs
@@ -8,6 +8,7 @@ using System.Collections.Generic;
 using System.Diagnostics;
 using System.Linq;
 using DatasmithSolidworks.Names;
+using static DatasmithSolidworks.Addin;
 
 namespace DatasmithSolidworks
 {
@@ -46,8 +47,12 @@ namespace DatasmithSolidworks
 				new Dictionary<FComponentName, Component2>();
 
 			/** Stores which mesh was exported for each component*/
-			public Dictionary<FComponentName, FMeshName> ComponentNameToMeshNameMap =
-				new Dictionary<FComponentName, FMeshName>();
+			public Dictionary<FComponentName, List<FMeshName>> ComponentNameToMeshNameMap =
+				new Dictionary<FComponentName, List<FMeshName>>();
+			public Dictionary<FMeshName, HashSet<FComponentName>> ComponentsForMesh =
+				new Dictionary<FMeshName, HashSet<FComponentName>>();
+
+			public FMeshes Meshes = null;
 		}
 
 		public AssemblyDoc SwAsmDoc { get; private set; } = null;
@@ -70,8 +75,9 @@ namespace DatasmithSolidworks
 			base.Destroy();
 		}
 
-		public override void PreExport(FMeshes Meshes, bool bConfigurations)
+		public override FMeshes GetMeshes(string ActiveConfigName)
 		{
+			return SyncState.Meshes ?? (SyncState.Meshes = new FMeshes(ActiveConfigName));
 		}
 
 		// Export to datasmith scene extracted configurations data
@@ -89,6 +95,8 @@ namespace DatasmithSolidworks
 				ConfigurationExporter.CombinedTree.Children[0], null);
 			SyncState.DirtyComponents.Clear();
 
+			ConfigurationExporter.FinalizeExport(this);
+
 			// todo: check animation export for DL
 			// Export animations (only allow when exporting to file)
 			{
@@ -114,11 +122,7 @@ namespace DatasmithSolidworks
 				SyncState.ComponentToPartMap.Remove(CompName);
 				SyncState.ExportedComponentsMap.Remove(CompName);
 
-				if (SyncState.ComponentNameToMeshNameMap.TryGetValue(CompName, out FMeshName MeshName))
-				{
-					SyncState.ComponentNameToMeshNameMap.Remove(CompName);
-					Exporter.RemoveMesh(MeshName);
-				}
+				ReleaseComponentMeshes(CompName);
 
 				SyncState.ComponentsMaterialsMap?.TryRemove(CompName, out FObjectMaterials _);
 				SyncState.ComponentsTransformsMap.Remove(CompName);
@@ -128,6 +132,39 @@ namespace DatasmithSolidworks
 			SyncState.ComponentsToDelete.Clear();
 		}
 
+		public override void ReleaseComponentMeshes(FComponentName CompName)
+		{
+			if (SyncState.ComponentNameToMeshNameMap.TryGetValue(CompName, out List<FMeshName> MeshNames))
+			{
+				SyncState.ComponentNameToMeshNameMap.Remove(CompName);
+
+				foreach (FMeshName MeshName in MeshNames)
+				{
+					HashSet<FComponentName> ComponentNames = SyncState.ComponentsForMesh[MeshName];
+					ComponentNames.Remove(CompName);
+					if (ComponentNames.Count == 0)
+					{
+						SyncState.ComponentsForMesh.Remove(MeshName);
+					}
+				}
+			}
+		}
+		public override void CleanupComponentMeshes()
+		{
+			// Filter meshes without components using them
+			List<FMeshName> MeshesToRemove = (SyncState.ComponentsForMesh.Where(KVP => KVP.Value.Count == 0).Select(KVP => KVP.Key)).ToList();
+
+			LogDebug("Removing unused meshes:");
+			foreach (FMeshName MeshName in MeshesToRemove)
+			{
+				LogIndent();
+				Exporter.RemoveMesh(MeshName);
+				LogDedent();
+
+				SyncState.ComponentsForMesh.Remove(MeshName);
+			}
+		}
+
 		public override ConcurrentDictionary<FComponentName, FObjectMaterials> LoadDocumentMaterials(
 			HashSet<FComponentName> ComponentNamesToExportSet)
 		{
@@ -137,7 +174,7 @@ namespace DatasmithSolidworks
 
 		public override void AddComponentMaterials(FComponentName ComponentName, FObjectMaterials Materials)
 		{
-			Addin.Instance.LogDebug($"AddComponentMaterials: {ComponentName} - {Materials}");
+			LogDebug($"AddComponentMaterials: {ComponentName} - {Materials}");
 
 			if (SyncState.ComponentsMaterialsMap == null)
 			{
@@ -148,7 +185,8 @@ namespace DatasmithSolidworks
 
 		public override void AddMeshForComponent(FComponentName ComponentName, FMeshName MeshName)
 		{
-			SyncState.ComponentNameToMeshNameMap[ComponentName] = MeshName;
+			FindOrAdd(SyncState.ComponentNameToMeshNameMap, ComponentName).Add(MeshName);
+			FindOrAdd(SyncState.ComponentsForMesh, MeshName).Add(ComponentName);
 		}
 
 		public override FObjectMaterials GetComponentMaterials(Component2 Comp)
@@ -158,71 +196,6 @@ namespace DatasmithSolidworks
 			return ComponentMaterials;
 		}
 
-		// Export to datasmith scene current configuration only
-		public override void ExportToDatasmithScene(FMeshes Meshes)
-		{
-			SetExportStatus("Actors");
-			ProcessComponentsPendingDelete();
-
-			Configuration CurrentConfig = SwDoc.GetActiveConfiguration() as Configuration;
-
-			SetExportStatus("");
-			Component2 Root = CurrentConfig.GetRootComponent3(true);
-
-			// Track components that need their mesh exported: we want to do that in parallel after the 
-			// actor hierarchy has been exported
-			FMeshes.FConfiguration MeshesConfiguration = Meshes.GetMeshesConfiguration(ConfigurationName: CurrentConfig.Name);
-
-			ExportComponentRecursive(Root, null, MeshesConfiguration);
-
-			// Export materials
-			SetExportStatus($"Component Materials");
-
-			HashSet<FComponentName> ComponentNamesToExportSet = new HashSet<FComponentName>();
-			foreach (FComponentName ComponentName in MeshesConfiguration.EnumerateComponentNames())
-			{
-				if (!ComponentNamesToExportSet.Contains(ComponentName))
-				{
-					ComponentNamesToExportSet.Add(ComponentName);
-				}
-			}
-
-			ConcurrentDictionary<FComponentName, FObjectMaterials> ModifiedComponentsMaterials =
-				FObjectMaterials.LoadAssemblyMaterials(this, ComponentNamesToExportSet,
-					swDisplayStateOpts_e.swThisDisplayState, null);
-
-			if (ModifiedComponentsMaterials != null)
-			{
-				foreach (var MatKVP in ModifiedComponentsMaterials)
-				{
-					AddComponentMaterials(MatKVP.Key, MatKVP.Value);
-				}
-			}
-
-			// Export meshes
-			SetExportStatus($"Component Meshes");
-			List<FDatasmithExporter.FMeshExportInfo> CreatedMeshes = ProcessConfigurationMeshes(MeshesConfiguration, null);
-			ExportMaterials();
-			AssignMaterialsToDatasmithMeshes(CreatedMeshes);
-
-			SyncState.DirtyComponents.Clear();
-
-			// Export animations (only allow when exporting to file)
-			//	if (bFileExportInProgress)
-			{
-				List<FAnimation> Animations = FAnimationExtractor.ExtractAnimations(SwAsmDoc, Root);
-				if (Animations != null)
-				{
-					SetExportStatus($"Animations");
-
-					foreach (FAnimation Anim in Animations)
-					{
-						Exporter.ExportAnimation(Anim);
-					}
-				}
-			}
-		}
-
 		public override FMeshData ExtractComponentMeshData(Component2 Comp)
 		{
 			FObjectMaterials ComponentMaterials = GetComponentMaterials(Comp);
@@ -249,7 +222,7 @@ namespace DatasmithSolidworks
 			// Components which stayed in the materials map
 			HashSet<FComponentName> CommonComponents = new HashSet<FComponentName>(CurrentComponents.Intersect(Components));
 
-			// Check components that have material added or removed
+			 
 			bool bHasDirtyMaterials = false;
 			IEnumerable<FComponentName> ComponentsWithAddedOrRemovedMaterial = Components.Union(CurrentComponents).Except(CommonComponents);
 			foreach (FComponentName CompName in ComponentsWithAddedOrRemovedMaterial)
@@ -267,7 +240,12 @@ namespace DatasmithSolidworks
 					}
 					catch
 					{
-						// todo: needs a comment on how this happens
+						// todo: exceptions here(and other places loading materials) thrown because of multi-threaded FMaterialUpdateChecker
+						// calling this code from its thread  performing COM calls into SW API 
+						// calling the API from any other thread than main is extremely slow so update checker is running forever
+						// meaning it has increased probability to have model update while it's running and have a race condition of any kind
+						// there's UE-177970 fot this
+						Debug.Assert(false);
 					}
 				}
 
@@ -288,7 +266,6 @@ namespace DatasmithSolidworks
 					bHasDirtyComponents = true;
 				}
 			}
-
 			return bHasDirtyComponents || bHasDirtyMaterials;
 		}
 
@@ -320,126 +297,6 @@ namespace DatasmithSolidworks
 			return ComponentConfig?.Transform ?? InNode.CommonConfig.Transform ?? MathUtils.TransformIdentity();
 		}
 
-		private void ExportComponentRecursive(Component2 InComponent, Component2 InParent,
-			FMeshes.FConfiguration OutMeshesToExportMap)
-		{
-			bool bHasDirtyTransform = false;
-
-			FComponentName ComponentName = new FComponentName(InComponent);
-			if (SyncState.ComponentsTransformsMap.ContainsKey(ComponentName))
-			{
-				float[] ComponentTm = GetComponentDatasmithTransform(InComponent);
-				bHasDirtyTransform =
-					!MathUtils.TransformsAreEqual(SyncState.ComponentsTransformsMap[ComponentName], ComponentTm);
-			}
-
-			if (bHasDirtyTransform || !SyncState.CleanComponents.Contains(ComponentName))
-			{
-				SetExportStatus(ComponentName.GetString());
-
-				FDatasmithActorExportInfo ActorExportInfo = new FDatasmithActorExportInfo();
-
-				FActorName ActorName = Exporter.GetComponentActorName(ComponentName);
-
-				ActorExportInfo.Label = ComponentName.GetLabel();
-				ActorExportInfo.Name = ActorName;
-
-				if (InParent != null)
-				{
-					ActorExportInfo.ParentName = Exporter.GetComponentActorName(InParent);
-				}
-
-				ActorExportInfo.bVisible = true;
-				ActorExportInfo.Type = Exporter.GetExportedActorType(ActorName) ?? EActorType.SimpleActor;
-				ActorExportInfo.Transform = GetComponentDatasmithTransform(InComponent);
-
-				SyncState.ComponentsTransformsMap[ComponentName] = ActorExportInfo.Transform;
-
-				// When exporting configurations we need to load all components to make their data available for export
-				// Suppressed or lightweight components don't have this
-
-				if (!InComponent.IsSuppressed())
-				{
-					dynamic ComponentVisibility =
-						InComponent.GetVisibilityInAsmDisplayStates((int)swDisplayStateOpts_e.swThisDisplayState, null);
-					if (ComponentVisibility != null)
-					{
-						int Visible = ComponentVisibility[0];
-						if (Visible == (int)swComponentVisibilityState_e.swComponentHidden)
-						{
-							ActorExportInfo.bVisible = false;
-						}
-					}
-				}
-				else
-				{
-					ActorExportInfo.bVisible = false;
-				}
-
-				bool bNeedsGeometryExport = false;
-
-				// ComponentDoc is null if component is suppressed or lightweight
-				ModelDoc2 ComponentDoc = InComponent.GetModelDoc2();
-
-				if (ComponentDoc is PartDoc) // Also possible to test document type with swDocumentTypes_e.swDocPART
-				{
-					bool bFirstExport = !SyncState.ExportedComponentsMap.ContainsKey(ComponentName);
-					bNeedsGeometryExport = bFirstExport;
-
-					if (!bFirstExport && SyncState.DirtyComponents.ContainsKey(ComponentName))
-					{
-						uint DirtyState = SyncState.DirtyComponents[ComponentName];
-						bNeedsGeometryExport =
-							((DirtyState & (1u << (int)EComponentDirtyState.Material)) != 0) ||
-							((DirtyState & (1u << (int)EComponentDirtyState.Geometry)) != 0) ||
-							((DirtyState & (1u << (int)EComponentDirtyState.Delete)) != 0);
-					}
-				}
-
-				if (bNeedsGeometryExport)
-				{
-					//TODO this will be null for new part, think of more solid solution
-					string PartPath = ComponentDoc.GetPathName();
-					if (!SyncState.PartsMap.ContainsKey(PartPath))
-					{
-						// New part
-						SyncState.PartsMap[PartPath] = AddTrackedAssemblyPart(InComponent, ComponentName);
-					}
-
-					SyncState.ComponentToPartMap[ComponentName] = PartPath;
-
-					// This component has associated part document -- treat is as a mesh actor
-					ActorExportInfo.Type = EActorType.MeshActor;
-
-					OutMeshesToExportMap.AddComponentWhichNeedsMesh(InComponent, ComponentName, ActorExportInfo.Name);
-				}
-
-				if (ActorExportInfo.Type == EActorType.MeshActor)
-				{
-					// todo: deduplicate mesh name calculation
-					ActorExportInfo.MeshName = new FMeshName(ComponentName);
-				}
-
-				Exporter.ExportOrUpdateActor(ActorExportInfo);
-
-				SyncState.CleanComponents.Add(ComponentName);
-			}
-
-			SyncState.ExportedComponentsMap[ComponentName] = InComponent;
-
-			// Export component children
-			object[] Children = (object[])InComponent.GetChildren();
-
-			if (Children != null)
-			{
-				foreach (object Obj in Children)
-				{
-					Component2 Child = (Component2)Obj;
-					ExportComponentRecursive(Child, InComponent, OutMeshesToExportMap);
-				}
-			}
-		}
-
 		private void ExportComponentRecursive(string ActiveConfigName,
 			FConfigurationExporter ConfigurationExporter,
 			FConfigurationTree.FComponentTreeNode InNode, FConfigurationTree.FComponentTreeNode InParent)
@@ -478,9 +335,7 @@ namespace DatasmithSolidworks
 
 				if (ActorExportInfo.Type == EActorType.MeshActor)
 				{
-					// todo: deduplicate mesh name calculation
-					
-					ActorExportInfo.MeshName = ConfigurationExporter.GetMeshName(InNode.ComponentName);
+					ConfigurationExporter.AddActorForMesh(ActorExportInfo.Name, InNode.ComponentName);
 				}
 
 				Exporter.ExportOrUpdateActor(ActorExportInfo);
@@ -533,12 +388,9 @@ namespace DatasmithSolidworks
 
 					ActorExportInfo.bVisible = ComponentConfig.bVisible && (ComponentConfig.ConfigName == ActiveConfigName);
 
-					// This component has associated part document -- treat is as a mesh actor
 					ActorExportInfo.Type = EActorType.MeshActor;
 
-					;
-					// todo: deduplicate mesh name calculation
-					ActorExportInfo.MeshName = ConfigurationExporter.GetMeshName(ComponentConfig.ConfigName, InNode.ComponentName);
+					ConfigurationExporter.AddActorForMesh(ActorExportInfo.Name, ComponentConfig.ConfigName, InNode.ComponentName);
 
 					Exporter.ExportOrUpdateActor(ActorExportInfo);
 				}
@@ -1094,6 +946,11 @@ namespace DatasmithSolidworks
 			return DocumentTracker.GetDirty();
 		}
 
+		public FDocumentTracker GetTracker()
+		{
+			return DocumentTracker;
+		}
+
 		public void SetDirty(bool bInDirty)  // After Sync, where called with true??? Probably not needed at all(i.e. all SetDirty is inside in notifiers, and with false can be in Sync)
 		{
 			DocumentTracker.SetDirty(bInDirty);
diff --git a/Engine/Source/Programs/Enterprise/Datasmith/DatasmithSolidworksExporter/Private/Body.cs b/Engine/Source/Programs/Enterprise/Datasmith/DatasmithSolidworksExporter/Private/Body.cs
--- a/Engine/Source/Programs/Enterprise/Datasmith/DatasmithSolidworksExporter/Private/Body.cs
+++ b/Engine/Source/Programs/Enterprise/Datasmith/DatasmithSolidworksExporter/Private/Body.cs
@@ -6,6 +6,7 @@ using System.Runtime.InteropServices;
 using SolidWorks.Interop.swconst;
 using System.Threading.Tasks;
 using System.Collections.Concurrent;
+using System.Diagnostics;
 
 namespace DatasmithSolidworks
 {
@@ -94,7 +95,10 @@ namespace DatasmithSolidworks
 					}
 				} while (Body != null);
 			}
-			catch { }
+			catch
+			{
+				Debug.Assert(false);
+			}
 
 			return Bodies;
 		}
diff --git a/Engine/Source/Programs/Enterprise/Datasmith/DatasmithSolidworksExporter/Private/ConfigurationExporter.cs b/Engine/Source/Programs/Enterprise/Datasmith/DatasmithSolidworksExporter/Private/ConfigurationExporter.cs
--- a/Engine/Source/Programs/Enterprise/Datasmith/DatasmithSolidworksExporter/Private/ConfigurationExporter.cs
+++ b/Engine/Source/Programs/Enterprise/Datasmith/DatasmithSolidworksExporter/Private/ConfigurationExporter.cs
@@ -1,18 +1,16 @@
 // Copyright Epic Games, Inc. All Rights Reserved.
 
 using System;
-using System.Collections;
 using System.Collections.Generic;
 using SolidWorks.Interop.sldworks;
 using SolidWorks.Interop.swconst;
-using System.Runtime.InteropServices;
 using System.Collections.Concurrent;
 using System.Linq;
+using System.Diagnostics;
+
 using DatasmithSolidworks.Names;
 using static DatasmithSolidworks.FConfigurationTree;
-using System.ComponentModel;
-using System.Diagnostics;
-using static DatasmithSolidworks.FMeshes;
+using static DatasmithSolidworks.Addin;
 
 namespace DatasmithSolidworks
 {
@@ -43,9 +41,13 @@ namespace DatasmithSolidworks
 
 	public class FMeshes
 	{
+		private readonly string MainConfigurationName;
 		private Dictionary<string, FConfiguration> Configurations = new Dictionary<string, FConfiguration>();
 
-		private HashSet<FMesh> Meshes = new HashSet<FMesh>();
+		// Store unique meshes(identical mesh from different component will use the same FMesh object)
+		// Using Dictionary instead of HashSet to be able to retrieve stored object by the key(which is a different object for identical bu different mesh)
+		private Dictionary<FMesh, FMesh> Meshes = new Dictionary<FMesh, FMesh>();  
+		private Dictionary<object, FMeshName> MeshNames = new Dictionary<object, FMeshName>();
 
 		public struct FMesh
 		{
@@ -143,6 +145,11 @@ namespace DatasmithSolidworks
 			}
 		}
 
+		public FMeshes(string InMainConfigurationName)
+		{
+			MainConfigurationName = InMainConfigurationName;
+		}
+
 		public FConfiguration GetMeshesConfiguration(string ConfigurationName)
 		{
 			if (Configurations.TryGetValue(ConfigurationName, out var Configuration))
@@ -150,7 +157,7 @@ namespace DatasmithSolidworks
 				return Configuration;
 			}			
 
-			Configuration = new FConfiguration(this, ConfigurationName);
+			Configuration = new FConfiguration(this, ConfigurationName, MainConfigurationName==ConfigurationName);
 			Configurations.Add(ConfigurationName, Configuration);
 			return Configuration;
 		}
@@ -173,28 +180,42 @@ namespace DatasmithSolidworks
 			public class FComponentMesh
 			{
 				public Component2 Component;
-				public FActorName ActorName;
 				public FMesh Mesh = new FMesh();
+				public FMeshName MeshName;
 			}
 
+			private readonly FMeshes Meshes;
+			private readonly string Name;
+			private readonly bool bIsMainConfiguration;
 			private ConcurrentDictionary<FComponentName, FComponentMesh> MeshForComponent = new ConcurrentDictionary<FComponentName, FComponentMesh>();
 
-			public FConfiguration(FMeshes Meshes, string Name)
+			public FConfiguration(FMeshes InMeshes, string InName, bool bInIsMainConfiguration)
 			{
+				Meshes = InMeshes;
+				Name = InName;
+				bIsMainConfiguration = bInIsMainConfiguration;
 			}
 
-			public void AddComponentWhichNeedsMesh(Component2 InComponent, FComponentName InComponentName, FActorName InActorName)
+			public void AddComponentWhichNeedsMesh(Component2 InComponent, FComponentName InComponentName)
 			{
-				MeshForComponent.TryAdd(InComponentName, new FComponentMesh() {Component = InComponent, ActorName = InActorName});
+				MeshForComponent.TryAdd(InComponentName, new FComponentMesh() {Component = InComponent});
 			}
 
-			public void AddMesh(Component2 Component, FMeshData MeshData)
+			public void AddMesh(FComponentName ComponentName, FMeshData MeshData, out FMeshName MeshName)
 			{
 				FComponentMesh ComponentMesh;
-				if (MeshForComponent.TryGetValue(new FComponentName(Component), out ComponentMesh))
+				if (MeshForComponent.TryGetValue(ComponentName, out ComponentMesh))
 				{
-					ComponentMesh.Mesh = new FMesh(MeshData);
+					// MeshName for active configuration is simplified(doesn't contain config name)
+					FMeshName MeshNameDesired = bIsMainConfiguration ? new FMeshName(ComponentName) : new FMeshName(ComponentName, Name);
+
+					// AddMesh may return mesh(and mesh name) of an identical mesh already processed - use these instead then
+					Meshes.AddMesh(MeshNameDesired,  MeshData, out ComponentMesh.MeshName, out ComponentMesh.Mesh);
+					MeshName = ComponentMesh.MeshName;
+					return;
 				}
+				MeshName = new FMeshName();
+				Debug.Assert(false);  // unexpected - component which can have mesh added should already be registered...
 			}
 
 			public bool GetMeshForComponent(FComponentName ComponentName, out FMesh Result)
@@ -217,6 +238,32 @@ namespace DatasmithSolidworks
 			{
 				return MeshForComponent.Select(KVP => KVP.Value.Component);
 			}
+
+			public FMeshName GetMeshNameForComponent(FComponentName ComponentName)
+			{
+				return MeshForComponent[ComponentName].MeshName;
+			}
+		}
+
+		// Returns true if this mesh was already there
+		private bool AddMesh(FMeshName MeshName, FMeshData MeshData, out FMeshName OutMeshName, out FMesh OutMesh)
+		{
+			LogDebug($"FMeshes.AddMesh('{MeshName}')");
+			// Store only one instance of identical mesh
+			FMesh Mesh = new FMesh(MeshData);
+			if (Meshes.TryGetValue(Mesh, out OutMesh))
+			{
+				OutMeshName = MeshNames[OutMesh];
+				LogDebug($"    identical mesh found: '{OutMeshName}'");
+				return true;
+			}
+
+			Meshes.Add(Mesh, Mesh);
+			MeshNames.Add(Mesh, MeshName); // Record name for the instance of the mesh also
+
+			OutMesh = Mesh;
+			OutMeshName = MeshName;
+			return false;;
 		}
 
 		public string GetConfigurationNameForComponentMesh(FComponentName ComponentName)
@@ -230,6 +277,11 @@ namespace DatasmithSolidworks
 			}
 			return null;
 		}
+
+		public FMeshName GetMeshName(string ConfigName, FComponentName ComponentName)
+		{
+			return Configurations[ConfigName].GetMeshNameForComponent(ComponentName);
+		}
 	}
 
 	public class FConfigurationExporter
@@ -237,19 +289,27 @@ namespace DatasmithSolidworks
 		public readonly FMeshes Meshes;
 		public readonly string[] ConfigurationNames;
 		private readonly string MainConfigurationName;
+		private readonly bool bExportDisplayStates;
 
 		public FConfigurationTree.FComponentTreeNode CombinedTree;
 
-		public FConfigurationExporter(FMeshes InMeshes, string[] InConfigurationNames, string InMainConfigurationName)
+		private List<FDatasmithExporter.FMeshExportInfo> ExtractedMeshes = new List<FDatasmithExporter.FMeshExportInfo>();
+		public Dictionary<FMeshName, List<FActorName>> ActorsForMesh = new Dictionary<FMeshName, List<FActorName>>();
+		public Dictionary<FComponentName, List<FMeshName>> MeshesForComponent = new Dictionary<FComponentName, List<FMeshName>>();
+
+		public FConfigurationExporter(FMeshes InMeshes, 
+			string[] InConfigurationNames, string InMainConfigurationName,
+			bool bInExportDisplayStates)
 		{
 			Meshes = InMeshes;
 			ConfigurationNames = InConfigurationNames;
 			MainConfigurationName = InMainConfigurationName;
+			bExportDisplayStates = bInExportDisplayStates;
 		}
 
 		public List<FConfigurationData> ExportConfigurations(FDocumentTracker InDoc)
 		{
-			Addin.Instance.LogDebug($"FConfigurationExporter.ExportConfigurations [{string.Join(",", ConfigurationNames)}]");
+			LogDebug($"FConfigurationExporter.ExportConfigurations [{string.Join(",", ConfigurationNames)}]");
 			// todo: may refactor to union behavior for 'no configuration'(i.e. active) and one-to-many configurations export
 			// Not sure how CfgNames might be technically null though
 			if (ConfigurationNames == null) // || CfgNames.Length <= 1)
@@ -283,27 +343,31 @@ namespace DatasmithSolidworks
 			bool bIsActiveConfiguration = true;
 			foreach (string CfgName in ConfigurationNamesActiveFirst)
 			{
-				Addin.Instance.LogDebug($"Configuration '{CfgName}'");
+				LogDebug($"Configuration '{CfgName}'");
 				IConfiguration swConfiguration = InDoc.SwDoc.GetConfigurationByName(CfgName) as IConfiguration;
 
 				if (!bIsActiveConfiguration)
 				{
-					Addin.Instance.LogDebug($"ShowConfiguration2 '{CfgName}'");
+					LogDebug($"ShowConfiguration2 '{CfgName}'");
 					InDoc.SwDoc.ShowConfiguration2(CfgName);
 				}
-				Addin.Instance.LogDebug($"ActiveConfiguration '{ConfigManager.ActiveConfiguration.Name}'");
+				LogDebug($"ActiveConfiguration '{ConfigManager.ActiveConfiguration.Name}'");
 
 				string ConfigName = CfgName;
 
-				int DisplayStateCount = swConfiguration.GetDisplayStatesCount();
 				string[] DisplayStates = null;
+				int DisplayStateCount = 0;
 
-				if (ConfigManager.LinkDisplayStatesToConfigurations && DisplayStateCount > 1)
+				if (bExportDisplayStates)
 				{
-					DisplayStates = swConfiguration.GetDisplayStates();
-					ConfigName = $"{CfgName}_{DisplayStates[0]}";
-				}
+					DisplayStateCount = swConfiguration.GetDisplayStatesCount();
 
+					if (ConfigManager.LinkDisplayStatesToConfigurations && DisplayStateCount > 1)
+					{
+						DisplayStates = swConfiguration.GetDisplayStates();
+						ConfigName = $"{CfgName}_{DisplayStates[0]}";
+					}
+				}
 				FConfigurationTree.FComponentTreeNode ConfigNode = new FConfigurationTree.FComponentTreeNode();
 				ConfigNode.ComponentInfo.ComponentName = FComponentName.FromCustomString(ConfigName);
 
@@ -314,13 +378,13 @@ namespace DatasmithSolidworks
 				// Build the tree and get default materials (which aren't affected by any configuration)
 				// Use GetRootComponent3() with Resolve = true to ensure suppressed components will be loaded
 				// todo: docs says that Part document SW returns null. Not the case. But probably better to add a guard
-				Addin.Instance.LogDebug($"Components:");
+				LogDebug($"Components:");
 				CollectComponentsRecursive(InDoc, swConfiguration.GetRootComponent3(true), ConfigNode, MeshesConfiguration, CfgName);
 
 				ExportedConfigurationNames.Add(CfgName, ConfigName);
 
-				Addin.Instance.LogDebug($"Materials:");
-				Dictionary<FComponentName, FObjectMaterials> MaterialsMap = GetComponentMaterials(InDoc, DisplayStates != null ? DisplayStates[0] : null, swConfiguration);
+				LogDebug($"Materials:");
+				Dictionary<FComponentName, FObjectMaterials> MaterialsMap = GetComponentMaterials(InDoc, DisplayStates?[0], swConfiguration);
 				SetComponentTreeMaterials(ConfigNode, MaterialsMap, null, false);
 
 				if (DisplayStates != null)
@@ -331,7 +395,7 @@ namespace DatasmithSolidworks
 					{
 						string DisplayState = DisplayStates[Index];
 						string DisplayStateTreeName = $"{CfgName}_DisplayState_{FDatasmithExporter.SanitizeName(DisplayState)}";
-						Addin.Instance.LogDebug($"Linked DisplayState '{DisplayState}' Materials for variant '{DisplayStateTreeName}'");
+						LogDebug($"Linked DisplayState '{DisplayState}' Materials for variant '{DisplayStateTreeName}'");
 
 						MaterialsMap = GetComponentMaterials(InDoc, DisplayState, swConfiguration);
 						SetComponentTreeMaterials(ConfigNode, MaterialsMap, DisplayStateTreeName, false);
@@ -362,10 +426,7 @@ namespace DatasmithSolidworks
 				// Export meshes
 				InDoc.SetExportStatus($"Component Meshes");
 
-				string MeshSuffix = bIsActiveConfiguration ? null : CfgName; // Suffix for main configuration is absent(so in default or single configuration export we don't have unnecessary long names)
-				List<FDatasmithExporter.FMeshExportInfo> CreatedMeshes = InDoc.ProcessConfigurationMeshes(MeshesConfiguration, MeshSuffix);
-				InDoc.ExportMaterials();
-				InDoc.AssignMaterialsToDatasmithMeshes(CreatedMeshes);
+				InDoc.ProcessConfigurationMeshes(ExtractedMeshes, MeshesConfiguration);
 
 				// Combine separate scene trees into the single one with configuration-specific data
 				FConfigurationTree.Merge(CombinedTree, ConfigNode, CfgName, bIsActiveConfiguration);
@@ -375,7 +436,7 @@ namespace DatasmithSolidworks
 
 			if (OriginalConfiguration != null && ConfigurationNames.Length > 1)
 			{
-				Addin.Instance.LogDebug($"ShowConfiguration2 '{OriginalConfiguration.Name}'");
+				LogDebug($"ShowConfiguration2 '{OriginalConfiguration.Name}'");
 				bool bShowConfigurationResult = InDoc.SwDoc.ShowConfiguration2(OriginalConfiguration.Name);
 				Debug.Assert(bShowConfigurationResult);
 			}
@@ -407,7 +468,7 @@ namespace DatasmithSolidworks
 
 			List<string> DisplayStateConfigurations = new List<string>();
 
-			if (!ConfigManager.LinkDisplayStatesToConfigurations)
+			if (bExportDisplayStates && !ConfigManager.LinkDisplayStatesToConfigurations)
 			{
 				// Export display states as separate configurations
 				string[] DisplayStates = OriginalConfiguration.GetDisplayStates();
@@ -422,7 +483,7 @@ namespace DatasmithSolidworks
 
 					string DisplayStateConfigName = $"DisplayState_{FDatasmithExporter.SanitizeName(DisplayState)}";
 					DisplayStateConfigurations.Add(DisplayStateConfigName);
-					Addin.Instance.LogDebug(
+					LogDebug(
 						$"DisplayState '{DisplayState}' Materials for variant '{DisplayStateConfigName}'");
 
 					// Specifying display state doesn't seem to work in the api like this:
@@ -479,6 +540,45 @@ namespace DatasmithSolidworks
 			return FlatConfigurationData;
 		}
 
+		// Assign meshes to actors, export materials and assign materials
+		public void FinalizeExport(FDocumentTracker InDoc)
+		{
+			List<FDatasmithExporter.FMeshExportInfo> ExportedMeshes = InDoc.Exporter.ExportMeshes(ExtractedMeshes).ToList();
+
+			foreach (FDatasmithExporter.FMeshExportInfo Info in ExportedMeshes)
+			{
+				LogDebug($"  AddMesh(DatasmithMeshName: {Info}");
+
+				InDoc.Exporter.AddMesh(Info);
+
+				// Assign to child variant actors of component actors registered 
+				if (ActorsForMesh.TryGetValue(Info.MeshName, out List<FActorName> Names))
+				{
+					foreach (FActorName ActorName in Names)
+					{
+
+						InDoc.Exporter.AssignMeshToDatasmithMeshActor(ActorName, Info.MeshElement);
+					}
+				}
+			}
+
+			foreach (KeyValuePair<FComponentName, List<FMeshName>> KVP in MeshesForComponent)
+			{
+				FComponentName ComponentName = KVP.Key;
+				InDoc.ReleaseComponentMeshes(ComponentName);
+				// Register that this mesh was used for the component
+				foreach (FMeshName MeshName in KVP.Value)
+				{
+					InDoc.AddMeshForComponent(ComponentName, MeshName);
+				}
+			}
+
+			InDoc.CleanupComponentMeshes();
+
+			InDoc.ExportMaterials();
+			InDoc.AssignMaterialsToDatasmithMeshes(ExportedMeshes);
+		}
+
 		private static Dictionary<FComponentName, FObjectMaterials> GetComponentMaterials(FDocumentTracker InDoc, string InDisplayState, IConfiguration InConfiguration)
 		{
 			Dictionary<FComponentName, FObjectMaterials> MaterialsMap = new Dictionary<FComponentName, FObjectMaterials>();
@@ -520,7 +620,7 @@ namespace DatasmithSolidworks
 
 		private static void SetComponentTreeMaterials(FConfigurationTree.FComponentTreeNode InComponentTree, Dictionary<FComponentName, FObjectMaterials> InComponentMaterialsMap, string InConfigurationName, bool bIsDisplayState)
 		{
-			Addin.Instance.LogDebug($"SetComponentTreeMaterials: '{InComponentTree.ComponentName}'");
+			LogDebug($"SetComponentTreeMaterials: '{InComponentTree.ComponentName}'");
 
 			FConfigurationTree.FComponentConfig TargetConfig = null;
 
@@ -543,20 +643,20 @@ namespace DatasmithSolidworks
 				TargetConfig.Materials = Materials;
 			}
 
-			Addin.Instance.LogDebug($"  Materials: {TargetConfig.Materials}");
+			LogDebug($"  Materials: {TargetConfig.Materials}");
 
 			foreach (FConfigurationTree.FComponentTreeNode Child in InComponentTree.EnumChildren())
 			{
-				Addin.Instance.LogIndent();
+				LogIndent();
 				SetComponentTreeMaterials(Child, InComponentMaterialsMap, InConfigurationName, bIsDisplayState);
-				Addin.Instance.LogDedent();
+				LogDedent();
 			}
 		}
 
 		private static void CollectComponentsRecursive(FDocumentTracker InDoc, Component2 InComponent,
 			FComponentTreeNode InParentNode, FMeshes.FConfiguration Meshes, string CfgName)
 		{
-			Addin.Instance.LogDebug($"'{InComponent.Name2}'");
+			LogDebug($"'{InComponent.Name2}'");
 
 			FConfigurationTree.FComponentTreeNode NewNode = new FConfigurationTree.FComponentTreeNode(InComponent);
 			InParentNode.Children.Add(NewNode);
@@ -606,9 +706,9 @@ namespace DatasmithSolidworks
 				foreach (object ObjChild in Children)
 				{
 					Component2 Child = (Component2)ObjChild;
-					Addin.Instance.LogIndent();
+					LogIndent();
 					CollectComponentsRecursive(InDoc, Child, NewNode, Meshes, CfgName);
-					Addin.Instance.LogDedent();
+					LogDedent();
 				}
 
 				// todo: sorting ensures that component order is stable in export(e.g. in variant property bindings order)
@@ -617,14 +717,13 @@ namespace DatasmithSolidworks
 				NewNode.Children.Sort((InA, InB) => InA.ComponentId - InB.ComponentId);
 			}
 
-			if (InDoc.NeedExportComponent(NewNode, NewNode.CommonConfig))
+			if (NewNode.IsPartComponent())
 			{
-				if (NewNode.IsPartComponent())
+				if (InDoc.NeedExportComponent(NewNode, NewNode.CommonConfig))
 				{
-					// Collect meshes to export when it's a part component
 					InDoc.AddPartDocument(NewNode);
-					Meshes.AddComponentWhichNeedsMesh(NewNode.Component, NewNode.ComponentName, NewNode.ActorName);
 				}
+				Meshes.AddComponentWhichNeedsMesh(NewNode.Component, NewNode.ComponentName);
 			}
 			InDoc.AddExportedComponent(NewNode);
 		}
@@ -642,7 +741,7 @@ namespace DatasmithSolidworks
 
 		public FMeshName GetMeshName(string ConfigName, FComponentName ComponentName)
 		{
-			return FMeshName.FromString(ConfigName == MainConfigurationName ? $"{ComponentName}_Mesh" : $"{ComponentName}_{ConfigName}_Mesh");
+			return Meshes.GetMeshName(ConfigName, ComponentName);
 		}
 
 		// Get mesh for component which has only single mesh configuration
@@ -651,5 +750,34 @@ namespace DatasmithSolidworks
 			string ConfigName = Meshes.GetConfigurationNameForComponentMesh(ComponentName);
 			return GetMeshName(ConfigName, ComponentName);
 		}
+
+		// Register actor which needs mesh of a component in specific config
+		// this is needed for mesh actors for variants
+		public void AddActorForMesh(FActorName ActorName, string ConfigName, FComponentName ComponentName)
+		{
+			LogDebug($"AddActorForMesh(ActorName='{ActorName}', ConfigName='{ConfigName}', ComponentName='{ComponentName}')");
+			AddActorForMesh(ActorName, ComponentName, GetMeshName(ConfigName, ComponentName));
+		}
+
+		// Register actor which needs mesh of a component in any config(i.e. single actor of a component which doesn't have different meshes in different configs)
+		// this is needed for mesh actors for variants
+		public void AddActorForMesh(FActorName ActorName, FComponentName ComponentName)
+		{
+			LogDebug($"AddActorForMesh(ActorName='{ActorName}', ComponentName='{ComponentName}')");
+			AddActorForMesh(ActorName, ComponentName, GetMeshName(ComponentName));
+		}
+		private void AddActorForMesh(FActorName ActorName, FComponentName ComponentName, FMeshName MeshName)
+		{
+			LogDebug($"  MeshName='{MeshName}')");
+			if (!MeshName.IsValid())
+			{
+				LogDebug(
+					$"  mesh INVALID for this component/configuration(e.g. empty mesh would not be exported)");
+				return;
+			}
+
+			FindOrAdd(ActorsForMesh, MeshName).Add(ActorName);
+			FindOrAdd(MeshesForComponent, ComponentName).Add(MeshName);
+		}
 	}
 }
\ No newline at end of file
diff --git a/Engine/Source/Programs/Enterprise/Datasmith/DatasmithSolidworksExporter/Private/DatasmithExporter.cs b/Engine/Source/Programs/Enterprise/Datasmith/DatasmithSolidworksExporter/Private/DatasmithExporter.cs
--- a/Engine/Source/Programs/Enterprise/Datasmith/DatasmithSolidworksExporter/Private/DatasmithExporter.cs
+++ b/Engine/Source/Programs/Enterprise/Datasmith/DatasmithSolidworksExporter/Private/DatasmithExporter.cs
@@ -7,10 +7,9 @@ using System.Collections.Generic;
 using System.Diagnostics;
 using System.IO;
 using System.Linq;
-using System.Runtime.InteropServices;
 using System.Threading.Tasks;
 using DatasmithSolidworks.Names;
-using static DatasmithSolidworks.FConfigurationTree;
+using static DatasmithSolidworks.Addin;
 
 
 namespace DatasmithSolidworks
@@ -151,6 +150,11 @@ namespace DatasmithSolidworks
 		public FMeshName MeshName;
 		public float[] Transform;
 		public bool bVisible;
+
+		public override string ToString()
+		{
+			return $"FDatasmithActorExportInfo(Name={Name}, Type={Type}, MeshName={MeshName}, ParentName={ParentName})";
+		}
 	};
 
 	public class FDatasmithExporter
@@ -178,7 +182,7 @@ namespace DatasmithSolidworks
 
 			public override string ToString()
 			{
-				return $"FMeshExportInfo(ComponentName={ComponentName}, MeshName={MeshName}, ActorName={ActorName}), MeshElement={(MeshElement==null?"<null>":MeshElement.GetName())}";
+				return $"FMeshExportInfo(ComponentName={ComponentName}, MeshName={MeshName}), MeshElement={(MeshElement==null?"<null>":MeshElement.GetName())}";
 			}
 		}
 
@@ -199,6 +203,7 @@ namespace DatasmithSolidworks
 
 		public FDatasmithFacadeActor ExportOrUpdateActor(FDatasmithActorExportInfo InExportInfo)
 		{
+			LogDebug($"ExportOrUpdateActor({InExportInfo})");
 			FDatasmithFacadeActor Actor = null;
 
 			if (ExportedActorsMap.ContainsKey(InExportInfo.Name))
@@ -244,6 +249,7 @@ namespace DatasmithSolidworks
 					DatasmithScene.AddActor(Actor);
 				}
 			}
+			LogDebug($"  Datasmith Name for the Actor '{Actor.GetName()}'");
 
 			// ImportBinding uses Tag[0] ('original name') to group parts used in variants
 			Actor.SetLabel(InExportInfo.Label);
@@ -255,24 +261,32 @@ namespace DatasmithSolidworks
 				FDatasmithFacadeActorMesh MeshActor = Actor as FDatasmithFacadeActorMesh;
 				if (InExportInfo.MeshName.IsValid())
 				{
-					MeshActor.SetMesh(GetDatasmithMeshName(InExportInfo.MeshName));
+					string DatasmithMeshName = GetDatasmithMeshName(InExportInfo.MeshName);
+					LogDebug($"  Set Mesh '{InExportInfo.MeshName}' with Datasmith Name '{DatasmithMeshName}' to DatasmithActorMesh)");
+					MeshActor.SetMesh(DatasmithMeshName);
 				}
 			}
 
 			return Actor;
 		}
 
-		private void AddMesh(FMeshExportInfo Info)
+		public void AddMesh(FMeshExportInfo Info)
 		{
+			LogDebug($"FDatasmithExporter.AddMesh('{Info.MeshName}')");
 			Debug.Assert(Info.MeshName.IsValid());
 			Debug.Assert(Info.MeshElement != null);
 
+
+			LogIndent();
 			RemoveMesh(Info.MeshName);
+			LogDedent();
 
 			ExportedMeshesMap.TryAdd(Info.MeshName, Info);
+
+			LogDebug($"DatasmithScene.AddMesh('{Info.MeshElement.GetName()}')");
 			DatasmithScene.AddMesh(Info.MeshElement);
 
-			// Increase user count on materials
+			// Increase user count on mesh materials
 			foreach (int MaterialId in Info.Materials)
 			{
 				AddMaterialUser(MaterialId);
@@ -281,12 +295,15 @@ namespace DatasmithSolidworks
 
 		public void RemoveMesh(FMeshName MeshName)
 		{
+			LogDebug($"FDatasmithExporter.RemoveMesh('{MeshName}')");
+
             if (!MeshName.IsValid())
             {
                 return;
             }
 			if (ExportedMeshesMap.TryRemove(MeshName, out FMeshExportInfo Info))
 			{
+				LogDebug($"  removing mesh element '{Info.MeshElement.GetName()}'");
 				DatasmithScene.RemoveMesh(Info.MeshElement);
 
 				// Decrease user count on materials
@@ -294,7 +311,10 @@ namespace DatasmithSolidworks
 				{
 					RemoveMaterialUser(MaterialId);
 				}
-
+			}
+			else
+			{
+				LogDebug($"  WARNING: NOT FOUND IN EXPORTED");
 			}
 		}
 
@@ -411,13 +431,13 @@ namespace DatasmithSolidworks
 
 			if (MeshData.Vertices == null || MeshData.Normals == null || MeshData.TexCoords == null || MeshData.Triangles == null)
 			{
-				Addin.Instance.LogDebugThread($"  skipping(some attributes are null) - Vertices:{MeshData.Vertices == null}, Normals:{MeshData.Normals == null}, TexCoords: {MeshData.TexCoords == null}, Triangles:{MeshData.Triangles}");
+				LogDebugThread($"  skipping(some attributes are null) - Vertices:{MeshData.Vertices == null}, Normals:{MeshData.Normals == null}, TexCoords: {MeshData.TexCoords == null}, Triangles:{MeshData.Triangles}");
 				return false;
 			}
 
 			if (MeshData.Vertices.Length == 0 || MeshData.Normals.Length == 0 || MeshData.TexCoords.Length == 0 || MeshData.Triangles.Length == 0)
 			{
-				Addin.Instance.LogDebugThread($"  skipping - Vertices: {MeshData.Vertices.Length}, Normals: {MeshData.Normals.Length}, TexCoords: {MeshData.TexCoords.Length} Triangles: {MeshData.Triangles.Length}");
+				LogDebugThread($"  skipping - Vertices: {MeshData.Vertices.Length}, Normals: {MeshData.Normals.Length}, TexCoords: {MeshData.TexCoords.Length} Triangles: {MeshData.Triangles.Length}");
 				return false;
 			}
 
@@ -475,10 +495,10 @@ namespace DatasmithSolidworks
 					Mesh.SetFaceUV(TriIndex, 0, Triangle[0], Triangle[1], Triangle[2]);
 				}
 
-				Addin.Instance.LogDebugThread($"  material ids assigned to faces: {string.Join(", ", MeshMaterialIds)})");
+				LogDebugThread($"  material ids assigned to faces: {string.Join(", ", MeshMaterialIds)})");
 				if (MeshAddedMaterials.Count > 0)
 				{
-					Addin.Instance.LogDebugThread($"  assigned materials: {string.Join(", ", MeshAddedMaterials)})");
+					LogDebugThread($"  assigned materials: {string.Join(", ", MeshAddedMaterials)})");
 				}
 
 				DatasmithScene.ExportDatasmithMesh(MeshElement, Mesh);
@@ -776,7 +796,7 @@ namespace DatasmithSolidworks
 
 		public void ExportMaterials(ConcurrentDictionary<int, FMaterial> InMaterialsMap)
 		{
-			Addin.Instance.LogDebug($"ExportMaterials: \n  {string.Join("  \n", InMaterialsMap.Select(KVP => $"{KVP.Key}: {KVP.Value}" ))}");
+			LogDebug($"ExportMaterials: \n  {string.Join("  \n", InMaterialsMap.Select(KVP => $"{KVP.Key}: {KVP.Value}" ))}");
 
 			List<FDatasmithFacadeTexture> CreatedTextures = new List<FDatasmithFacadeTexture>();
 			List<FDatasmithFacadeMaterialInstance> CreatedMaterials = new List<FDatasmithFacadeMaterialInstance>();
@@ -787,14 +807,14 @@ namespace DatasmithSolidworks
 				int MaterialId = MatKVP.Key;
 				FMaterial Material = MatKVP.Value;
 
-				Addin.Instance.LogDebug($"  Material: {MaterialId}, {Material.Name}");
+				LogDebug($"  Material: {MaterialId}, {Material.Name}");
 
 				if (!MaterialUsers.ContainsKey(MaterialId))
 				{
-					Addin.Instance.LogDebug($"    skipping, no users");
+					LogDebug($"    skipping, no users");
 					if (ExportedMaterialInstancesMap.ContainsKey(MaterialId))
 					{
-						Addin.Instance.LogDebug($"      but should remove exported");
+						LogDebug($"      but should remove exported");
 						if (ExportedMaterialInstancesMap.TryRemove(MaterialId, out FDatasmithFacadeMaterialInstance MaterialInstance))
 						{
 							UnusedMaterials.Add(MaterialInstance);
@@ -805,7 +825,7 @@ namespace DatasmithSolidworks
 
 				if (ExportedMaterialInstancesMap.ContainsKey(MaterialId))
 				{
-					Addin.Instance.LogDebug($"    skipping, already exported");
+					LogDebug($"    skipping, already exported");
 					continue;
 				}
 
@@ -821,17 +841,17 @@ namespace DatasmithSolidworks
 				}
 			}
 
-			Addin.Instance.LogDebug($"  Remove Materials");
+			LogDebug($"  Remove Materials");
 			foreach (FDatasmithFacadeMaterialInstance Mat in UnusedMaterials)
 			{
-				Addin.Instance.LogDebug($"    Material: {Mat.GetName()}: {Mat.GetLabel()}");
+				LogDebug($"    Material: {Mat.GetName()}: {Mat.GetLabel()}");
 				DatasmithScene.RemoveMaterial(Mat);
 			}
 
-			Addin.Instance.LogDebug($"  AddMaterials");
+			LogDebug($"  AddMaterials");
 			foreach (FDatasmithFacadeMaterialInstance Mat in CreatedMaterials)
 			{
-				Addin.Instance.LogDebug($"    Material: {Mat.GetName()}: {Mat.GetLabel()}");
+				LogDebug($"    Material: {Mat.GetName()}: {Mat.GetLabel()}");
 				DatasmithScene.AddMaterial(Mat);
 			}
 
@@ -844,12 +864,12 @@ namespace DatasmithSolidworks
 
 		private void ConvertMaterialToDatasmith(FMaterial InMaterial, out List<FDatasmithFacadeTexture> OutCreatedTextures, out FDatasmithFacadeMaterialInstance OutCreatedMaterial)
 		{
-			Addin.Instance.LogDebug($"ConvertMaterialToDatasmith({InMaterial.Name}: {InMaterial.ID})");
+			LogDebug($"ConvertMaterialToDatasmith({InMaterial.Name}: {InMaterial.ID})");
 
 			OutCreatedTextures = null;
 			OutCreatedMaterial = null;
 
-			Addin.Instance.LogDebug($"  making, ShaderName: '{InMaterial.ShaderName}'");
+			LogDebug($"  making, ShaderName: '{InMaterial.ShaderName}'");
 
 			FMaterial.EMaterialType Type = FMaterial.GetMaterialType(InMaterial.ShaderName);
 
@@ -1125,9 +1145,9 @@ namespace DatasmithSolidworks
 		}
 
 		// Export meshes
-		public void ExportMeshes(List<FMeshExportInfo> MeshExportInfos, out List<FMeshExportInfo> OutCreatedMeshes)
+		public IEnumerable<FMeshExportInfo> ExportMeshes(IEnumerable<FMeshExportInfo> MeshExportInfos)
 		{
-			Addin.Instance.LogDebug($"ExportMeshes: {string.Join(", ", MeshExportInfos)}");
+			LogDebug($"ExportMeshes: {string.Join(", ", MeshExportInfos)}");
 
 			ConcurrentQueue<FMeshExportInfo> CreatedMeshes = new ConcurrentQueue<FMeshExportInfo>();
 
@@ -1139,48 +1159,48 @@ namespace DatasmithSolidworks
 				}
 			});
 
-			// Remember exported meshes and assign to actor
-			foreach (FMeshExportInfo Info in CreatedMeshes)
-			{
-				Addin.Instance.LogDebug($"  AddMesh(DatasmithMeshName: {Info}");
-
-				AddMesh(Info);
+			return CreatedMeshes;
+		}
 
-				if (Info.ActorName.IsValid())
+		public void AssignMeshToDatasmithMeshActor(FActorName ActorName, FDatasmithFacadeMeshElement MeshElement)
+		{
+			LogDebug($"AssignMeshToDatasmithMeshActor('{ActorName}', '{MeshElement.GetName()}')");
+			if (ActorName.IsValid())
+			{
+				Tuple<EActorType, FDatasmithFacadeActor> ExportedActorInfo = null;
+				if (ExportedActorsMap.TryGetValue(ActorName, out ExportedActorInfo) &&
+				    ExportedActorInfo.Item1 == EActorType.MeshActor)
 				{
-					Tuple<EActorType, FDatasmithFacadeActor> ExportedActorInfo = null;
-					if (ExportedActorsMap.TryGetValue(Info.ActorName, out ExportedActorInfo) && ExportedActorInfo.Item1 == EActorType.MeshActor)
-					{
-						FDatasmithFacadeActorMesh MeshActor = ExportedActorInfo.Item2 as FDatasmithFacadeActorMesh;
-						Addin.Instance.LogDebug($"  SetMesh on MeshActor: {MeshActor.GetName()}: {MeshActor.GetLabel()}  )");
-						MeshActor.SetMesh(Info.MeshElement.GetName());
-					}
+					FDatasmithFacadeActorMesh MeshActor = ExportedActorInfo.Item2 as FDatasmithFacadeActorMesh;
+					LogDebug($"  SetMesh on MeshActor: {MeshActor.GetName()}: {MeshActor.GetLabel()}  )");
+					MeshActor.SetMesh(MeshElement.GetName());
+				}
+				else
+				{
+					LogDebug($"  NOT setting mesh yet - actor '{ActorName}' not exported')");
 				}
 			}
-
-			OutCreatedMeshes = CreatedMeshes.ToList();
-
 		}
 
 		public void AssignMaterialsToDatasmithMeshes(List<FMeshExportInfo> MeshExportInfos)
 		{
-			Addin.Instance.LogDebug($"AssignMaterialsToDatasmithMeshes: {string.Join(", ", MeshExportInfos)}");
+			LogDebug($"AssignMaterialsToDatasmithMeshes: {string.Join(", ", MeshExportInfos)}");
 
 			foreach (FMeshExportInfo Info in MeshExportInfos)
 			{
-				Addin.Instance.LogDebug($"  {Info}");
+				LogDebug($"  {Info}");
 
 				foreach (int MaterialId in Info.Materials)
 				{
-					Addin.Instance.LogDebug($"    Material: {MaterialId}");
+					LogDebug($"    Material: {MaterialId}");
 					if (GetDatasmithMaterial(MaterialId, out FDatasmithFacadeMaterialInstance Material))
 					{
-						Addin.Instance.LogDebug($"      SetMaterial({Material.GetName()}, SlotId: {MaterialId})");
+						LogDebug($"      SetMaterial({Material.GetName()}, SlotId: {MaterialId})");
 						Info.MeshElement.SetMaterial(Material.GetName(), MaterialId);
 					}
 					else
 					{
-						Addin.Instance.LogDebug($"      NOT FOUND among exported materials!");
+						LogDebug($"      NOT FOUND among exported materials!");
 					}
 				}
 			}
diff --git a/Engine/Source/Programs/Enterprise/Datasmith/DatasmithSolidworksExporter/Private/Document.cs b/Engine/Source/Programs/Enterprise/Datasmith/DatasmithSolidworksExporter/Private/Document.cs
--- a/Engine/Source/Programs/Enterprise/Datasmith/DatasmithSolidworksExporter/Private/Document.cs
+++ b/Engine/Source/Programs/Enterprise/Datasmith/DatasmithSolidworksExporter/Private/Document.cs
@@ -6,12 +6,9 @@ using System.Collections.Concurrent;
 using System.Collections.Generic;
 using System.Diagnostics;
 using System.IO;
-using System.Linq;
 using System.Threading;
-using System.Threading.Tasks;
 using DatasmithSolidworks.Names;
-using static DatasmithSolidworks.FAssemblyDocumentTracker;
-using System.Windows.Forms;
+using static DatasmithSolidworks.Addin;
 
 namespace DatasmithSolidworks
 {
@@ -35,6 +32,7 @@ namespace DatasmithSolidworks
 			MaterialCheckerEvent.Set();
 
 			MaterialCheckerThread = new Thread(CheckForMaterialUpdatesProc);
+			MaterialCheckerThread.Name = "DatasmithSolidworks.FMaterialUpdateChecker Thread";
 			MaterialCheckerThread.Start();
 		}
 
@@ -72,11 +70,25 @@ namespace DatasmithSolidworks
 					{
 						bSketchMode = (DocumentTracker.SwDoc.SketchManager?.ActiveSketch != null);
 					}
-					catch { }
+					catch
+					{
+						Debug.Assert(false);
+					}
 				
-					if (!bSketchMode && DocumentTracker.HasMaterialUpdates())
+					if (!bSketchMode)
 					{
-						DocumentTracker.SetDirty(true);
+#if DEBUG
+						Stopwatch Watch = Stopwatch.StartNew();
+#endif
+						bool bHasMaterialUpdates = DocumentTracker.HasMaterialUpdates();
+#if DEBUG
+						Watch.Stop();
+						Debug.WriteLine($"FMaterialUpdateChecker->HasMaterialUpdates: {(double)Watch.ElapsedMilliseconds / 1000.0}");
+#endif
+						if (bHasMaterialUpdates)
+						{
+							DocumentTracker.SetDirty(true);
+						}
 					}
 				}
 				Thread.Sleep(600);
@@ -175,6 +187,7 @@ namespace DatasmithSolidworks
 				Exporter.ExportLevelVariantSets(Configs, MaterialBindings);
 			}
 
+			// Export materials after processing variants - so that any material used in variant is registered and exported
 			ExportMaterials();
 
 			// Assign exported datasmith material instances to bindings
@@ -205,6 +218,8 @@ namespace DatasmithSolidworks
 			}
 		}
 
+		public abstract FMeshes GetMeshes(string ActiveConfigName);
+
 		public void Export()
 		{
 			string OutDir = Path.GetDirectoryName(DatasmithFileExportPath);
@@ -220,13 +235,14 @@ namespace DatasmithSolidworks
 
 			Exporter = new FDatasmithExporter(DatasmithScene);
 
-			FMeshes Meshes = new FMeshes();
-
-			PreExport(Meshes, true);
-
 			ConfigurationManager ConfigManager = SwDoc.ConfigurationManager;
 			string[] ConfigurationNames = SwDoc?.GetConfigurationNames();
-			FConfigurationExporter ConfigurationExporter = new FConfigurationExporter(Meshes, ConfigurationNames, ConfigManager.ActiveConfiguration.Name);
+			string ActiveConfigurationName = ConfigManager.ActiveConfiguration.Name;
+
+
+			FMeshes Meshes = new FMeshes(ActiveConfigurationName);
+
+			FConfigurationExporter ConfigurationExporter = new FConfigurationExporter(Meshes, ConfigurationNames, ActiveConfigurationName, true);
 
 			List<FConfigurationData> Configs = ConfigurationExporter.ExportConfigurations(this);
 
@@ -246,16 +262,24 @@ namespace DatasmithSolidworks
 
 		public void Sync(string InOutputPath)
 		{
+			LogDebug($"Sync('{InOutputPath}')");
+			LogIndent();
+
 			DatasmithScene.SetName(SwDoc.GetTitle());
 			DatasmithScene.SetOutputPath(InOutputPath);
 
 #if DEBUG
 			Stopwatch Watch = Stopwatch.StartNew();
 #endif
-			FMeshes Meshes = new FMeshes();
+			ConfigurationManager ConfigManager = SwDoc.ConfigurationManager;
+			string ActiveConfigurationName = ConfigManager.ActiveConfiguration.Name;
+
+			FMeshes Meshes = GetMeshes(ActiveConfigurationName);
+			FConfigurationExporter ConfigurationExporter = new FConfigurationExporter(Meshes, new []{ ActiveConfigurationName }, ActiveConfigurationName, false);
+			List<FConfigurationData> Configs = ConfigurationExporter.ExportConfigurations(this);
+			bHasConfigurations = (Configs != null) && (Configs.Count != 0);
+			ExportToDatasmithScene(ConfigurationExporter);
 
-			PreExport(Meshes, false);
-			ExportToDatasmithScene(Meshes);
 
 			ExportLights();
 
@@ -263,6 +287,7 @@ namespace DatasmithSolidworks
 			Watch.Stop();
 			Debug.WriteLine($"EXPORT TIME: {(double)Watch.ElapsedMilliseconds / 1000.0}");
 #endif
+			LogDedent();
 		}
 
 		// todo: make abstract
@@ -271,9 +296,6 @@ namespace DatasmithSolidworks
 			throw new NotImplementedException();
 		}
 
-		public abstract void ExportToDatasmithScene(FMeshes Meshes);
-		public abstract void PreExport(FMeshes Meshes, bool bConfigurations);  // Called before configurations are parsed to prepare meshes needed to identify if they create different configurations
-
 		public virtual bool NeedExportComponent(FConfigurationTree.FComponentTreeNode InComponent,
 			FConfigurationTree.FComponentConfig ActiveComponentConfig)
 		{
@@ -288,7 +310,13 @@ namespace DatasmithSolidworks
 		public abstract void AddComponentMaterials(FComponentName ComponentName, FObjectMaterials Materials);
 		public abstract FObjectMaterials GetComponentMaterials(Component2 Comp);
 
+
+		// Record which meshes are used aby a component
 		public abstract void AddMeshForComponent(FComponentName ComponentName, FMeshName MeshName);
+		public abstract void ReleaseComponentMeshes(FComponentName CompName);
+
+		// Remove unused meshes from Datasmith scene
+		public abstract void CleanupComponentMeshes();
 
 		public virtual void Destroy()
 		{
@@ -296,51 +324,36 @@ namespace DatasmithSolidworks
 
 		public abstract FMeshData ExtractComponentMeshData(Component2 Comp);
 
-		// Extracts, exports meshes used for the assembly configuration, assigns them to actors
-		// todo: separate just Datasmith Mesh export to its own thread, all other code should be in single thread
-		// Datasmith actor assignment doesn't need threading(and not supposed to be thread-safe?) and Solidworks multithreading is supposed to be slower(SW does all the work in main thread)
-		public List<FDatasmithExporter.FMeshExportInfo> ProcessConfigurationMeshes(FMeshes.FConfiguration MeshesConfiguration, string MeshSuffix)
+
+		// Extracts meshes used for the assembly configuration
+		public void ProcessConfigurationMeshes(List<FDatasmithExporter.FMeshExportInfo> MeshExportInfos, FMeshes.FConfiguration MeshesConfiguration)
 		{
 			
-			List<FDatasmithExporter.FMeshExportInfo> MeshExportInfos = new List<FDatasmithExporter.FMeshExportInfo>();
-
 			// Extract meshes data and prepare for parallel datasmith export
+			// note: mesh data need to be extracted from the component when required configuration is active(i.e. can't move it outside of configuration enumeration loop)
 			foreach (Component2 Comp in MeshesConfiguration.EnumerateComponents())
 			{
 				FMeshData MeshData = ExtractComponentMeshData(Comp);
-				MeshesConfiguration.AddMesh(Comp, MeshData);
-
+				FComponentName ComponentName = new FComponentName(Comp);
+				
 				if (MeshData != null)
 				{
-
-					FActorName ComponentActorName = Exporter.GetComponentActorName(Comp);
-					FComponentName ComponentName = new FComponentName(Comp);
+					MeshesConfiguration.AddMesh(ComponentName, MeshData, out FMeshName MeshName);
 					MeshExportInfos.Add(new FDatasmithExporter.FMeshExportInfo()
 					{
 						ComponentName = ComponentName,
-						MeshName = FMeshName.FromString(MeshSuffix == null
-							? $"{ComponentName}_Mesh"
-							: $"{ComponentName}_{MeshSuffix}_Mesh"),
-						ActorName = ComponentActorName,
+						MeshName = MeshName,
 						MeshData = MeshData
 					});
 				}
 			}
-
-			Exporter.ExportMeshes(MeshExportInfos, out List<FDatasmithExporter.FMeshExportInfo> CreatedMeshes);
-
-			foreach (FDatasmithExporter.FMeshExportInfo Info in CreatedMeshes)
-			{
-				AddMeshForComponent(Info.ComponentName, Info.MeshName);  // Register that this mesh was used for the component
-			}
-
-			return CreatedMeshes;
 		}
 
 		public void AssignMaterialsToDatasmithMeshes(List<FDatasmithExporter.FMeshExportInfo> CreatedMeshes)
 		{
 			Exporter.AssignMaterialsToDatasmithMeshes(CreatedMeshes);
 		}
+
 	};
 
 
@@ -362,6 +375,8 @@ namespace DatasmithSolidworks
 		// Is any changes pending(simple test for AutoSync)
 		bool GetDirty();
 
+		FDocumentTracker GetTracker();
+
 		// After Sync
 		// todo: Probably not needed at all(i.e. all SetDirty(false) is only deep inside in notifiers, and with false can be moved  into Sync?)
 		void SetDirty(bool bInDirty);  
diff --git a/Engine/Source/Programs/Enterprise/Datasmith/DatasmithSolidworksExporter/Private/Material.cs b/Engine/Source/Programs/Enterprise/Datasmith/DatasmithSolidworksExporter/Private/Material.cs
--- a/Engine/Source/Programs/Enterprise/Datasmith/DatasmithSolidworksExporter/Private/Material.cs
+++ b/Engine/Source/Programs/Enterprise/Datasmith/DatasmithSolidworksExporter/Private/Material.cs
@@ -7,6 +7,7 @@ using System.Drawing;
 using System.Collections.Generic;
 using System.Globalization;
 using SolidWorks.Interop.sldworks;
+using System.Diagnostics;
 
 namespace DatasmithSolidworks
 {
@@ -791,6 +792,7 @@ namespace DatasmithSolidworks
 						}
 						catch
 						{
+							Debug.Assert(false);
 						}
 						if (Appearances != null && Appearances.Length > 0)
 						{
diff --git a/Engine/Source/Programs/Enterprise/Datasmith/DatasmithSolidworksExporter/Private/ObjectMaterials.cs b/Engine/Source/Programs/Enterprise/Datasmith/DatasmithSolidworksExporter/Private/ObjectMaterials.cs
--- a/Engine/Source/Programs/Enterprise/Datasmith/DatasmithSolidworksExporter/Private/ObjectMaterials.cs
+++ b/Engine/Source/Programs/Enterprise/Datasmith/DatasmithSolidworksExporter/Private/ObjectMaterials.cs
@@ -6,6 +6,7 @@ using SolidWorks.Interop.swconst;
 using System.Linq;
 using System.Collections.Concurrent;
 using DatasmithSolidworks.Names;
+using static DatasmithSolidworks.Addin;
 
 namespace DatasmithSolidworks
 {
@@ -175,7 +176,7 @@ namespace DatasmithSolidworks
 
 		public static FObjectMaterials LoadPartMaterials(FDocumentTracker InOwnerDoc, PartDoc InPartDoc, swDisplayStateOpts_e InDisplayState, string[] InDisplayStateNames)
 		{
-			Addin.Instance.LogDebug($"LoadPartMaterials: '{(InPartDoc as ModelDoc2).GetPathName()}', {InDisplayState}, {(InDisplayStateNames == null ? "<null>" : string.Join(", ", InDisplayStateNames))}");
+			LogDebug($"LoadPartMaterials: '{(InPartDoc as ModelDoc2).GetPathName()}', {InDisplayState}, {(InDisplayStateNames == null ? "<null>" : string.Join(", ", InDisplayStateNames))}");
 
 			ModelDoc2 Doc = InPartDoc as ModelDoc2;
 
@@ -187,7 +188,7 @@ namespace DatasmithSolidworks
 			IModelDocExtension Ext = Doc.Extension;
 			int NumMaterials = Ext.GetRenderMaterialsCount2((int)InDisplayState, InDisplayStateNames);
 
-			Addin.Instance.LogDebug($"    GetRenderMaterialsCount2 -> {NumMaterials}");
+			LogDebug($"    GetRenderMaterialsCount2 -> {NumMaterials}");
 
 			if (NumMaterials == 0)
 			{
@@ -202,8 +203,8 @@ namespace DatasmithSolidworks
 			{
 				RenderMaterial RenderMat = ObjMat as RenderMaterial;
 				int NumUsers = RenderMat.GetEntitiesCount();
-				Addin.Instance.LogDebug($"  FileName: {RenderMat.FileName}");
-				Addin.Instance.LogDebug($"    Users({NumUsers}):");
+				LogDebug($"  FileName: {RenderMat.FileName}");
+				LogDebug($"    Users({NumUsers}):");
 
 				if (NumUsers == 0)
 				{
@@ -216,12 +217,12 @@ namespace DatasmithSolidworks
 					switch (ObjUser)
 					{
 						case IPartDoc _:
-							Addin.Instance.LogDebug($"    PartDoc '{(ObjUser as ModelDoc2).GetPathName()}'");
+							LogDebug($"    PartDoc '{(ObjUser as ModelDoc2).GetPathName()}'");
 							PartMaterials.SetPartMaterial(RenderMat, Doc);
 							continue;
 						case IBody2 Body:
 						{
-							Addin.Instance.LogDebug($"    Body '{Body.Name}'");
+							LogDebug($"    Body '{Body.Name}'");
 
 							string BodyId = FPartDocumentTracker.GetBodyPath(InPartDoc as PartDoc, Body);
 							PartMaterials.RegisterMaterial(PartMaterials.BodyMaterialsMap, Doc, RenderMat, BodyId);
@@ -230,20 +231,20 @@ namespace DatasmithSolidworks
 						case IFace2 Face:
 						{
 							uint FaceId = PartMaterials.OwnerDoc.GetFaceId(Face);
-							Addin.Instance.LogDebug($"    Face '{Face.GetFaceId()}'");
+							LogDebug($"    Face '{Face.GetFaceId()}'");
 							PartMaterials.RegisterMaterial(PartMaterials.FaceMaterialsMap, Doc, RenderMat, FaceId.ToString());
 							continue;
 						}
 						case IFeature Feat:
 						{
-							Addin.Instance.LogDebug($"    Feature '{Feat.Name}'");
+							LogDebug($"    Feature '{Feat.Name}'");
 							string FeatureId = FPartDocumentTracker.GetFeaturePath(InPartDoc as PartDoc, Feat);
 							PartMaterials.RegisterMaterial(PartMaterials.FeatureMaterialsMap, Doc, RenderMat, FeatureId);
 							continue;
 						}
 						default:
 						{
-							Addin.Instance.LogDebug($"    <unknown>'");
+							LogDebug($"    <unknown>'");
 							break;
 						}
 					}
@@ -280,11 +281,11 @@ namespace DatasmithSolidworks
 						string BodyId = FPartDocumentTracker.GetBodyPath(InPartDoc as PartDoc, Body);
 						if (PartMaterials.BodyMaterialsMap.TryGetValue(BodyId, out int AssignedMaterialId))
 						{
-							Addin.Instance.LogDebug($"  Body '{Body.Name}' already has material assigned {PartMaterials.GetMaterial(AssignedMaterialId).Name}");
+							LogDebug($"  Body '{Body.Name}' already has material assigned {PartMaterials.GetMaterial(AssignedMaterialId).Name}");
 							continue;
 						}
 
-						Addin.Instance.LogDebug($"  Body '{Body.Name}' has no material assigned, searching material matching its Appearance..");
+						LogDebug($"  Body '{Body.Name}' has no material assigned, searching material matching its Appearance..");
 
 						foreach (KeyValuePair<int, FMaterial> KVP in PartMaterials.GlobalMaterialsMap)
 						{
@@ -292,7 +293,7 @@ namespace DatasmithSolidworks
 
 							if (Material.EqualsAppearance(Appearance))
 							{
-								Addin.Instance.LogDebug($"Matching material for Body '{Body.Name}'  is {Material.Name}");
+								LogDebug($"Matching material for Body '{Body.Name}'  is {Material.Name}");
 
 								PartMaterials.BodyMaterialsMap[BodyId] = KVP.Key; 
 								break;
@@ -337,11 +338,11 @@ namespace DatasmithSolidworks
 						string FeatureId = FPartDocumentTracker.GetFeaturePath(InPartDoc as PartDoc, Feat);
 						if (PartMaterials.FeatureMaterialsMap.TryGetValue(FeatureId, out int AssignedMaterialId))
 						{
-							Addin.Instance.LogDebug($"  Feature '{Feat.Name}' already has material assigned {PartMaterials.GetMaterial(AssignedMaterialId).Name}");
+							LogDebug($"  Feature '{Feat.Name}' already has material assigned {PartMaterials.GetMaterial(AssignedMaterialId).Name}");
 							continue;
 						}
 
-						Addin.Instance.LogDebug($"  Feature '{Feat.Name}' has no material assigned, searching material matching its Appearance..");
+						LogDebug($"  Feature '{Feat.Name}' has no material assigned, searching material matching its Appearance..");
 
 						foreach (KeyValuePair<int, FMaterial> KVP in PartMaterials.GlobalMaterialsMap)
 						{
@@ -349,7 +350,7 @@ namespace DatasmithSolidworks
 
 							if (Material.EqualsAppearance(Appearance))
 							{
-								Addin.Instance.LogDebug($"Matching material for Feature '{Feat.Name}'  is {Material.Name}");
+								LogDebug($"Matching material for Feature '{Feat.Name}'  is {Material.Name}");
 
 								PartMaterials.FeatureMaterialsMap[FeatureId] = KVP.Key; 
 								break;
@@ -364,7 +365,7 @@ namespace DatasmithSolidworks
 
 		public static FObjectMaterials LoadComponentMaterials(FDocumentTracker InComponentOwner, Component2 InComponent, swDisplayStateOpts_e InDisplayState, string[] InDisplayStateNames)
 		{
-			Addin.Instance.LogDebug($"LoadComponentMaterials: '{InComponent.Name2}', '{InComponentOwner.GetPathName()}', {InDisplayState}, {(InDisplayStateNames == null ? "<null>" : string.Join(", ", InDisplayStateNames))}");
+			LogDebug($"LoadComponentMaterials: '{InComponent.Name2}', '{InComponentOwner.GetPathName()}', {InDisplayState}, {(InDisplayStateNames == null ? "<null>" : string.Join(", ", InDisplayStateNames))}");
 
 			ModelDoc2 ComponentDoc = InComponent.GetModelDoc2() as ModelDoc2;
 
@@ -374,9 +375,9 @@ namespace DatasmithSolidworks
 				return null;
 			}
 
-			Addin.Instance.LogDebug($"  GetRenderMaterialsCount2({InDisplayState}, {InDisplayStateNames})");
+			LogDebug($"  GetRenderMaterialsCount2({InDisplayState}, {InDisplayStateNames})");
 			int NumMaterials = InComponent.GetRenderMaterialsCount2((int)InDisplayState, InDisplayStateNames);
-			Addin.Instance.LogDebug($"    GetRenderMaterialsCount2 -> {NumMaterials}");
+			LogDebug($"    GetRenderMaterialsCount2 -> {NumMaterials}");
 
 			if (NumMaterials == 0)
 			{
@@ -387,73 +388,79 @@ namespace DatasmithSolidworks
 
 			FObjectMaterials ComponentMaterials = new FObjectMaterials(InComponentOwner, ComponentDoc as PartDoc, ref InComponentOwner.ExportedMaterialsMap);
 
-			foreach (object ObjMat in ObjMaterials)
+			if (ObjMaterials != null)
 			{
-				RenderMaterial RenderMat = ObjMat as RenderMaterial;
-				int NumUsers = RenderMat.GetEntitiesCount();
+				foreach (object ObjMat in ObjMaterials)
+				{
+					RenderMaterial RenderMat = ObjMat as RenderMaterial;
+					int NumUsers = RenderMat.GetEntitiesCount();
 
-				Addin.Instance.LogDebug($"  FileName: {RenderMat.FileName}");
-				Addin.Instance.LogDebug($"    Users({NumUsers}):");
+					LogDebug($"  FileName: {RenderMat.FileName}");
+					LogDebug($"    Users({NumUsers}):");
 
-				if (NumUsers == 0)
-				{
-					continue;
-				}
+					if (NumUsers == 0)
+					{
+						continue;
+					}
 
-				object[] ObjUsers = RenderMat.GetEntities();
-				foreach (object ObjUser in ObjUsers)
-				{
-					switch (ObjUser)
+					object[] ObjUsers = RenderMat.GetEntities();
+					foreach (object ObjUser in ObjUsers)
 					{
-						case IBody2 Body:
-						{
-							Addin.Instance.LogDebug($"    Body '{Body.Name}'");
-							string BodyId = FPartDocumentTracker.GetBodyPath(ComponentDoc as PartDoc, Body);
-							ComponentMaterials.RegisterMaterial(ComponentMaterials.BodyMaterialsMap, InComponentOwner.SwDoc, RenderMat, BodyId);
-							break;
-						}
-						case IFace2 Face:
-						{
-							// note: plugins sets FaceId on each face to identify per-face materials later in CreateMeshData
-							uint FaceId = ComponentMaterials.OwnerDoc.GetFaceId(Face);
-							Addin.Instance.LogDebug($"    Face '{Face.GetFaceId()}'");
-							ComponentMaterials.RegisterMaterial(ComponentMaterials.FaceMaterialsMap, InComponentOwner.SwDoc, RenderMat, FaceId.ToString());
-							break;
-						}
-						case IFeature Feat:
-						{
-							Addin.Instance.LogDebug($"    Feature '{Feat.Name}'");
-							string FeatureId = FPartDocumentTracker.GetFeaturePath(ComponentDoc as PartDoc, Feat);
-							ComponentMaterials.RegisterMaterial(ComponentMaterials.FeatureMaterialsMap, InComponentOwner.SwDoc, RenderMat, FeatureId);
-							break;
-						}
-						case IPartDoc Doc:
-						{
-							Addin.Instance.LogDebug($"    Part Doc '{(Doc as ModelDoc2).GetPathName()}'");
-							ComponentMaterials.SetPartMaterial(RenderMat, InComponentOwner.SwDoc);
-							break;
-						}
-						default:
+						switch (ObjUser)
 						{
-							Addin.Instance.LogDebug($"    <unknown>'");
-							break;
+							case IBody2 Body:
+							{
+								LogDebug($"    Body '{Body.Name}'");
+								string BodyId = FPartDocumentTracker.GetBodyPath(ComponentDoc as PartDoc, Body);
+								ComponentMaterials.RegisterMaterial(ComponentMaterials.BodyMaterialsMap,
+									InComponentOwner.SwDoc, RenderMat, BodyId);
+								break;
+							}
+							case IFace2 Face:
+							{
+								// note: plugins sets FaceId on each face to identify per-face materials later in CreateMeshData
+								uint FaceId = ComponentMaterials.OwnerDoc.GetFaceId(Face);
+								LogDebug($"    Face '{Face.GetFaceId()}'");
+								ComponentMaterials.RegisterMaterial(ComponentMaterials.FaceMaterialsMap,
+									InComponentOwner.SwDoc, RenderMat, FaceId.ToString());
+								break;
+							}
+							case IFeature Feat:
+							{
+								LogDebug($"    Feature '{Feat.Name}'");
+								string FeatureId = FPartDocumentTracker.GetFeaturePath(ComponentDoc as PartDoc, Feat);
+								ComponentMaterials.RegisterMaterial(ComponentMaterials.FeatureMaterialsMap,
+									InComponentOwner.SwDoc, RenderMat, FeatureId);
+								break;
+							}
+							case IPartDoc Doc:
+							{
+								LogDebug($"    Part Doc '{(Doc as ModelDoc2).GetPathName()}'");
+								ComponentMaterials.SetPartMaterial(RenderMat, InComponentOwner.SwDoc);
+								break;
+							}
+							default:
+							{
+								LogDebug($"    <unknown>'");
+								break;
+							}
 						}
 					}
 				}
-			}
+			}			
 
 			return ComponentMaterials;
 		}
 
 		public static ConcurrentDictionary<FComponentName, FObjectMaterials> LoadAssemblyMaterials(FAssemblyDocumentTracker InAsmDoc, HashSet<FComponentName> InComponentsSet, swDisplayStateOpts_e InDisplayState, string[] InDisplayStateNames)
 		{
-			Addin.Instance.LogDebug($"LoadAssemblyMaterials: {InAsmDoc.GetPathName()}, {InDisplayState}, [{(InDisplayStateNames == null ? "" : string.Join(",", InDisplayStateNames))}]");
-			Addin.Instance.LogDebug($"  for components: [{string.Join(", ", InComponentsSet.OrderBy(Name => Name.GetString()))}]");
+			LogDebug($"LoadAssemblyMaterials: {InAsmDoc.GetPathName()}, {InDisplayState}, [{(InDisplayStateNames == null ? "" : string.Join(",", InDisplayStateNames))}]");
+			LogDebug($"  for components: [{string.Join(", ", InComponentsSet.OrderBy(Name => Name.GetString()))}]");
 
 			IModelDocExtension Ext = InAsmDoc.SwDoc.Extension;
 			int NumMaterials = Ext.GetRenderMaterialsCount2((int)InDisplayState, InDisplayStateNames);
 
-			Addin.Instance.LogDebug($"  Assembly materials={NumMaterials}");
+			LogDebug($"  Assembly materials={NumMaterials}");
 
 			ConcurrentDictionary<FComponentName, FObjectMaterials> DocMaterials = new ConcurrentDictionary<FComponentName, FObjectMaterials>();
 
@@ -465,11 +472,11 @@ namespace DatasmithSolidworks
 				{
 					RenderMaterial RenderMat = ObjMat as RenderMaterial;
 					
-					Addin.Instance.LogDebug($"    Material: '{RenderMat.FileName}' ");
+					LogDebug($"    Material: '{RenderMat.FileName}' ");
 
 					int NumUsers = RenderMat.GetEntitiesCount();
 
-					Addin.Instance.LogDebug($"    NumUsers: '{NumUsers}' ");
+					LogDebug($"    NumUsers: '{NumUsers}' ");
 
 					if (NumUsers == 0)
 					{
@@ -481,7 +488,7 @@ namespace DatasmithSolidworks
 					{
 						if (ObjUser is IEntity Entity)
 						{
-							Addin.Instance.LogDebug($"    User: type '{Entity.GetType()}' ");
+							LogDebug($"    User: type '{Entity.GetType()}' ");
 							
 						}
 						
@@ -489,15 +496,15 @@ namespace DatasmithSolidworks
 						{
 							case Component2 Comp:
 							{
-								Addin.Instance.LogDebug($"      Component '{Comp.Name2}' ");
+								LogDebug($"      Component '{Comp.Name2}' ");
 								if (Comp.GetModelDoc2() is PartDoc CompDoc)
 								{
-									Addin.Instance.LogDebug($"        registering material for Part Component in DocMaterials");
+									LogDebug($"        registering material for Part Component in DocMaterials");
 									FObjectMaterials ComponentMaterials = new FObjectMaterials(InAsmDoc, CompDoc, ref InAsmDoc.ExportedMaterialsMap);
 
-									Addin.Instance.LogIndent();
+									LogIndent();
 									ComponentMaterials.SetComponentMaterial(RenderMat, InAsmDoc.SwDoc);
-									Addin.Instance.LogDedent();
+									LogDedent();
 
 									DocMaterials[new FComponentName(Comp)] = ComponentMaterials;
 								}
@@ -505,11 +512,11 @@ namespace DatasmithSolidworks
 							}
 							case IPartDoc Doc:
 							{
-								Addin.Instance.LogDebug($"      Part Doc at '{(Doc as ModelDoc2).GetPathName()}' ");
+								LogDebug($"      Part Doc at '{(Doc as ModelDoc2).GetPathName()}' ");
 
-								Addin.Instance.LogIndent();
+								LogIndent();
 								FObjectMaterials PartMaterials = LoadPartMaterials(InAsmDoc, Doc as PartDoc, InDisplayState, InDisplayStateNames);
-								Addin.Instance.LogDedent();
+								LogDedent();
 
 								if (PartMaterials != null)
 								{
@@ -524,30 +531,30 @@ namespace DatasmithSolidworks
 				}
 			}
 
-			Addin.Instance.LogDebug($"Load Each Component's materials");
-			Addin.Instance.LogIndent();
+			LogDebug($"Load Each Component's materials");
+			LogIndent();
 			foreach(FComponentName CompName in InComponentsSet.OrderBy(Name => Name.GetString()))
 			{
 				Component2 Comp = null;
 
-				Addin.Instance.LogDebug($"Component '{CompName}' - adding DocMaterials");
-				Addin.Instance.LogIndent();
+				LogDebug($"Component '{CompName}' - adding DocMaterials");
+				LogIndent();
 
 				if (InAsmDoc.SyncState.ExportedComponentsMap.TryGetValue(CompName, out Comp))
 				{
 					if (!DocMaterials.ContainsKey(CompName))
 					{
-						Addin.Instance.LogDebug($"Adding to DocMaterials");
+						LogDebug($"Adding to DocMaterials");
 
-						Addin.Instance.LogIndent();
+						LogIndent();
 						FObjectMaterials ComponentMaterials = FObjectMaterials.LoadComponentMaterials(InAsmDoc, Comp, swDisplayStateOpts_e.swThisDisplayState, null);
-						Addin.Instance.LogDedent();
+						LogDedent();
 
 						if (ComponentMaterials == null && Comp.GetModelDoc2() is PartDoc)
 						{
-							Addin.Instance.LogIndent();
+							LogIndent();
 							ComponentMaterials = FObjectMaterials.LoadPartMaterials(InAsmDoc, Comp.GetModelDoc2() as PartDoc, swDisplayStateOpts_e.swThisDisplayState, null);
-							Addin.Instance.LogDedent();
+							LogDedent();
 						}
 						if (ComponentMaterials != null)
 						{
@@ -556,17 +563,17 @@ namespace DatasmithSolidworks
 					}
 					else
 					{
-						Addin.Instance.LogDebug($"...already in DocMaterials");
+						LogDebug($"...already in DocMaterials");
 					}
 				}
 				else
 				{
-					Addin.Instance.LogDebug($"...not in ExportedComponentsMap");
+					LogDebug($"...not in ExportedComponentsMap");
 				}
 
-				Addin.Instance.LogDedent();
+				LogDedent();
 			}
-			Addin.Instance.LogDedent();
+			LogDedent();
 
 			return DocMaterials.Count > 0 ? DocMaterials : null;
 		}
diff --git a/Engine/Source/Programs/Enterprise/Datasmith/DatasmithSolidworksExporter/Private/PartDocument.cs b/Engine/Source/Programs/Enterprise/Datasmith/DatasmithSolidworksExporter/Private/PartDocument.cs
--- a/Engine/Source/Programs/Enterprise/Datasmith/DatasmithSolidworksExporter/Private/PartDocument.cs
+++ b/Engine/Source/Programs/Enterprise/Datasmith/DatasmithSolidworksExporter/Private/PartDocument.cs
@@ -6,10 +6,10 @@ using System;
 using System.Collections.Concurrent;
 using System.Collections.Generic;
 using System.IO;
-using System.Security.Cryptography;
 using DatasmithSolidworks.Names;
-using static System.Windows.Forms.VisualStyles.VisualStyleElement;
 using static DatasmithSolidworks.FDatasmithExporter;
+using System.Diagnostics;
+using System.Linq;
 
 namespace DatasmithSolidworks
 {
@@ -21,6 +21,8 @@ namespace DatasmithSolidworks
 		private FMeshName MeshName = new FMeshName();
 		public readonly FPartDocument PartDocument;
 
+		private FMeshes Meshes;
+
 		public FObjectMaterials ExportedPartMaterials { get; set; }
 
 		public string PathName { get; private set; }
@@ -56,6 +58,7 @@ namespace DatasmithSolidworks
 				}
 				catch
 				{
+					Debug.Assert(false);
 				}
 			}
 
@@ -76,6 +79,7 @@ namespace DatasmithSolidworks
 				}
 				catch
 				{
+					Debug.Assert(false);
 				}
 			}
 
@@ -93,34 +97,23 @@ namespace DatasmithSolidworks
 			}
 		}
 
-		public override void PreExport(FMeshes Meshes, bool bConfigurations)
+		public override FMeshes GetMeshes(string ActiveConfigName)
 		{
-			if (!bConfigurations)
-			{
-				// Not exporting configurations - just load materials here
-				ExportedPartMaterials =
-					FObjectMaterials.LoadPartMaterials(this, SwPartDoc, swDisplayStateOpts_e.swThisDisplayState, null);
-			}
+			return Meshes ?? (Meshes = new FMeshes(ActiveConfigName));
 		}
 
 		public override void ExportToDatasmithScene(FConfigurationExporter ConfigurationExporter)
 		{
-			ExportToDatasmithScene(ConfigurationExporter.Meshes);
-		}
+			FMeshes Meshes = ConfigurationExporter.Meshes;
 
-		public override void ExportToDatasmithScene(FMeshes Meshes)
-		{
 			string PartName = Path.GetFileNameWithoutExtension(PathName);
 
-			SetExportStatus($"{PartName} Materials");
-
 			SetExportStatus($"{PartName} Meshes");
 
 			ConfigurationManager ConfigManager = (SwPartDoc as ModelDoc2).ConfigurationManager;
 			if (bHasConfigurations)
 			{
 				string ActiveConfigurationName = ConfigManager.ActiveConfiguration.Name;
-				FConfigurationExporter ConfigurationExporter = new FConfigurationExporter(Meshes, null, ActiveConfigurationName);
 
 				foreach (string ConfigurationName in SwDoc.GetConfigurationNames())
 				{
@@ -146,12 +139,18 @@ namespace DatasmithSolidworks
 						{
 							ComponentName = RootComponentName,
 							MeshName = MeshNameForConfiguration,
-							ActorName = ActorName,
 							MeshData = MeshData
 						}
 					};
 
-					Exporter.ExportMeshes(MeshExportInfos, out List<FMeshExportInfo> OutCreatedMeshes);
+					List<FMeshExportInfo> OutCreatedMeshes = Exporter.ExportMeshes(MeshExportInfos).ToList();
+
+					// OutCreatedMeshes might be empty(e.g. mesh is empty it won't be created)
+					foreach (FMeshExportInfo Info in OutCreatedMeshes)
+					{
+						Exporter.AddMesh(Info);
+					}
+
 					ExportMaterials();
 					Exporter.AssignMaterialsToDatasmithMeshes(OutCreatedMeshes);
 					bool bHasMesh = OutCreatedMeshes.Count > 0;
@@ -189,12 +188,18 @@ namespace DatasmithSolidworks
 						{
 							ComponentName = RootComponentName,
 							MeshName = MeshName,
-							ActorName = ActorName,
 							MeshData = MeshData
 						}
 					};
 
-					Exporter.ExportMeshes(MeshExportInfos, out List<FMeshExportInfo> OutCreatedMeshes);
+					List<FMeshExportInfo> OutCreatedMeshes = Exporter.ExportMeshes(MeshExportInfos).ToList();
+
+					// OutCreatedMeshes might be empty(e.g. mesh is empty it won't be created)
+					foreach (FMeshExportInfo Info in OutCreatedMeshes)
+					{
+						Exporter.AddMesh(Info);
+					}
+
 					ExportMaterials();
 					Exporter.AssignMaterialsToDatasmithMeshes(OutCreatedMeshes);
 					bHasMesh = OutCreatedMeshes.Count > 0;
@@ -238,6 +243,14 @@ namespace DatasmithSolidworks
 		{
 		}
 
+		public override void ReleaseComponentMeshes(FComponentName CompName)
+		{
+		}
+
+		public override void CleanupComponentMeshes()
+		{
+		}
+
 		public override FMeshData ExtractComponentMeshData(Component2 Comp)
 		{
 			return ExtractPartMeshData();  // Part document has single component
@@ -293,6 +306,7 @@ namespace DatasmithSolidworks
 			SwPartDoc.BodyVisibleChangeNotify += new DPartDocEvents_BodyVisibleChangeNotifyEventHandler(OnPartBodyVisibleChangeNotify);
 			SwPartDoc.FileSaveAsNotify2 += new DPartDocEvents_FileSaveAsNotify2EventHandler(OnPartFileSaveAsNotify2);
 			SwPartDoc.FileSavePostNotify += new DPartDocEvents_FileSavePostNotifyEventHandler(OnPartFileSavePostNotify);
+
 		}
 
 		public void Destroy()
@@ -321,7 +335,6 @@ namespace DatasmithSolidworks
 			SwPartDoc.FileSavePostNotify -= new DPartDocEvents_FileSavePostNotifyEventHandler(OnPartFileSavePostNotify);
 		}
 
-
 		int OnPartDocumentDestroyNotify2(int DestroyType)
 		{
 			Addin.Instance.CloseDocument(DocId);
@@ -516,6 +529,11 @@ namespace DatasmithSolidworks
 			return DocumentTracker.GetDirty();
 		}
 
+		public FDocumentTracker GetTracker()
+		{
+			return DocumentTracker;
+		}
+
 		public void SetDirty(bool bInDirty)
 		{
 			DocumentTracker.SetDirty(bInDirty);
diff --git a/Engine/Source/Programs/Enterprise/Datasmith/DatasmithSolidworksExporter/Private/StripGeometry.cs b/Engine/Source/Programs/Enterprise/Datasmith/DatasmithSolidworksExporter/Private/StripGeometry.cs
--- a/Engine/Source/Programs/Enterprise/Datasmith/DatasmithSolidworksExporter/Private/StripGeometry.cs
+++ b/Engine/Source/Programs/Enterprise/Datasmith/DatasmithSolidworksExporter/Private/StripGeometry.cs
@@ -2,6 +2,7 @@
 
 using System.Collections.Concurrent;
 using System.Collections.Generic;
+using System.Diagnostics;
 using System.Runtime.InteropServices;
 using System.Threading.Tasks;
 
@@ -73,7 +74,10 @@ namespace DatasmithSolidworks
 					Body.ExtractGeometry(Chunks);
 				}
 			}
-			catch { }
+			catch
+			{
+				Debug.Assert(false);
+			}
 
 			return Chunks;
 		}
@@ -111,7 +115,10 @@ namespace DatasmithSolidworks
 					MeshData = FMeshData.Create(Chunks);
 				}
 			}
-			catch { }
+			catch
+			{
+				Debug.Assert(false);
+			}
 
 			return MeshData;
 		}
