[Solidworks] Fixed child components variant transforms

- relative transform was canculated incorrectly
- typesafety for converter trasform
- added log for transforms debugging

diff --git a/Engine/Source/Programs/Enterprise/Datasmith/DatasmithSolidworksExporter/Private/Addin.cs b/Engine/Source/Programs/Enterprise/Datasmith/DatasmithSolidworksExporter/Private/Addin.cs
--- a/Engine/Source/Programs/Enterprise/Datasmith/DatasmithSolidworksExporter/Private/Addin.cs
+++ b/Engine/Source/Programs/Enterprise/Datasmith/DatasmithSolidworksExporter/Private/Addin.cs
@@ -652,7 +652,11 @@ namespace DatasmithSolidworks
 
 		#endregion
 
-		public static V FindOrAdd<K, V>(Dictionary<K, V> Map, K Key) 
+	}
+
+	public static class DictionaryExtensions
+	{
+		public static V FindOrAdd<K, V>(this Dictionary<K, V> Map, K Key) 
 			where V : new()
 		{
 			if (!Map.TryGetValue(Key, out V Value))
@@ -662,6 +666,5 @@ namespace DatasmithSolidworks
 			}
 			return Value;
 		}
-
 	}
 }
diff --git a/Engine/Source/Programs/Enterprise/Datasmith/DatasmithSolidworksExporter/Private/Animation.cs b/Engine/Source/Programs/Enterprise/Datasmith/DatasmithSolidworksExporter/Private/Animation.cs
--- a/Engine/Source/Programs/Enterprise/Datasmith/DatasmithSolidworksExporter/Private/Animation.cs
+++ b/Engine/Source/Programs/Enterprise/Datasmith/DatasmithSolidworksExporter/Private/Animation.cs
@@ -40,7 +40,7 @@ namespace DatasmithSolidworks
 				Key.Step = InStep;
 				Key.Time = InTime;
 				Key.Owner = this;
-				Key.LocalMatrix = new FMatrix4(MathUtils.ConvertFromSolidworksTransform(InLocalTransform, 100f /*GeomScale*/));
+				Key.LocalMatrix = new FMatrix4(MathUtils.ConvertFromSolidworksTransform(InLocalTransform, 100f /*GeomScale*/).Matrix);
 
 				Keyframes.Add(Key);
 				Id2Key.Add(Key.Step, Key);
diff --git a/Engine/Source/Programs/Enterprise/Datasmith/DatasmithSolidworksExporter/Private/AssemblyDocument.cs b/Engine/Source/Programs/Enterprise/Datasmith/DatasmithSolidworksExporter/Private/AssemblyDocument.cs
--- a/Engine/Source/Programs/Enterprise/Datasmith/DatasmithSolidworksExporter/Private/AssemblyDocument.cs
+++ b/Engine/Source/Programs/Enterprise/Datasmith/DatasmithSolidworksExporter/Private/AssemblyDocument.cs
@@ -32,8 +32,8 @@ namespace DatasmithSolidworks
 
 			public Dictionary<FComponentName, string> ComponentToPartMap = new Dictionary<FComponentName, string>();
 
-			public Dictionary<FComponentName, float[]> ComponentsTransformsMap =
-				new Dictionary<FComponentName, float[]>();
+			public Dictionary<FComponentName, FConvertedTransform> ComponentsTransformsMap =
+				new Dictionary<FComponentName, FConvertedTransform>();
 
 			/// All updated(not dirty) components. Updated component will be removed from dirty on export end
 			public HashSet<FComponentName> CleanComponents = new HashSet<FComponentName>();
@@ -185,8 +185,8 @@ namespace DatasmithSolidworks
 
 		public override void AddMeshForComponent(FComponentName ComponentName, FMeshName MeshName)
 		{
-			FindOrAdd(SyncState.ComponentNameToMeshNameMap, ComponentName).Add(MeshName);
-			FindOrAdd(SyncState.ComponentsForMesh, MeshName).Add(ComponentName);
+			SyncState.ComponentNameToMeshNameMap.FindOrAdd(ComponentName).Add(MeshName);
+			SyncState.ComponentsForMesh.FindOrAdd(MeshName).Add(ComponentName);
 		}
 
 		public override FObjectMaterials GetComponentMaterials(Component2 Comp)
@@ -269,7 +269,7 @@ namespace DatasmithSolidworks
 			return bHasDirtyComponents || bHasDirtyMaterials;
 		}
 
-		public float[] GetComponentDatasmithTransform(Component2 InComponent)
+		public FConvertedTransform GetComponentDatasmithTransform(Component2 InComponent)
 		{
 			MathTransform ComponentTransform = InComponent.GetTotalTransform(true);
 			if (ComponentTransform == null)
@@ -277,24 +277,34 @@ namespace DatasmithSolidworks
 				ComponentTransform = InComponent.Transform2;
 			}
 
-			float[] DatasmithTransform = null;
+			FConvertedTransform DatasmithTransform;
 			if (ComponentTransform != null)
 			{
 				DatasmithTransform = MathUtils.ConvertFromSolidworksTransform(ComponentTransform, 100f /*GeomScale*/);
 			}
 			else
 			{
-				DatasmithTransform = MathUtils.TransformIdentity();
+				DatasmithTransform = FConvertedTransform.Identity();
 			}
 
 			return DatasmithTransform;
 		}
 
 		// Get component transform in specified configuration
-		private float[] GetComponentDatasmithTransform(FConfigurationTree.FComponentTreeNode InNode,
+		private FConvertedTransform GetComponentDatasmithTransform(FConfigurationTree.FComponentTreeNode InNode,
 			FConfigurationTree.FComponentConfig ComponentConfig)
 		{
-			return ComponentConfig?.Transform ?? InNode.CommonConfig.Transform ?? MathUtils.TransformIdentity();
+			if (ComponentConfig != null && ComponentConfig.Transform.IsValid())
+			{
+				return ComponentConfig.Transform;
+			}
+
+			if (InNode.CommonConfig != null && InNode.CommonConfig.Transform.IsValid())
+			{
+				return InNode.CommonConfig.Transform;
+			}
+
+			return FConvertedTransform.Identity();
 		}
 
 		private void ExportComponentRecursive(string ActiveConfigName,
@@ -304,7 +314,7 @@ namespace DatasmithSolidworks
 			SetExportStatus(InNode.ComponentName.GetString());
 
 			FConfigurationTree.FComponentConfig ActiveConfig = InNode.Configurations?.Find(Config => Config.ConfigName == ActiveConfigName);
-			float[] Transform = GetComponentDatasmithTransform(InNode, ActiveConfig);
+			FConvertedTransform Transform = GetComponentDatasmithTransform(InNode, ActiveConfig);
 			SyncState.ComponentsTransformsMap[InNode.ComponentName] = Transform;
 
 			if (InNode.bGeometrySame)
@@ -416,7 +426,7 @@ namespace DatasmithSolidworks
 			bool bHasDirtyTransform = false;
 			if (SyncState.ComponentsTransformsMap.ContainsKey(InComponent.ComponentName))
 			{
-				float[] ComponentTm = GetComponentDatasmithTransform(InComponent, ActiveComponentConfig);
+				FConvertedTransform ComponentTm = GetComponentDatasmithTransform(InComponent, ActiveComponentConfig);
 				bHasDirtyTransform =
 					!MathUtils.TransformsAreEqual(SyncState.ComponentsTransformsMap[InComponent.ComponentName],
 						ComponentTm);
@@ -687,10 +697,10 @@ namespace DatasmithSolidworks
 			foreach (var KVP in AsmDocumentTracker.SyncState.ExportedComponentsMap)
 			{
 				Component2 Comp = KVP.Value;
-				float[] PrevCompTransform = null;
+				FConvertedTransform PrevCompTransform;
 				if (AsmDocumentTracker.SyncState.ComponentsTransformsMap.TryGetValue(new FComponentName(Comp), out PrevCompTransform))
 				{
-					float[] CompTransform = AsmDocumentTracker.GetComponentDatasmithTransform(Comp);
+					FConvertedTransform CompTransform = AsmDocumentTracker.GetComponentDatasmithTransform(Comp);
 
 					if (!MathUtils.TransformsAreEqual(CompTransform, PrevCompTransform))
 					{
diff --git a/Engine/Source/Programs/Enterprise/Datasmith/DatasmithSolidworksExporter/Private/ConfigurationExporter.cs b/Engine/Source/Programs/Enterprise/Datasmith/DatasmithSolidworksExporter/Private/ConfigurationExporter.cs
--- a/Engine/Source/Programs/Enterprise/Datasmith/DatasmithSolidworksExporter/Private/ConfigurationExporter.cs
+++ b/Engine/Source/Programs/Enterprise/Datasmith/DatasmithSolidworksExporter/Private/ConfigurationExporter.cs
@@ -18,7 +18,7 @@ namespace DatasmithSolidworks
 	{
 		public string Name;
 		public bool bIsDisplayStateConfiguration = false;
-		public Dictionary<FComponentName, float[]> ComponentTransform = new Dictionary<FComponentName, float[]>();  // Relative(to parent) transform
+		public Dictionary<FComponentName, FConvertedTransform> ComponentTransform = new Dictionary<FComponentName, FConvertedTransform>();  // Relative(to parent) transform
 		public Dictionary<FComponentName, bool> ComponentVisibility = new Dictionary<FComponentName, bool>();
 		public Dictionary<FComponentName, FObjectMaterials> ComponentMaterials = new Dictionary<FComponentName, FObjectMaterials>();
 
@@ -683,19 +683,23 @@ namespace DatasmithSolidworks
 			if (ComponentTransform != null)
 			{
 				NewNode.CommonConfig.Transform = MathUtils.ConvertFromSolidworksTransform(ComponentTransform, 100f /*GeomScale*/);
+				LogDebug($" Transform: {NewNode.CommonConfig.Transform}");
 
-				if (InParentNode.CommonConfig.Transform != null)
+				if (InParentNode.CommonConfig.Transform.IsValid())
 				{
 					// Convert transform to parent space (original transform value fetched from Solidworks
 					// is in the root component's space). Datasmith wants relative transform for variants.
-					FMatrix4 ParentTransform = new FMatrix4(InParentNode.CommonConfig.Transform);
+					FMatrix4 ParentTransform = new FMatrix4(InParentNode.CommonConfig.Transform.Matrix);
 					FMatrix4 InverseParentTransform = ParentTransform.Inverse();
-					NewNode.CommonConfig.RelativeTransform = FMatrix4.FMatrix4x4Multiply(InverseParentTransform, NewNode.CommonConfig.Transform);
+					NewNode.CommonConfig.RelativeTransform = new FConvertedTransform(FMatrix4.FMatrix4x4Multiply(NewNode.CommonConfig.Transform.Matrix, InverseParentTransform));
+
 				}
 				else
 				{
 					NewNode.CommonConfig.RelativeTransform = NewNode.CommonConfig.Transform;
 				}
+				LogDebug($"   ParentTransform: {InParentNode.CommonConfig.Transform}");
+				LogDebug($"   RelativeTransform: {NewNode.CommonConfig.RelativeTransform}");
 			}
 
 			// Process children components
@@ -776,8 +780,8 @@ namespace DatasmithSolidworks
 				return;
 			}
 
-			FindOrAdd(ActorsForMesh, MeshName).Add(ActorName);
-			FindOrAdd(MeshesForComponent, ComponentName).Add(MeshName);
+			ActorsForMesh.FindOrAdd(MeshName).Add(ActorName);
+			MeshesForComponent.FindOrAdd(ComponentName).Add(MeshName);
 		}
 	}
 }
\ No newline at end of file
diff --git a/Engine/Source/Programs/Enterprise/Datasmith/DatasmithSolidworksExporter/Private/ConfigurationTree.cs b/Engine/Source/Programs/Enterprise/Datasmith/DatasmithSolidworksExporter/Private/ConfigurationTree.cs
--- a/Engine/Source/Programs/Enterprise/Datasmith/DatasmithSolidworksExporter/Private/ConfigurationTree.cs
+++ b/Engine/Source/Programs/Enterprise/Datasmith/DatasmithSolidworksExporter/Private/ConfigurationTree.cs
@@ -9,6 +9,37 @@ using DatasmithSolidworks.Names;
 
 namespace DatasmithSolidworks
 {
+	public struct FConvertedTransform
+	{
+		public float[] Matrix;
+
+		public FConvertedTransform(float[] Floats)
+		{
+			Matrix = Floats;
+		}
+
+		public bool IsValid()
+		{
+			return Matrix != null;
+		}
+
+		public override string ToString()
+		{
+			return (Matrix==null ? "null" : string.Join(", ", Matrix));
+		}
+
+		public static FConvertedTransform Identity()
+		{
+			return new FConvertedTransform(new[]
+			{
+				1.0f, 0.0f, 0.0f, 0.0f,
+				0.0f, 1.0f, 0.0f, 0.0f,
+				0.0f, 0.0f, 1.0f, 0.0f,
+				0.0f, 0.0f, 0.0f, 1.0f
+			});
+		}
+	}
+
 	public class FConfigurationTree
 	{
 		// Component state for a single configuration
@@ -21,8 +52,8 @@ namespace DatasmithSolidworks
 			public bool bSuppressed = false;
 			
 			// Node transform. Null value if transform is not changed in this configuration.
-			public float[] Transform = null;
-			public float[] RelativeTransform = null;
+			public FConvertedTransform Transform;
+			public FConvertedTransform RelativeTransform;
 
 			// Null value if configuration doesn't override the material.
 			public FObjectMaterials Materials = null;
@@ -255,15 +286,15 @@ namespace DatasmithSolidworks
 			if (InNode.Configurations != null && InNode.Configurations.Count > 0)
 			{
 				// Check transform
-				float[] Transform = InNode.Configurations[0].Transform;
-				float[] RelativeTransform = InNode.Configurations[0].RelativeTransform;
+				FConvertedTransform Transform = InNode.Configurations[0].Transform;
+				FConvertedTransform RelativeTransform = InNode.Configurations[0].RelativeTransform;
 				bool bAllTransformsAreSame = true;
-				if (RelativeTransform != null)
+				if (RelativeTransform.IsValid())
 				{
 					// There could be components without a transform, so we're checking if for null
 					for (int Idx = 1; Idx < InNode.Configurations.Count; Idx++)
 					{
-						if (!InNode.Configurations[Idx].RelativeTransform.SequenceEqual(RelativeTransform))
+						if (!MathUtils.TransformsAreEqual(InNode.Configurations[Idx].RelativeTransform, RelativeTransform))
 						{
 							bAllTransformsAreSame = false;
 							break;
@@ -275,8 +306,8 @@ namespace DatasmithSolidworks
 						InNode.CommonConfig.RelativeTransform = RelativeTransform;
 						foreach (FComponentConfig Config in InNode.Configurations)
 						{
-							Config.RelativeTransform = null;
-							Config.Transform = null;
+							Config.RelativeTransform = new FConvertedTransform();
+							Config.Transform = new FConvertedTransform();
 						}
 					}
 				}
@@ -392,7 +423,7 @@ namespace DatasmithSolidworks
 			if (NodeConfig != null)
 			{
 				// Only add variant information if attribute is not the same in all configurations
-				if (NodeConfig.RelativeTransform != null)
+				if (NodeConfig.RelativeTransform.IsValid())
 				{
 					OutData.ComponentTransform.Add(InNode.ComponentName, NodeConfig.RelativeTransform);
 				}
diff --git a/Engine/Source/Programs/Enterprise/Datasmith/DatasmithSolidworksExporter/Private/DatasmithExporter.cs b/Engine/Source/Programs/Enterprise/Datasmith/DatasmithSolidworksExporter/Private/DatasmithExporter.cs
--- a/Engine/Source/Programs/Enterprise/Datasmith/DatasmithSolidworksExporter/Private/DatasmithExporter.cs
+++ b/Engine/Source/Programs/Enterprise/Datasmith/DatasmithSolidworksExporter/Private/DatasmithExporter.cs
@@ -148,7 +148,7 @@ namespace DatasmithSolidworks
 		public FActorName Name;
 		public FActorName ParentName;
 		public FMeshName MeshName;
-		public float[] Transform;
+		public FConvertedTransform Transform;
 		public bool bVisible;
 
 		public override string ToString()
@@ -660,7 +660,7 @@ namespace DatasmithSolidworks
 					FDatasmithFacadeActorBinding Binding = GetActorBinding(TransformMap.Key, Variant);
 					if (Binding != null)
 					{
-						Binding.AddRelativeTransformCapture(TransformMap.Value);
+						Binding.AddRelativeTransformCapture(TransformMap.Value.Matrix);
 					}
 				}
 			}
@@ -1118,14 +1118,13 @@ namespace DatasmithSolidworks
 			return Result;
 		}
 
-		private FMatrix4 AdjustTransformForDatasmith(float[] InXForm)
+		private FMatrix4 AdjustTransformForDatasmith(FConvertedTransform InXForm)
 		{
 			FMatrix4 RotMatrix = FMatrix4.FromRotationX(-90f);
-			if (InXForm != null)
+			if (InXForm.IsValid())
 			{
-				FMatrix4 Mat = new FMatrix4(InXForm);
-				Mat = Mat * RotMatrix;
-				return Mat;
+				FMatrix4 Mat = new FMatrix4(InXForm.Matrix);
+				return Mat * RotMatrix;
 			}
 			return RotMatrix;
 		}
diff --git a/Engine/Source/Programs/Enterprise/Datasmith/DatasmithSolidworksExporter/Private/MathUtils.cs b/Engine/Source/Programs/Enterprise/Datasmith/DatasmithSolidworksExporter/Private/MathUtils.cs
--- a/Engine/Source/Programs/Enterprise/Datasmith/DatasmithSolidworksExporter/Private/MathUtils.cs
+++ b/Engine/Source/Programs/Enterprise/Datasmith/DatasmithSolidworksExporter/Private/MathUtils.cs
@@ -52,7 +52,7 @@ namespace DatasmithSolidworks
 		// 1 5 9  13 (ty)
 		// 2 6 10 14 (tz)
 		// 3(0.0) 7(0.0) 11(0.0) 15(1.0)
-		public static float[] ConvertFromSolidworksTransform(MathTransform InSwTransform, float InGeomScale)
+		public static FConvertedTransform ConvertFromSolidworksTransform(MathTransform InSwTransform, float InGeomScale)
 		{
 			float[] Result = new float[16];
 			double[] InputMatrix = (double[])InSwTransform.ArrayData;
@@ -74,21 +74,10 @@ namespace DatasmithSolidworks
 			Result[13] = (float)InputMatrix[10] * InGeomScale;
 			Result[14] = (float)InputMatrix[11] * InGeomScale;
 
-			return Result;
+			return new FConvertedTransform(Result);
 		}
 
-		public static float[] TransformIdentity()
-		{
-			return new[]
-			{
-				1.0f, 0.0f, 0.0f, 0.0f,
-				0.0f, 1.0f, 0.0f, 0.0f,
-				0.0f, 0.0f, 1.0f, 0.0f,
-				0.0f, 0.0f, 0.0f, 1.0f
-			};
-		}
-
-		public static float[] LookAt(FVec3 InDirection, FVec3 InFromPoint, float InGeomScale)
+		public static FConvertedTransform LookAt(FVec3 InDirection, FVec3 InFromPoint, float InGeomScale)
 		{
 			float[] Ret = new float[16];
 
@@ -123,10 +112,10 @@ namespace DatasmithSolidworks
 				}
 			}
 
-			return Ret;
+			return new FConvertedTransform(Ret);
 		}
 
-		public static float[] Translation(FVec3 InTranslation, float InGeomScale)
+		public static FConvertedTransform Translation(FVec3 InTranslation, float InGeomScale)
 		{
 			float[] Ret = new float[16];
 
@@ -151,16 +140,16 @@ namespace DatasmithSolidworks
 			Ret[13] = InTranslation.Y * InGeomScale;
 			Ret[14] = InTranslation.Z * InGeomScale;
 
-			return Ret;
+			return new FConvertedTransform(Ret);
 		}
 
-		public static bool TransformsAreEqual(float[] InTransformA, float[] InTransformB)
+		public static bool TransformsAreEqual(FConvertedTransform InTransformA, FConvertedTransform InTransformB)
 		{
-			Debug.Assert(InTransformA.Length == 16 && InTransformA.Length == InTransformB.Length);
+			Debug.Assert(InTransformA.Matrix.Length == 16 && InTransformA.Matrix.Length == InTransformB.Matrix.Length);
 
 			for (int Idx = 0; Idx < 16; ++Idx)
 			{
-				if (!Equals(InTransformA[Idx], InTransformB[Idx]))
+				if (!Equals(InTransformA.Matrix[Idx], InTransformB.Matrix[Idx]))
 				{
 					return false;
 				}
