[Solidworks] Exploded Views exported for each of configurations as additional variants

diff --git a/Engine/Source/Programs/Enterprise/Datasmith/DatasmithSolidworksExporter/Private/Addin.cs b/Engine/Source/Programs/Enterprise/Datasmith/DatasmithSolidworksExporter/Private/Addin.cs
--- a/Engine/Source/Programs/Enterprise/Datasmith/DatasmithSolidworksExporter/Private/Addin.cs
+++ b/Engine/Source/Programs/Enterprise/Datasmith/DatasmithSolidworksExporter/Private/Addin.cs
@@ -650,6 +650,29 @@ namespace DatasmithSolidworks
 			Instance.DebugLog.Dedent();
 		}
 
+		public static IDisposable LogScopedIndent()
+		{
+#if DatasmithSolidworksDebugOutput
+			return new FLogScopedIndent();
+#else
+			return null;
+#endif
+		}
+
+		private class FLogScopedIndent : IDisposable
+		{
+			public FLogScopedIndent()
+			{
+				LogIndent();
+			}
+
+			public void Dispose()
+			{
+				LogDedent();
+			}
+		}
+
+
 		#endregion
 
 	}
diff --git a/Engine/Source/Programs/Enterprise/Datasmith/DatasmithSolidworksExporter/Private/ConfigurationExporter.cs b/Engine/Source/Programs/Enterprise/Datasmith/DatasmithSolidworksExporter/Private/ConfigurationExporter.cs
--- a/Engine/Source/Programs/Enterprise/Datasmith/DatasmithSolidworksExporter/Private/ConfigurationExporter.cs
+++ b/Engine/Source/Programs/Enterprise/Datasmith/DatasmithSolidworksExporter/Private/ConfigurationExporter.cs
@@ -162,6 +162,17 @@ namespace DatasmithSolidworks
 			return Configuration;
 		}
 
+		// Add variant which inherits geometry from one of the 'real' Solidworks configurations
+		// This could be (linked)display state or exploded view
+		public FConfiguration AddSubConfiguration(FConfiguration MeshesConfiguration, string ConfigurationName)
+		{
+			Debug.Assert(!Configurations.ContainsKey(ConfigurationName));
+
+			FConfiguration Configuration = new FConfiguration(this, ConfigurationName, MainConfigurationName==ConfigurationName, InParentConfiguration: MeshesConfiguration);
+			Configurations.Add(ConfigurationName, Configuration);
+			return Configuration;
+		}
+
 		public bool DoesComponentHaveVisibleButDifferentMeshesInConfigs(FComponentName ComponentName, string ConfigNameA, string ConfigNameB)
 		{
 			bool bHasMeshA = GetMeshesConfiguration(ConfigNameA).GetMeshForComponent(ComponentName, out FMesh MeshA);
@@ -188,12 +199,14 @@ namespace DatasmithSolidworks
 			private readonly string Name;
 			private readonly bool bIsMainConfiguration;
 			private ConcurrentDictionary<FComponentName, FComponentMesh> MeshForComponent = new ConcurrentDictionary<FComponentName, FComponentMesh>();
+			private readonly FConfiguration ParentConfiguration;
 
-			public FConfiguration(FMeshes InMeshes, string InName, bool bInIsMainConfiguration)
+			public FConfiguration(FMeshes InMeshes, string InName, bool bInIsMainConfiguration, FConfiguration InParentConfiguration = null)
 			{
 				Meshes = InMeshes;
 				Name = InName;
 				bIsMainConfiguration = bInIsMainConfiguration;
+				ParentConfiguration = InParentConfiguration;
 			}
 
 			public void AddComponentWhichNeedsMesh(Component2 InComponent, FComponentName InComponentName)
@@ -241,7 +254,12 @@ namespace DatasmithSolidworks
 
 			public FMeshName GetMeshNameForComponent(FComponentName ComponentName)
 			{
-				return MeshForComponent[ComponentName].MeshName;
+				if (MeshForComponent.TryGetValue(ComponentName, out FComponentMesh Mesh))
+				{
+					return Mesh.MeshName;
+				}
+
+				return ParentConfiguration?.GetMeshNameForComponent(ComponentName) ?? new FMeshName();
 			}
 		}
 
@@ -282,6 +300,7 @@ namespace DatasmithSolidworks
 		{
 			return Configurations[ConfigName].GetMeshNameForComponent(ComponentName);
 		}
+
 	}
 
 	public class FConfigurationExporter
@@ -290,6 +309,7 @@ namespace DatasmithSolidworks
 		public readonly string[] ConfigurationNames;
 		private readonly string MainConfigurationName;
 		private readonly bool bExportDisplayStates;
+		private readonly bool bExportExplodedViews;
 
 		public FConfigurationTree.FComponentTreeNode CombinedTree;
 
@@ -299,12 +319,14 @@ namespace DatasmithSolidworks
 
 		public FConfigurationExporter(FMeshes InMeshes, 
 			string[] InConfigurationNames, string InMainConfigurationName,
-			bool bInExportDisplayStates)
+			bool bInExportDisplayStates,
+			bool bInExportExplodedViews)
 		{
 			Meshes = InMeshes;
 			ConfigurationNames = InConfigurationNames;
 			MainConfigurationName = InMainConfigurationName;
 			bExportDisplayStates = bInExportDisplayStates;
+			bExportExplodedViews = bInExportExplodedViews;
 		}
 
 		public List<FConfigurationData> ExportConfigurations(FDocumentTracker InDoc)
@@ -329,7 +351,7 @@ namespace DatasmithSolidworks
 
 			string OriginalConfigurationName = ConfigManager.ActiveConfiguration.Name;
 
-			Dictionary<string, string> ExportedConfigurationNames = new Dictionary<string, string>();
+			Dictionary<string, List<string>> ExportedConfigurationNames = new Dictionary<string, List<string>>();
 
 			// Enumerate configurations starting with Active - to avoid redundant config switching(which takes time) and,
 			// additionally, to export meshes in Active configuration without configuration name suffixes
@@ -353,6 +375,8 @@ namespace DatasmithSolidworks
 				}
 				LogDebug($"ActiveConfiguration '{ConfigManager.ActiveConfiguration.Name}'");
 
+				List<string> ExplodedViews = InitExplodedViews(InDoc, CfgName);
+
 				string ConfigName = CfgName;
 
 				string[] DisplayStates = null;
@@ -368,6 +392,7 @@ namespace DatasmithSolidworks
 						ConfigName = $"{CfgName}_{DisplayStates[0]}";
 					}
 				}
+
 				FConfigurationTree.FComponentTreeNode ConfigNode = new FConfigurationTree.FComponentTreeNode();
 				ConfigNode.ComponentInfo.ComponentName = FComponentName.FromCustomString(ConfigName);
 
@@ -379,9 +404,9 @@ namespace DatasmithSolidworks
 				// Use GetRootComponent3() with Resolve = true to ensure suppressed components will be loaded
 				// todo: docs says that Part document SW returns null. Not the case. But probably better to add a guard
 				LogDebug($"Components:");
-				CollectComponentsRecursive(InDoc, swConfiguration.GetRootComponent3(true), ConfigNode, MeshesConfiguration, CfgName);
+				CollectComponentsRecursive(InDoc, swConfiguration.GetRootComponent3(true), ConfigNode, MeshesConfiguration);
 
-				ExportedConfigurationNames.Add(CfgName, ConfigName);
+				ExportedConfigurationNames.Add(CfgName, new List<string>(){ConfigName});
 
 				LogDebug($"Materials:");
 				Dictionary<FComponentName, FObjectMaterials> MaterialsMap = GetComponentMaterials(InDoc, DisplayStates?[0], swConfiguration);
@@ -397,11 +422,64 @@ namespace DatasmithSolidworks
 						string DisplayStateTreeName = $"{CfgName}_DisplayState_{FDatasmithExporter.SanitizeName(DisplayState)}";
 						LogDebug($"Linked DisplayState '{DisplayState}' Materials for variant '{DisplayStateTreeName}'");
 
+						Meshes.AddSubConfiguration(MeshesConfiguration, ConfigurationName: DisplayStateTreeName);
+						ExportedConfigurationNames[CfgName].Add(DisplayStateTreeName);
 						MaterialsMap = GetComponentMaterials(InDoc, DisplayState, swConfiguration);
 						SetComponentTreeMaterials(ConfigNode, MaterialsMap, DisplayStateTreeName, false);
 					}
 				}
 
+				// Add variant per exploded view within the configuration
+				if (ExplodedViews != null && ExplodedViews.Count > 0)
+				{
+					LogIndent();
+
+
+					InDoc.SetExportStatus($"Exploded Views");
+					LogDebug($"Exploded Views:");
+					foreach (string ExplodedViewName in ExplodedViews)
+					{
+						LogIndent();
+
+						LogDebug($"View '{ExplodedViewName}'");
+
+						// Switch document state to the exploded view
+						switch (InDoc)
+						{
+							case FAssemblyDocumentTracker AsmDoc:
+							{
+								AsmDoc.SwAsmDoc?.ShowExploded2(true, ExplodedViewName);
+								break;
+							}
+							case FPartDocumentTracker PartDoc:
+							{
+								PartDoc.SwPartDoc?.ShowExploded(true, ExplodedViewName);
+								break;
+							}
+						}
+
+						string ExplodedViewVariantName = $"{CfgName}_Exploded_{FDatasmithExporter.SanitizeName(ExplodedViewName)}";
+						Meshes.AddSubConfiguration(MeshesConfiguration, ConfigurationName: ExplodedViewVariantName);
+						ExportedConfigurationNames[CfgName].Add(ExplodedViewVariantName);
+
+						// Add variant to each node with transform initialized from exploded state
+						foreach (FComponentTreeNode Child in ConfigNode.Children)
+						{
+							Child.Traverse(ConfigNode.CommonConfig, 
+								(InParentConfig, InNode) =>
+								{
+									FComponentConfig Config = InNode.AddConfiguration(ExplodedViewVariantName, bIsDisplayState: false, bIsActiveConfiguration: false);
+									ComputeNodeTransform(InParentConfig, InNode, Config);
+									return Config;
+								});
+						}
+
+						LogDedent();
+					}
+
+					LogDedent();
+				}
+
 				// Export materials
 				InDoc.SetExportStatus($"Component Materials");
 				HashSet<FComponentName> ComponentNamesToExportSet = new HashSet<FComponentName>();
@@ -507,9 +585,8 @@ namespace DatasmithSolidworks
 
 			List<FConfigurationData> FlatConfigurationData = new List<FConfigurationData>();
 
-
 			// Order exported variants by original configurations order
-			IEnumerable<string> ExportedConfigurationNamesOrdered = ExportedConfigurationNames.OrderBy(V => ConfigurationOrder[V.Key]).Select(V => V.Value);
+			IEnumerable<string> ExportedConfigurationNamesOrdered = ExportedConfigurationNames.OrderBy(V => ConfigurationOrder[V.Key]).SelectMany(V => V.Value);
 
 			// Add roots for configurations
 			foreach (string CfgName in ExportedConfigurationNamesOrdered)
@@ -540,6 +617,90 @@ namespace DatasmithSolidworks
 			return FlatConfigurationData;
 		}
 
+		// Collect exploded views which need to be exported and
+		// collapse current configuration exploded views in case we are exporting exploded views(we want top configuration to be in un-exploded state)
+		private List<string> InitExplodedViews(FDocumentTracker InDoc, string CfgName)
+		{
+			if (!bExportExplodedViews)
+			{
+				return null;
+			}
+
+			switch (InDoc)
+			{
+				case FAssemblyDocumentTracker AsmDoc:
+				{
+					AssemblyDoc SwAsmDoc = AsmDoc.SwAsmDoc;
+
+					if (SwAsmDoc == null)
+					{
+						return null;
+					}
+
+					int ExplodedViewCount = SwAsmDoc.GetExplodedViewCount2(CfgName);
+					LogDebug($"Exploded Views({ExplodedViewCount})");
+
+					if (ExplodedViewCount <= 0)
+					{
+						return null;
+					}
+
+					string[] ExplodedViewNames = (string[])SwAsmDoc.GetExplodedViewNames2(CfgName);
+					if (ExplodedViewNames == null)
+					{
+						return null;
+					}
+
+					List<string> ExplodedViews = new List<string>();
+					ExplodedViews.AddRange(ExplodedViewNames);
+
+					// Whe exporting exploded view collapse exploded view when exporting parent configuration itself
+					foreach (string ExplodedViewName in ExplodedViews)
+					{
+						SwAsmDoc.ShowExploded2(false, ExplodedViewName);
+					}
+
+					return ExplodedViews;
+				}
+				case FPartDocumentTracker PartDoc:
+				{
+					PartDoc SwPartDoc = PartDoc.SwPartDoc;
+
+					if (SwPartDoc == null)
+					{
+						return null;
+					}
+
+					int ExplodedViewCount = SwPartDoc.GetExplodedViewCount(CfgName);
+					LogDebug($"Exploded Views({ExplodedViewCount})");
+
+					if (ExplodedViewCount <= 0)
+					{
+						return null;
+					}
+
+					string[] ExplodedViewNames = (string[])SwPartDoc.GetExplodedViewNames(CfgName);
+					if (ExplodedViewNames == null)
+					{
+						return null;
+					}
+
+					List<string> ExplodedViews = new List<string>();
+					ExplodedViews.AddRange(ExplodedViewNames);
+
+					// Whe exporting exploded view collapse exploded view when exporting parent configuration itself
+					foreach (string ExplodedViewName in ExplodedViews)
+					{
+						SwPartDoc.ShowExploded(false, ExplodedViewName);
+					}
+
+					return ExplodedViews;
+				}
+			}
+
+			return null;
+		}
+
 		// Assign meshes to actors, export materials and assign materials
 		public void FinalizeExport(FDocumentTracker InDoc)
 		{
@@ -654,7 +815,7 @@ namespace DatasmithSolidworks
 		}
 
 		private static void CollectComponentsRecursive(FDocumentTracker InDoc, Component2 InComponent,
-			FComponentTreeNode InParentNode, FMeshes.FConfiguration Meshes, string CfgName)
+			FComponentTreeNode InParentNode, FMeshes.FConfiguration Meshes)
 		{
 			LogDebug($"'{InComponent.Name2}'");
 
@@ -673,34 +834,7 @@ namespace DatasmithSolidworks
 			NewNode.CommonConfig.bVisible = InComponent.Visible != (int)swComponentVisibilityState_e.swComponentHidden;
 			NewNode.CommonConfig.bSuppressed = InComponent.IsSuppressed();
 
-			// Read transform
-			MathTransform ComponentTransform = InComponent.GetTotalTransform(true);
-			if (ComponentTransform == null)
-			{
-				ComponentTransform = InComponent.Transform2;
-			}
-
-			if (ComponentTransform != null)
-			{
-				NewNode.CommonConfig.Transform = MathUtils.ConvertFromSolidworksTransform(ComponentTransform, 100f /*GeomScale*/);
-				LogDebug($" Transform: {NewNode.CommonConfig.Transform}");
-
-				if (InParentNode.CommonConfig.Transform.IsValid())
-				{
-					// Convert transform to parent space (original transform value fetched from Solidworks
-					// is in the root component's space). Datasmith wants relative transform for variants.
-					FMatrix4 ParentTransform = new FMatrix4(InParentNode.CommonConfig.Transform.Matrix);
-					FMatrix4 InverseParentTransform = ParentTransform.Inverse();
-					NewNode.CommonConfig.RelativeTransform = new FConvertedTransform(FMatrix4.FMatrix4x4Multiply(NewNode.CommonConfig.Transform.Matrix, InverseParentTransform));
-
-				}
-				else
-				{
-					NewNode.CommonConfig.RelativeTransform = NewNode.CommonConfig.Transform;
-				}
-				LogDebug($"   ParentTransform: {InParentNode.CommonConfig.Transform}");
-				LogDebug($"   RelativeTransform: {NewNode.CommonConfig.RelativeTransform}");
-			}
+			ComputeNodeTransform(InParentNode.CommonConfig, NewNode, NewNode.CommonConfig);
 
 			// Process children components
 			var Children = (Object[])InComponent.GetChildren();
@@ -711,7 +845,7 @@ namespace DatasmithSolidworks
 				{
 					Component2 Child = (Component2)ObjChild;
 					LogIndent();
-					CollectComponentsRecursive(InDoc, Child, NewNode, Meshes, CfgName);
+					CollectComponentsRecursive(InDoc, Child, NewNode, Meshes);
 					LogDedent();
 				}
 
@@ -732,6 +866,36 @@ namespace DatasmithSolidworks
 			InDoc.AddExportedComponent(NewNode);
 		}
 
+		private static void ComputeNodeTransform(FComponentConfig ParentConfig, FComponentTreeNode InNode, FComponentConfig ComponentConfig)
+		{
+			// Read transform
+			MathTransform ComponentTransform = InNode.Component.GetTotalTransform(true) ?? InNode.Component.Transform2;
+
+			if (ComponentTransform != null)
+			{
+				ComponentConfig.Transform =
+					MathUtils.ConvertFromSolidworksTransform(ComponentTransform, 100f /*GeomScale*/);
+				LogDebug($" Transform: {ComponentConfig.Transform}");
+
+				if (ParentConfig != null && ParentConfig.Transform.IsValid())
+				{
+					// Convert transform to parent space (original transform value fetched from Solidworks
+					// is in the root component's space). Datasmith wants relative transform for variants.
+					FMatrix4 ParentTransform = new FMatrix4(ParentConfig.Transform.Matrix);
+					FMatrix4 InverseParentTransform = ParentTransform.Inverse();
+					ComponentConfig.RelativeTransform = new FConvertedTransform(
+						FMatrix4.FMatrix4x4Multiply(ComponentConfig.Transform.Matrix, InverseParentTransform));
+				}
+				else
+				{
+					ComponentConfig.RelativeTransform = ComponentConfig.Transform;
+				}
+
+				LogDebug($"   ParentTransform: {ParentConfig?.Transform}");
+				LogDebug($"   RelativeTransform: {ComponentConfig.RelativeTransform}");
+			}
+		}
+
 		public bool DoesComponentHaveVisibleButDifferentMeshesInConfigs(FComponentName ComponentName, string ConfigNameA, string ConfigNameB)
 		{
 			return Meshes.DoesComponentHaveVisibleButDifferentMeshesInConfigs(ComponentName, ConfigNameA, ConfigNameB);
diff --git a/Engine/Source/Programs/Enterprise/Datasmith/DatasmithSolidworksExporter/Private/ConfigurationTree.cs b/Engine/Source/Programs/Enterprise/Datasmith/DatasmithSolidworksExporter/Private/ConfigurationTree.cs
--- a/Engine/Source/Programs/Enterprise/Datasmith/DatasmithSolidworksExporter/Private/ConfigurationTree.cs
+++ b/Engine/Source/Programs/Enterprise/Datasmith/DatasmithSolidworksExporter/Private/ConfigurationTree.cs
@@ -1,11 +1,13 @@
 // Copyright Epic Games, Inc. All Rights Reserved.
 
+using System;
 using SolidWorks.Interop.sldworks;
 using System.Collections.Generic;
 using System.Diagnostics;
 using System.Linq;
 using System.Runtime.InteropServices;
 using DatasmithSolidworks.Names;
+using static DatasmithSolidworks.Addin;
 
 namespace DatasmithSolidworks
 {
@@ -114,6 +116,28 @@ namespace DatasmithSolidworks
 
 			public List<FComponentTreeNode> Children;
 
+			// Traverse tree passing result of function computation for each node to its children
+			// Function - receives parent's computed value and current node and returns value computed for the node(to pass to children)
+			public void Traverse<T>(T ParentValue, Func<T, FComponentTreeNode, T> Function)
+			{
+				LogDebug($" {ComponentName}");
+
+				T Value = Function(ParentValue, this);
+
+				if (Children == null)
+				{
+					return;
+				}				
+				
+				foreach (FComponentTreeNode Child in Children)
+				{
+					using (LogScopedIndent())
+					{
+						Child.Traverse(Value, Function);
+					}
+				}
+			}
+
 			public IEnumerable<FComponentTreeNode> EnumChildren()
 			{
 				if (Children == null)
@@ -242,6 +266,12 @@ namespace DatasmithSolidworks
 				FComponentConfig NodeConfig = CombinedChild.AddConfiguration(InConfigurationName, false, bIsActiveConfiguration);
 				NodeConfig.CopyFrom(Child.CommonConfig);
 
+				// Copy 'extra' configurations created for this actual configuration - linked display states and exploded views
+				if (Child.Configurations != null)
+				{
+					CombinedChild.Configurations.AddRange(Child.Configurations);
+				}
+
 				// Recurse to children
 				Merge(CombinedChild, Child, InConfigurationName, bIsActiveConfiguration);
 			}
@@ -294,7 +324,8 @@ namespace DatasmithSolidworks
 					// There could be components without a transform, so we're checking if for null
 					for (int Idx = 1; Idx < InNode.Configurations.Count; Idx++)
 					{
-						if (!MathUtils.TransformsAreEqual(InNode.Configurations[Idx].RelativeTransform, RelativeTransform))
+						FConvertedTransform OtherTransform = InNode.Configurations[Idx].RelativeTransform;
+						if (OtherTransform.IsValid() && !MathUtils.TransformsAreEqual(OtherTransform, RelativeTransform))
 						{
 							bAllTransformsAreSame = false;
 							break;
diff --git a/Engine/Source/Programs/Enterprise/Datasmith/DatasmithSolidworksExporter/Private/Document.cs b/Engine/Source/Programs/Enterprise/Datasmith/DatasmithSolidworksExporter/Private/Document.cs
--- a/Engine/Source/Programs/Enterprise/Datasmith/DatasmithSolidworksExporter/Private/Document.cs
+++ b/Engine/Source/Programs/Enterprise/Datasmith/DatasmithSolidworksExporter/Private/Document.cs
@@ -242,7 +242,7 @@ namespace DatasmithSolidworks
 
 			FMeshes Meshes = new FMeshes(ActiveConfigurationName);
 
-			FConfigurationExporter ConfigurationExporter = new FConfigurationExporter(Meshes, ConfigurationNames, ActiveConfigurationName, true);
+			FConfigurationExporter ConfigurationExporter = new FConfigurationExporter(Meshes, ConfigurationNames, ActiveConfigurationName, bInExportDisplayStates: true, bInExportExplodedViews: true);
 
 			List<FConfigurationData> Configs = ConfigurationExporter.ExportConfigurations(this);
 
@@ -275,7 +275,7 @@ namespace DatasmithSolidworks
 			string ActiveConfigurationName = ConfigManager.ActiveConfiguration.Name;
 
 			FMeshes Meshes = GetMeshes(ActiveConfigurationName);
-			FConfigurationExporter ConfigurationExporter = new FConfigurationExporter(Meshes, new []{ ActiveConfigurationName }, ActiveConfigurationName, false);
+			FConfigurationExporter ConfigurationExporter = new FConfigurationExporter(Meshes, new []{ ActiveConfigurationName }, ActiveConfigurationName, bInExportDisplayStates: false, bInExportExplodedViews: false);
 			List<FConfigurationData> Configs = ConfigurationExporter.ExportConfigurations(this);
 			bHasConfigurations = (Configs != null) && (Configs.Count != 0);
 			ExportToDatasmithScene(ConfigurationExporter);
