[Solidworks] Fixed AutoSync MaterialUpdater crash and other issues

- replaced separate thread to check for material updates with iterative checkinbg in the main thread on idle

- this fixed random crashes due to race conditions

- slightly improved speed of registering material changes (e.g. from minute to under seconds for some real models)

diff --git a/Engine/Source/Programs/Enterprise/Datasmith/DatasmithSolidworksExporter/Private/Addin.cs b/Engine/Source/Programs/Enterprise/Datasmith/DatasmithSolidworksExporter/Private/Addin.cs
--- a/Engine/Source/Programs/Enterprise/Datasmith/DatasmithSolidworksExporter/Private/Addin.cs
+++ b/Engine/Source/Programs/Enterprise/Datasmith/DatasmithSolidworksExporter/Private/Addin.cs
@@ -671,8 +671,6 @@ namespace DatasmithSolidworks
 				LogDedent();
 			}
 		}
-
-
 		#endregion
 
 	}
@@ -689,5 +687,17 @@ namespace DatasmithSolidworks
 			}
 			return Value;
 		}
+
+		public static bool TryRemove<K, V>(this Dictionary<K, V> Map, K Key, out V OutValue) 
+		{
+			if (Map.TryGetValue(Key, out OutValue))
+			{
+				Map.Remove(Key);
+				return true;
+			}			
+			return false;
+
+
+		}
 	}
 }
diff --git a/Engine/Source/Programs/Enterprise/Datasmith/DatasmithSolidworksExporter/Private/AssemblyDocument.cs b/Engine/Source/Programs/Enterprise/Datasmith/DatasmithSolidworksExporter/Private/AssemblyDocument.cs
--- a/Engine/Source/Programs/Enterprise/Datasmith/DatasmithSolidworksExporter/Private/AssemblyDocument.cs
+++ b/Engine/Source/Programs/Enterprise/Datasmith/DatasmithSolidworksExporter/Private/AssemblyDocument.cs
@@ -9,10 +9,93 @@ using System.Diagnostics;
 using System.Linq;
 using DatasmithSolidworks.Names;
 using static DatasmithSolidworks.Addin;
+using static DatasmithSolidworks.FAssemblyDocumentTracker;
 
 namespace DatasmithSolidworks
 {
 
+	public class FAssemblyLazyUpdateChecker: FLazyUpdateCheckerBase
+	{
+		private readonly FAssemblyDocumentTracker AsmDocumentTracker;
+
+		private bool bHasDirtyMaterials;
+		private bool bHasDirtyComponents;
+		HashSet<FComponentName> AllExportedComponents;
+
+		public FAssemblyLazyUpdateChecker(FAssemblyDocumentTracker InDocumentTracker): base(InDocumentTracker)
+		{
+			AsmDocumentTracker = InDocumentTracker;
+		}
+
+		public override bool HasUpdates()
+		{
+			return bHasDirtyComponents || bHasDirtyMaterials;
+		}
+
+		public override void Restart()
+		{
+			AllExportedComponents = new HashSet<FComponentName>();
+			bHasDirtyMaterials = false;
+			bHasDirtyComponents = false;
+		}
+
+		public override IEnumerable<bool> CheckForModificationsEnum()
+		{
+			FSyncState SyncState = AsmDocumentTracker.SyncState;
+
+			// Dig into part level materials (they wont be read by LoadDocumentMaterials)
+			AllExportedComponents.UnionWith(SyncState.CleanComponents);
+			AllExportedComponents.UnionWith(SyncState.DirtyComponents.Keys);
+
+			Dictionary<int, FMaterial> MaterialsMap = new Dictionary<int, FMaterial>();
+			Dictionary<FComponentName, FObjectMaterials> CurrentDocMaterialsMap  = new Dictionary<FComponentName, FObjectMaterials>();
+
+			foreach (bool _ in FObjectMaterials.LoadAssemblyMaterialsEnum(AsmDocumentTracker, AllExportedComponents, CurrentDocMaterialsMap, MaterialsMap))
+			{
+				yield return _;
+			}
+
+			HashSet<FComponentName> Components =  SyncState.ComponentsMaterialsMap == null ? new HashSet<FComponentName>() : new HashSet<FComponentName>(SyncState.ComponentsMaterialsMap.Keys);
+			HashSet<FComponentName> CurrentComponents =  CurrentDocMaterialsMap == null ? new HashSet<FComponentName>() : new HashSet<FComponentName>(CurrentDocMaterialsMap.Keys);
+
+			// Components which stayed in the materials map
+			HashSet<FComponentName> CommonComponents = new HashSet<FComponentName>(CurrentComponents.Intersect(Components));
+			 
+			IEnumerable<FComponentName> ComponentsWithAddedOrRemovedMaterial = Components.Union(CurrentComponents).Except(CommonComponents);
+			foreach (FComponentName CompName in ComponentsWithAddedOrRemovedMaterial)
+			{
+				bool bShouldSyncComponentMaterial = false;
+
+				if (SyncState.CollectedComponentsMap.ContainsKey(CompName))
+				{
+					Component2 Comp = SyncState.CollectedComponentsMap[CompName];
+					bShouldSyncComponentMaterial = !Comp.IsSuppressed() &&
+					                               (Comp.Visible == (int)swComponentVisibilityState_e
+						                               .swComponentVisible);
+					yield return true;
+				}
+
+				if (bShouldSyncComponentMaterial)
+				{
+					bHasDirtyMaterials = true;
+					AsmDocumentTracker.SetComponentDirty(CompName, EComponentDirtyState.Material);
+				}
+			}
+
+			// Check if components have their material modified
+			foreach (FComponentName ComponentName in CommonComponents)
+			{
+				if (!SyncState.ComponentsMaterialsMap[ComponentName]
+					    .EqualMaterials(CurrentDocMaterialsMap[ComponentName]))
+				{
+					AsmDocumentTracker.SetComponentDirty(ComponentName, EComponentDirtyState.Material);
+					bHasDirtyComponents = true;
+					yield return true;
+				}
+			}
+		}
+	};
+
 	public class FAssemblyDocumentTracker : FDocumentTracker
 	{
 		public enum EComponentDirtyState
@@ -28,7 +111,7 @@ namespace DatasmithSolidworks
 		{
 			public Dictionary<string, FPartDocument> PartsMap = new Dictionary<string, FPartDocument>();
 
-			public ConcurrentDictionary<FComponentName, FObjectMaterials> ComponentsMaterialsMap = null;
+			public Dictionary<FComponentName, FObjectMaterials> ComponentsMaterialsMap = null;
 
 			public Dictionary<FComponentName, string> ComponentToPartMap = new Dictionary<FComponentName, string>();
 
@@ -38,13 +121,13 @@ namespace DatasmithSolidworks
 			/// All updated(not dirty) components. Updated component will be removed from dirty on export end
 			public HashSet<FComponentName> CleanComponents = new HashSet<FComponentName>();
 
-			/** Flags for each component that needs update. Cleared when export completes*/
+			/// Flags for each component that needs update. Cleared when export completes*/
 			public Dictionary<FComponentName, uint> DirtyComponents = new Dictionary<FComponentName, uint>();
 
 			public HashSet<FComponentName> ComponentsToDelete = new HashSet<FComponentName>();
 
-			public Dictionary<FComponentName, Component2> ExportedComponentsMap =
-				new Dictionary<FComponentName, Component2>();
+			/// all components in the document
+			public Dictionary<FComponentName, Component2> CollectedComponentsMap = new Dictionary<FComponentName, Component2>();
 
 			/** Stores which mesh was exported for each component*/
 			public Dictionary<FComponentName, List<FMeshName>> ComponentNameToMeshNameMap =
@@ -59,6 +142,8 @@ namespace DatasmithSolidworks
 
 		public FSyncState SyncState { get; private set; } = new FSyncState();
 
+		private FAssemblyLazyUpdateChecker LazyUpdateChecker;
+
 		public FAssemblyDocumentTracker(FAssemblyDocument InDoc, AssemblyDoc InSwDoc, FDatasmithExporter InExporter) : base(InDoc, InExporter)
 		{
 			SwAsmDoc = InSwDoc;
@@ -120,11 +205,11 @@ namespace DatasmithSolidworks
 			{
 				FActorName ActorName = Exporter.GetComponentActorName(CompName);
 				SyncState.ComponentToPartMap.Remove(CompName);
-				SyncState.ExportedComponentsMap.Remove(CompName);
+				SyncState.CollectedComponentsMap.Remove(CompName);
 
 				ReleaseComponentMeshes(CompName);
 
-				SyncState.ComponentsMaterialsMap?.TryRemove(CompName, out FObjectMaterials _);
+				SyncState.ComponentsMaterialsMap?.Remove(CompName);
 				SyncState.ComponentsTransformsMap.Remove(CompName);
 				Exporter.RemoveActor(ActorName);
 			}
@@ -165,7 +250,7 @@ namespace DatasmithSolidworks
 			}
 		}
 
-		public override ConcurrentDictionary<FComponentName, FObjectMaterials> LoadDocumentMaterials(
+		public override Dictionary<FComponentName, FObjectMaterials> LoadDocumentMaterials(
 			HashSet<FComponentName> ComponentNamesToExportSet)
 		{
 			return FObjectMaterials.LoadAssemblyMaterials(this, ComponentNamesToExportSet,
@@ -178,7 +263,7 @@ namespace DatasmithSolidworks
 
 			if (SyncState.ComponentsMaterialsMap == null)
 			{
-				SyncState.ComponentsMaterialsMap = new ConcurrentDictionary<FComponentName, FObjectMaterials>();
+				SyncState.ComponentsMaterialsMap = new Dictionary<FComponentName, FObjectMaterials>();
 			}
 			SyncState.ComponentsMaterialsMap[ComponentName] = Materials;
 		}
@@ -204,71 +289,6 @@ namespace DatasmithSolidworks
 			return MeshData;
 		}
 
-		public override bool HasMaterialUpdates()
-		{
-			// Dig into part level materials (they wont be read by LoadDocumentMaterials)
-			HashSet<FComponentName> AllExportedComponents = new HashSet<FComponentName>();
-			AllExportedComponents.UnionWith(SyncState.CleanComponents);
-			AllExportedComponents.UnionWith(SyncState.DirtyComponents.Keys);
-
-			ConcurrentDictionary<FComponentName, FObjectMaterials> CurrentDocMaterialsMap =
-				FObjectMaterials.LoadAssemblyMaterials(this, AllExportedComponents,
-					swDisplayStateOpts_e.swThisDisplayState, null);
-
-
-			HashSet<FComponentName> Components =  SyncState.ComponentsMaterialsMap == null ? new HashSet<FComponentName>() : new HashSet<FComponentName>(SyncState.ComponentsMaterialsMap.Keys);
-			HashSet<FComponentName> CurrentComponents =  CurrentDocMaterialsMap == null ? new HashSet<FComponentName>() : new HashSet<FComponentName>(CurrentDocMaterialsMap.Keys);
-
-			// Components which stayed in the materials map
-			HashSet<FComponentName> CommonComponents = new HashSet<FComponentName>(CurrentComponents.Intersect(Components));
-
-			 
-			bool bHasDirtyMaterials = false;
-			IEnumerable<FComponentName> ComponentsWithAddedOrRemovedMaterial = Components.Union(CurrentComponents).Except(CommonComponents);
-			foreach (FComponentName CompName in ComponentsWithAddedOrRemovedMaterial)
-			{
-				bool bShouldSyncComponentMaterial = false;
-
-				if (SyncState.ExportedComponentsMap.ContainsKey(CompName))
-				{
-					try
-					{
-						Component2 Comp = SyncState.ExportedComponentsMap[CompName];
-						bShouldSyncComponentMaterial = !Comp.IsSuppressed() &&
-						                               (Comp.Visible == (int)swComponentVisibilityState_e
-							                               .swComponentVisible);
-					}
-					catch
-					{
-						// todo: exceptions here(and other places loading materials) thrown because of multi-threaded FMaterialUpdateChecker
-						// calling this code from its thread  performing COM calls into SW API 
-						// calling the API from any other thread than main is extremely slow so update checker is running forever
-						// meaning it has increased probability to have model update while it's running and have a race condition of any kind
-						// there's UE-177970 fot this
-						Debug.Assert(false);
-					}
-				}
-
-				if (bShouldSyncComponentMaterial)
-				{
-					bHasDirtyMaterials = true;
-					SetComponentDirty(CompName, EComponentDirtyState.Material);
-				}
-			}
-
-			// Check if components have their material modified
-			bool bHasDirtyComponents = false;
-			foreach (FComponentName ComponentName in CommonComponents)
-			{
-				if (!SyncState.ComponentsMaterialsMap[ComponentName].EqualMaterials(CurrentDocMaterialsMap[ComponentName]))
-				{
-					SetComponentDirty(ComponentName, EComponentDirtyState.Material);
-					bHasDirtyComponents = true;
-				}
-			}
-			return bHasDirtyComponents || bHasDirtyMaterials;
-		}
-
 		public FConvertedTransform GetComponentDatasmithTransform(Component2 InComponent)
 		{
 			MathTransform ComponentTransform = InComponent.GetTotalTransform(true);
@@ -415,9 +435,9 @@ namespace DatasmithSolidworks
 			}
 		}
 
-		public override void AddExportedComponent(FConfigurationTree.FComponentTreeNode InNode)
+		public override void AddCollectedComponent(FConfigurationTree.FComponentTreeNode InNode)
 		{
-			SyncState.ExportedComponentsMap[InNode.ComponentName] = InNode.Component;
+			SyncState.CollectedComponentsMap[InNode.ComponentName] = InNode.Component;
 		}
 
 		public override bool NeedExportComponent(FConfigurationTree.FComponentTreeNode InComponent,
@@ -505,7 +525,17 @@ namespace DatasmithSolidworks
 		{
 			SyncState.CleanComponents?.Clear();
 			SyncState.DirtyComponents?.Clear();
-			SyncState.ExportedComponentsMap?.Clear();
+			SyncState.CollectedComponentsMap?.Clear();
+		}
+
+		public void Tick()
+		{
+			LazyUpdateChecker?.Tick();
+		}
+
+		public void TrackChanges()
+		{
+			LazyUpdateChecker = new FAssemblyLazyUpdateChecker(this);
 		}
 	}
 
@@ -694,7 +724,7 @@ namespace DatasmithSolidworks
 		{
 			// Check each exported component's transform for changes, since 
 			// this callback does not tell us what changed (and there's no other way to know that)!
-			foreach (var KVP in AsmDocumentTracker.SyncState.ExportedComponentsMap)
+			foreach (var KVP in AsmDocumentTracker.SyncState.CollectedComponentsMap)
 			{
 				Component2 Comp = KVP.Value;
 				FConvertedTransform PrevCompTransform;
@@ -885,7 +915,6 @@ namespace DatasmithSolidworks
 	{
 		private readonly FAssemblyDocumentTracker DocumentTracker;
 		private readonly FAssemblyDocumentEvents Events;
-		private FMaterialUpdateChecker MaterialUpdateChecker;
 
 		public FAssemblyDocumentNotifications(FAssemblyDocument InAssemblyDocument, FAssemblyDocumentTracker InAsmDocumentTracker)
 		{
@@ -895,26 +924,25 @@ namespace DatasmithSolidworks
 
 		public void Destroy()
 		{
-			MaterialUpdateChecker?.Destroy();
 			Events.Destroy();
 		}
 
 		public void Start()
 		{
-			if (MaterialUpdateChecker == null)
-			{
-				MaterialUpdateChecker = new FMaterialUpdateChecker(DocumentTracker);
-			}
+			DocumentTracker.TrackChanges();
 		}
 
 		public void Resume()
 		{
-			MaterialUpdateChecker?.Resume();
 		}
 
 		public void Pause()
 		{
-			MaterialUpdateChecker?.Pause();
+		}
+
+		public void Idle()
+		{
+			DocumentTracker.Tick();
 		}
 	}
 
@@ -936,6 +964,11 @@ namespace DatasmithSolidworks
 			DocumentTracker.Destroy();
 		}
 
+		public void Idle()
+		{
+			Notifications.Idle();
+		}
+
 		public void Start()
 		{
 			Notifications.Start();
diff --git a/Engine/Source/Programs/Enterprise/Datasmith/DatasmithSolidworksExporter/Private/Body.cs b/Engine/Source/Programs/Enterprise/Datasmith/DatasmithSolidworksExporter/Private/Body.cs
--- a/Engine/Source/Programs/Enterprise/Datasmith/DatasmithSolidworksExporter/Private/Body.cs
+++ b/Engine/Source/Programs/Enterprise/Datasmith/DatasmithSolidworksExporter/Private/Body.cs
@@ -123,15 +123,13 @@ namespace DatasmithSolidworks
 				}
 			}
 
-			Parallel.ForEach(AllBodies, ObjBody =>
+			foreach(object ObjBody in AllBodies)
 			{
-				Body2 Body = ObjBody as Body2;
-
-				if (Body != null && Body.Visible && !Body.IsTemporaryBody())
+				if (ObjBody is Body2 Body && Body.Visible && !Body.IsTemporaryBody())
 				{
 					ResultBodies.Add(new FBody(Body));
 				}
-			});
+			}
 
 			return ResultBodies;
 		}
diff --git a/Engine/Source/Programs/Enterprise/Datasmith/DatasmithSolidworksExporter/Private/ConfigurationExporter.cs b/Engine/Source/Programs/Enterprise/Datasmith/DatasmithSolidworksExporter/Private/ConfigurationExporter.cs
--- a/Engine/Source/Programs/Enterprise/Datasmith/DatasmithSolidworksExporter/Private/ConfigurationExporter.cs
+++ b/Engine/Source/Programs/Enterprise/Datasmith/DatasmithSolidworksExporter/Private/ConfigurationExporter.cs
@@ -198,7 +198,7 @@ namespace DatasmithSolidworks
 			private readonly FMeshes Meshes;
 			private readonly string Name;
 			private readonly bool bIsMainConfiguration;
-			private ConcurrentDictionary<FComponentName, FComponentMesh> MeshForComponent = new ConcurrentDictionary<FComponentName, FComponentMesh>();
+			private Dictionary<FComponentName, FComponentMesh> MeshForComponent = new Dictionary<FComponentName, FComponentMesh>();
 			private readonly FConfiguration ParentConfiguration;
 
 			public FConfiguration(FMeshes InMeshes, string InName, bool bInIsMainConfiguration, FConfiguration InParentConfiguration = null)
@@ -211,7 +211,7 @@ namespace DatasmithSolidworks
 
 			public void AddComponentWhichNeedsMesh(Component2 InComponent, FComponentName InComponentName)
 			{
-				MeshForComponent.TryAdd(InComponentName, new FComponentMesh() {Component = InComponent});
+				MeshForComponent[InComponentName] = new FComponentMesh() {Component = InComponent};
 			}
 
 			public void AddMesh(FComponentName ComponentName, FMeshData MeshData, out FMeshName MeshName)
@@ -485,13 +485,10 @@ namespace DatasmithSolidworks
 				HashSet<FComponentName> ComponentNamesToExportSet = new HashSet<FComponentName>();
 				foreach (FComponentName ComponentName in MeshesConfiguration.EnumerateComponentNames())
 				{
-					if (!ComponentNamesToExportSet.Contains(ComponentName))
-					{
-						ComponentNamesToExportSet.Add(ComponentName);
-					}
+					ComponentNamesToExportSet.Add(ComponentName);
 				}
 
-				ConcurrentDictionary<FComponentName, FObjectMaterials> ModifiedComponentsMaterials = InDoc.LoadDocumentMaterials(ComponentNamesToExportSet);
+				Dictionary<FComponentName, FObjectMaterials> ModifiedComponentsMaterials = InDoc.LoadDocumentMaterials(ComponentNamesToExportSet);
 
 				if (ModifiedComponentsMaterials != null)
 				{
@@ -531,7 +528,7 @@ namespace DatasmithSolidworks
 
 					FConfigurationTree.FComponentTreeNode NewNode = new FConfigurationTree.FComponentTreeNode(Component);
 					NewNode.ComponentInfo.ComponentName = new FComponentName(Component);
-					InDoc.AddExportedComponent(NewNode);
+					InDoc.AddCollectedComponent(NewNode);
 
 					object[] Children = (object[])Component.GetChildren();
 					if (Children != null)
@@ -749,14 +746,9 @@ namespace DatasmithSolidworks
 
 			if (InDoc is FAssemblyDocumentTracker AsmDoc)
 			{
-				HashSet<FComponentName> ComponentsSet = new HashSet<FComponentName>();
-
-				foreach (FComponentName CompName in AsmDoc.SyncState.ExportedComponentsMap.Keys)
-				{
-					ComponentsSet.Add(CompName);
-				}
+				HashSet<FComponentName> ComponentsSet = new HashSet<FComponentName>(AsmDoc.SyncState.CollectedComponentsMap.Keys);
 
-				ConcurrentDictionary<FComponentName, FObjectMaterials> ComponentMaterials =
+				Dictionary<FComponentName, FObjectMaterials> ComponentMaterials =
 					FObjectMaterials.LoadAssemblyMaterials(AsmDoc, ComponentsSet, Option, DisplayStates);
 
 				if (ComponentMaterials != null)
@@ -770,7 +762,7 @@ namespace DatasmithSolidworks
 			else
 			{
 				FPartDocumentTracker PartDocumentTracker = InDoc as FPartDocumentTracker;
-				FObjectMaterials PartMaterials = FObjectMaterials.LoadPartMaterials(InDoc, PartDocumentTracker.SwPartDoc, Option, DisplayStates);
+				FObjectMaterials PartMaterials = FObjectMaterials.LoadPartMaterials(InDoc, PartDocumentTracker.SwPartDoc, Option, DisplayStates, InDoc.ExportedMaterialsMap);
 				Component2 Comp = InConfiguration.GetRootComponent3(true);
 
 				MaterialsMap.Add(FComponentName.FromCustomString(Comp?.Name2 ?? ""), PartMaterials);
@@ -863,7 +855,7 @@ namespace DatasmithSolidworks
 				}
 				Meshes.AddComponentWhichNeedsMesh(NewNode.Component, NewNode.ComponentName);
 			}
-			InDoc.AddExportedComponent(NewNode);
+			InDoc.AddCollectedComponent(NewNode);
 		}
 
 		private static void ComputeNodeTransform(FComponentConfig ParentConfig, FComponentTreeNode InNode, FComponentConfig ComponentConfig)
diff --git a/Engine/Source/Programs/Enterprise/Datasmith/DatasmithSolidworksExporter/Private/DatasmithExporter.cs b/Engine/Source/Programs/Enterprise/Datasmith/DatasmithSolidworksExporter/Private/DatasmithExporter.cs
--- a/Engine/Source/Programs/Enterprise/Datasmith/DatasmithSolidworksExporter/Private/DatasmithExporter.cs
+++ b/Engine/Source/Programs/Enterprise/Datasmith/DatasmithSolidworksExporter/Private/DatasmithExporter.cs
@@ -160,9 +160,9 @@ namespace DatasmithSolidworks
 	public class FDatasmithExporter
 	{
 		private Dictionary<FActorName, Tuple<EActorType, FDatasmithFacadeActor>> ExportedActorsMap = new Dictionary<FActorName, Tuple<EActorType, FDatasmithFacadeActor>>();
-		private ConcurrentDictionary<FMeshName, FMeshExportInfo> ExportedMeshesMap = new ConcurrentDictionary<FMeshName, FMeshExportInfo>();
-		private ConcurrentDictionary<int, FDatasmithFacadeMaterialInstance> ExportedMaterialInstancesMap = new ConcurrentDictionary<int, FDatasmithFacadeMaterialInstance>();
-		private ConcurrentDictionary<string, FDatasmithFacadeTexture> ExportedTexturesMap = new ConcurrentDictionary<string, FDatasmithFacadeTexture>();
+		private Dictionary<FMeshName, FMeshExportInfo> ExportedMeshesMap = new Dictionary<FMeshName, FMeshExportInfo>();
+		private Dictionary<int, FDatasmithFacadeMaterialInstance> ExportedMaterialInstancesMap = new Dictionary<int, FDatasmithFacadeMaterialInstance>();
+		private Dictionary<string, FDatasmithFacadeTexture> ExportedTexturesMap = new Dictionary<string, FDatasmithFacadeTexture>();
 
 		// Number of meshes/variants using this material to allow exporting only used materials
 		private Dictionary<int, int> MaterialUsers = new Dictionary<int, int>();
@@ -281,7 +281,7 @@ namespace DatasmithSolidworks
 			RemoveMesh(Info.MeshName);
 			LogDedent();
 
-			ExportedMeshesMap.TryAdd(Info.MeshName, Info);
+			ExportedMeshesMap[Info.MeshName] = Info;
 
 			LogDebug($"DatasmithScene.AddMesh('{Info.MeshElement.GetName()}')");
 			DatasmithScene.AddMesh(Info.MeshElement);
@@ -301,6 +301,8 @@ namespace DatasmithSolidworks
             {
                 return;
             }
+
+            
 			if (ExportedMeshesMap.TryRemove(MeshName, out FMeshExportInfo Info))
 			{
 				LogDebug($"  removing mesh element '{Info.MeshElement.GetName()}'");
@@ -794,7 +796,7 @@ namespace DatasmithSolidworks
 			}
 		}
 
-		public void ExportMaterials(ConcurrentDictionary<int, FMaterial> InMaterialsMap)
+		public void ExportMaterials(Dictionary<int, FMaterial> InMaterialsMap)
 		{
 			LogDebug($"ExportMaterials: \n  {string.Join("  \n", InMaterialsMap.Select(KVP => $"{KVP.Key}: {KVP.Value}" ))}");
 
@@ -1044,7 +1046,7 @@ namespace DatasmithSolidworks
 					TextureElement.SetTextureAddressY(FDatasmithFacadeTexture.ETextureAddress.Wrap);
 					FDatasmithFacadeTexture.ETextureMode TextureMode = FDatasmithFacadeTexture.ETextureMode.Diffuse;
 					TextureElement.SetTextureMode(TextureMode);
-					ExportedTexturesMap.TryAdd(InMaterial.Texture, TextureElement);
+					ExportedTexturesMap[InMaterial.Texture] = TextureElement;
 
 					InMaterialInstance.AddTexture(InParamName, TextureElement);
 				}
@@ -1076,7 +1078,7 @@ namespace DatasmithSolidworks
 					TextureElement.SetTextureAddressY(FDatasmithFacadeTexture.ETextureAddress.Wrap);
 					FDatasmithFacadeTexture.ETextureMode TextureMode = FDatasmithFacadeTexture.ETextureMode.Normal;
 					TextureElement.SetTextureMode(TextureMode);
-					ExportedTexturesMap.TryAdd(InMaterial.BumpTextureFileName, TextureElement);
+					ExportedTexturesMap[InMaterial.BumpTextureFileName] = TextureElement;
 
 					InMaterialInstance.AddTexture(InParamName, TextureElement);
 				}
@@ -1144,7 +1146,7 @@ namespace DatasmithSolidworks
 		}
 
 		// Export meshes
-		public IEnumerable<FMeshExportInfo> ExportMeshes(IEnumerable<FMeshExportInfo> MeshExportInfos)
+		public IEnumerable<FMeshExportInfo> ExportMeshes(List<FMeshExportInfo> MeshExportInfos)
 		{
 			LogDebug($"ExportMeshes: {string.Join(", ", MeshExportInfos)}");
 
diff --git a/Engine/Source/Programs/Enterprise/Datasmith/DatasmithSolidworksExporter/Private/Document.cs b/Engine/Source/Programs/Enterprise/Datasmith/DatasmithSolidworksExporter/Private/Document.cs
--- a/Engine/Source/Programs/Enterprise/Datasmith/DatasmithSolidworksExporter/Private/Document.cs
+++ b/Engine/Source/Programs/Enterprise/Datasmith/DatasmithSolidworksExporter/Private/Document.cs
@@ -12,89 +12,6 @@ using static DatasmithSolidworks.Addin;
 
 namespace DatasmithSolidworks
 {
-	// Checks for material updates periodically, using a separate thread to execute checks
-	// todo: Solidworks API is not concurrent-friendly - https://forum.solidworks.com/thread/43129
-	//   need to test but it's probably faster to check material updates simply in Sync. Maybe 
-	//   with the same period as currently if needed so
-	public class FMaterialUpdateChecker
-	{
-		public Thread MaterialCheckerThread = null;
-		public bool bExitMaterialUpdateThread = false;
-		public ManualResetEvent MaterialCheckerEvent = null;
-
-		public FDocumentTracker DocumentTracker;
-
-		public FMaterialUpdateChecker(FDocumentTracker InDocumentTracker)
-		{
-			DocumentTracker = InDocumentTracker;
-
-			MaterialCheckerEvent = new ManualResetEvent(false);
-			MaterialCheckerEvent.Set();
-
-			MaterialCheckerThread = new Thread(CheckForMaterialUpdatesProc);
-			MaterialCheckerThread.Name = "DatasmithSolidworks.FMaterialUpdateChecker Thread";
-			MaterialCheckerThread.Start();
-		}
-
-		public void Destroy()
-		{
-			bExitMaterialUpdateThread = true;
-
-			if (MaterialCheckerThread != null && !MaterialCheckerThread.Join(500))
-			{
-				MaterialCheckerThread.Abort();
-			}
-		}
-
-		public void Resume()
-		{
-			MaterialCheckerEvent.Set();
-		}
-
-		public void Pause()
-		{
-			MaterialCheckerEvent.Reset();
-		}
-
-		private void CheckForMaterialUpdatesProc()
-		{
-			while (!bExitMaterialUpdateThread)
-			{
-				MaterialCheckerEvent.WaitOne();
-
-				// Only check for material updates when we are not current exporting
-				if (!DocumentTracker.IsUpdateInProgress())
-				{
-					bool bSketchMode = false;
-					try
-					{
-						bSketchMode = (DocumentTracker.SwDoc.SketchManager?.ActiveSketch != null);
-					}
-					catch
-					{
-						Debug.Assert(false);
-					}
-				
-					if (!bSketchMode)
-					{
-#if DEBUG
-						Stopwatch Watch = Stopwatch.StartNew();
-#endif
-						bool bHasMaterialUpdates = DocumentTracker.HasMaterialUpdates();
-#if DEBUG
-						Watch.Stop();
-						Debug.WriteLine($"FMaterialUpdateChecker->HasMaterialUpdates: {(double)Watch.ElapsedMilliseconds / 1000.0}");
-#endif
-						if (bHasMaterialUpdates)
-						{
-							DocumentTracker.SetDirty(true);
-						}
-					}
-				}
-				Thread.Sleep(600);
-			}
-		}
-	}
 
 	// Build DatasmithScene from Solidworks Document, supporting scene changes/modifications
 	// Doesn't incorporate change notifications callbacks/handlers(i.e. notification handlers call into this class to tell what was changed and it updates Datasmith scene accordingly)
@@ -103,7 +20,7 @@ namespace DatasmithSolidworks
 		public readonly FDocument Doc;  // Top Document that is tracked/exported
 		public ModelDoc2 SwDoc => Doc.SwDoc;
 
-		public ConcurrentDictionary<int, FMaterial> ExportedMaterialsMap = new ConcurrentDictionary<int, FMaterial>();
+		public Dictionary<int, FMaterial> ExportedMaterialsMap = new Dictionary<int, FMaterial>();
 		public bool bHasConfigurations = false;
 
 		public FDatasmithFacadeScene DatasmithScene = null;
@@ -139,8 +56,6 @@ namespace DatasmithSolidworks
 			bDocumentIsDirty = bInDirty;
 		}
 
-		public abstract bool HasMaterialUpdates();
-
 		public uint GetFaceId(IFace2 InFace)
 		{
 			uint FaceId = unchecked((uint)InFace.GetFaceId());
@@ -304,9 +219,9 @@ namespace DatasmithSolidworks
 
 		public virtual void AddPartDocument(FConfigurationTree.FComponentTreeNode InNode){}
 
-		public virtual void AddExportedComponent(FConfigurationTree.FComponentTreeNode InNode) {}
+		public virtual void AddCollectedComponent(FConfigurationTree.FComponentTreeNode InNode) {}
 
-		public abstract ConcurrentDictionary<FComponentName, FObjectMaterials> LoadDocumentMaterials(HashSet<FComponentName> ComponentNamesToExportSet);
+		public abstract Dictionary<FComponentName, FObjectMaterials> LoadDocumentMaterials(HashSet<FComponentName> ComponentNamesToExportSet);
 		public abstract void AddComponentMaterials(FComponentName ComponentName, FObjectMaterials Materials);
 		public abstract FObjectMaterials GetComponentMaterials(Component2 Comp);
 
@@ -353,13 +268,13 @@ namespace DatasmithSolidworks
 		{
 			Exporter.AssignMaterialsToDatasmithMeshes(CreatedMeshes);
 		}
-
 	};
 
 
 	// Controls Syncing of a tracked Document, handling all events
 	public interface IDocumentSyncer
 	{
+
 		// Make Datasmith Scene up to date with the Solidworks Document
 		void Sync(string InOutputPath);
 
@@ -385,6 +300,8 @@ namespace DatasmithSolidworks
 
 		// Explicit immediate release of resources(line event handlers)
 		void Destroy();
+
+		void Idle();
 	}
 
 	/**
@@ -460,6 +377,7 @@ namespace DatasmithSolidworks
 
 		public void OnIdle()
 		{
+			DocumentSyncer.Idle();
 			if (bDirectLinkAutoSync && DocumentSyncer.GetDirty())
 			{
 				OnDirectLinkSync();
diff --git a/Engine/Source/Programs/Enterprise/Datasmith/DatasmithSolidworksExporter/Private/ObjectMaterials.cs b/Engine/Source/Programs/Enterprise/Datasmith/DatasmithSolidworksExporter/Private/ObjectMaterials.cs
--- a/Engine/Source/Programs/Enterprise/Datasmith/DatasmithSolidworksExporter/Private/ObjectMaterials.cs
+++ b/Engine/Source/Programs/Enterprise/Datasmith/DatasmithSolidworksExporter/Private/ObjectMaterials.cs
@@ -7,16 +7,20 @@ using System.Linq;
 using System.Collections.Concurrent;
 using DatasmithSolidworks.Names;
 using static DatasmithSolidworks.Addin;
+using static System.Windows.Forms.VisualStyles.VisualStyleElement;
 
 namespace DatasmithSolidworks
 {
+	// Contains material assignment structure for a Component in assembly or a Part document
+	// Materials(appearances) in SW can be assigned on different levels - document(top or sub-assembly/part), component, body, feature
+	// Highest priority has appearance assigned to Component(top-level assembly document). Then  Face, Feature, Body, Part (from highest to lowest priority).
+	// see "Appearance Hierarchy" in SW docs
 	public class FObjectMaterials
 	{
 		// Note: part document might be a different document than the owner one:
 		// for example, when loading component materials, owner will be the document 
 		// that component resides in, while part doc will be the document that component references!
-		public PartDoc PartDocument {  get; private set; }
-		public FDocumentTracker OwnerDoc {  get; private set; }
+		public PartDoc PartDocument;
 
 		public int ComponentMaterialID { get; private set; } = -1;
 		public int PartMaterialID { get; private set; } = -1;
@@ -25,12 +29,11 @@ namespace DatasmithSolidworks
 		public Dictionary<string, int> FaceMaterialsMap { get; private set; } = new Dictionary<string, int>();
 		public Dictionary<string, int> FeatureMaterialsMap { get; private set; } = new Dictionary<string, int>();
 
-		public ConcurrentDictionary<int, FMaterial> GlobalMaterialsMap { get; private set; } = null;
+		private readonly Dictionary<int, FMaterial> GlobalMaterialsMap;  // All collected materials
 
-		public FObjectMaterials(FDocumentTracker InOwnerDoc, PartDoc InPartDocument, ref ConcurrentDictionary<int, FMaterial> InOutMaterialsMap)
+		public FObjectMaterials(PartDoc InPartDocument, Dictionary<int, FMaterial> InOutMaterialsMap)
 		{
 			GlobalMaterialsMap = InOutMaterialsMap;
-			OwnerDoc = InOwnerDoc;
 			PartDocument = InPartDocument;
 		}
 
@@ -49,7 +52,7 @@ namespace DatasmithSolidworks
 			int MatID = MaterialUtils.GetMaterialID(InRenderMat);
 			if (!GlobalMaterialsMap.ContainsKey(MatID))
 			{
-				GlobalMaterialsMap.TryAdd(MatID, new FMaterial(InRenderMat, InDoc.Extension));
+				GlobalMaterialsMap[MatID] = new FMaterial(InRenderMat, InDoc.Extension);
 			}
 			ComponentMaterialID = MatID;
 		}
@@ -60,7 +63,7 @@ namespace DatasmithSolidworks
 			if (!GlobalMaterialsMap.ContainsKey(MatID))
 			{
 				FMaterial Mat = new FMaterial(InRenderMat, InDoc.Extension);
-				GlobalMaterialsMap.TryAdd(MatID, Mat);
+				GlobalMaterialsMap[MatID] = Mat;
 			}
 			PartMaterialID = MatID;
 		}
@@ -165,7 +168,7 @@ namespace DatasmithSolidworks
 
 			if (!GlobalMaterialsMap.ContainsKey(MaterialID))
 			{
-				GlobalMaterialsMap.TryAdd(MaterialID, new FMaterial(RenderMat, InDoc.Extension));
+				GlobalMaterialsMap[MaterialID] = new FMaterial(RenderMat, InDoc.Extension);
 			}
 		}
 
@@ -174,7 +177,7 @@ namespace DatasmithSolidworks
 			return InDict1.OrderBy(KVP => KVP.Key).SequenceEqual(InDict2.OrderBy(KVP => KVP.Key));
 		}
 
-		public static FObjectMaterials LoadPartMaterials(FDocumentTracker InOwnerDoc, PartDoc InPartDoc, swDisplayStateOpts_e InDisplayState, string[] InDisplayStateNames)
+		public static FObjectMaterials LoadPartMaterials(FDocumentTracker InOwnerDoc, PartDoc InPartDoc, swDisplayStateOpts_e InDisplayState, string[] InDisplayStateNames, Dictionary<int, FMaterial> MaterialsMap)
 		{
 			LogDebug($"LoadPartMaterials: '{(InPartDoc as ModelDoc2).GetPathName()}', {InDisplayState}, {(InDisplayStateNames == null ? "<null>" : string.Join(", ", InDisplayStateNames))}");
 
@@ -197,7 +200,7 @@ namespace DatasmithSolidworks
 
 			object[] ObjMaterials = Ext.GetRenderMaterials2((int)InDisplayState, InDisplayStateNames);
 
-			FObjectMaterials PartMaterials = new FObjectMaterials(InOwnerDoc, InPartDoc, ref InOwnerDoc.ExportedMaterialsMap);
+			FObjectMaterials PartMaterials = new FObjectMaterials(InPartDoc, MaterialsMap);
 
 			foreach (object ObjMat in ObjMaterials)
 			{
@@ -230,7 +233,7 @@ namespace DatasmithSolidworks
 						}
 						case IFace2 Face:
 						{
-							uint FaceId = PartMaterials.OwnerDoc.GetFaceId(Face);
+							uint FaceId = InOwnerDoc.GetFaceId(Face);
 							LogDebug($"    Face '{Face.GetFaceId()}'");
 							PartMaterials.RegisterMaterial(PartMaterials.FaceMaterialsMap, Doc, RenderMat, FaceId.ToString());
 							continue;
@@ -287,7 +290,7 @@ namespace DatasmithSolidworks
 
 						LogDebug($"  Body '{Body.Name}' has no material assigned, searching material matching its Appearance..");
 
-						foreach (KeyValuePair<int, FMaterial> KVP in PartMaterials.GlobalMaterialsMap)
+						foreach (KeyValuePair<int, FMaterial> KVP in MaterialsMap)
 						{
 							FMaterial Material = KVP.Value;
 
@@ -344,7 +347,7 @@ namespace DatasmithSolidworks
 
 						LogDebug($"  Feature '{Feat.Name}' has no material assigned, searching material matching its Appearance..");
 
-						foreach (KeyValuePair<int, FMaterial> KVP in PartMaterials.GlobalMaterialsMap)
+						foreach (KeyValuePair<int, FMaterial> KVP in MaterialsMap)
 						{
 							FMaterial Material = KVP.Value;
 
@@ -363,7 +366,9 @@ namespace DatasmithSolidworks
 			return PartMaterials;
 		}
 
-		public static FObjectMaterials LoadComponentMaterials(FDocumentTracker InComponentOwner, Component2 InComponent, swDisplayStateOpts_e InDisplayState, string[] InDisplayStateNames)
+		public static FObjectMaterials LoadComponentMaterials(FDocumentTracker InComponentOwner, Component2 InComponent,
+			swDisplayStateOpts_e InDisplayState, string[] InDisplayStateNames,
+			Dictionary<int, FMaterial> MaterialsMap)
 		{
 			LogDebug($"LoadComponentMaterials: '{InComponent.Name2}', '{InComponentOwner.GetPathName()}', {InDisplayState}, {(InDisplayStateNames == null ? "<null>" : string.Join(", ", InDisplayStateNames))}");
 
@@ -386,7 +391,7 @@ namespace DatasmithSolidworks
 
 			object[] ObjMaterials = InComponent.GetRenderMaterials2((int)InDisplayState, InDisplayStateNames);
 
-			FObjectMaterials ComponentMaterials = new FObjectMaterials(InComponentOwner, ComponentDoc as PartDoc, ref InComponentOwner.ExportedMaterialsMap);
+			FObjectMaterials ComponentMaterials = new FObjectMaterials(ComponentDoc as PartDoc, MaterialsMap);
 
 			if (ObjMaterials != null)
 			{
@@ -419,7 +424,7 @@ namespace DatasmithSolidworks
 							case IFace2 Face:
 							{
 								// note: plugins sets FaceId on each face to identify per-face materials later in CreateMeshData
-								uint FaceId = ComponentMaterials.OwnerDoc.GetFaceId(Face);
+								uint FaceId = InComponentOwner.GetFaceId(Face);
 								LogDebug($"    Face '{Face.GetFaceId()}'");
 								ComponentMaterials.RegisterMaterial(ComponentMaterials.FaceMaterialsMap,
 									InComponentOwner.SwDoc, RenderMat, FaceId.ToString());
@@ -452,18 +457,104 @@ namespace DatasmithSolidworks
 			return ComponentMaterials;
 		}
 
-		public static ConcurrentDictionary<FComponentName, FObjectMaterials> LoadAssemblyMaterials(FAssemblyDocumentTracker InAsmDoc, HashSet<FComponentName> InComponentsSet, swDisplayStateOpts_e InDisplayState, string[] InDisplayStateNames)
+		// Load all materials in the assembly - including subassemblies and parts within
+		// returns: material assignment for each component
+		public static Dictionary<FComponentName, FObjectMaterials> LoadAssemblyMaterials(FAssemblyDocumentTracker InAsmDoc, HashSet<FComponentName> InComponentsSet, swDisplayStateOpts_e InDisplayState, string[] InDisplayStateNames)
 		{
 			LogDebug($"LoadAssemblyMaterials: {InAsmDoc.GetPathName()}, {InDisplayState}, [{(InDisplayStateNames == null ? "" : string.Join(",", InDisplayStateNames))}]");
 			LogDebug($"  for components: [{string.Join(", ", InComponentsSet.OrderBy(Name => Name.GetString()))}]");
 
+			Dictionary<int, FMaterial> MaterialsMap = InAsmDoc.ExportedMaterialsMap;
+			Dictionary<FComponentName, FObjectMaterials> DocMaterials = new Dictionary<FComponentName, FObjectMaterials>();
+
+			LoadMaterialsWithUsers(InAsmDoc, InDisplayState, InDisplayStateNames, DocMaterials, MaterialsMap);
+
+			LogDebug($"Load Each Component's materials");
+			LogIndent();
+
+			foreach (FComponentName CompName in InComponentsSet.OrderBy(Name => Name.GetString()))
+			{
+				UpdateComponentMaterials(InAsmDoc, CompName, DocMaterials, MaterialsMap);
+			}
+
+			LogDedent();
+
+			return DocMaterials.Count > 0 ? DocMaterials : null;
+		}
+
+		public static IEnumerable<bool> LoadAssemblyMaterialsEnum(FAssemblyDocumentTracker InAsmDoc, HashSet<FComponentName> InComponentsSet, Dictionary<FComponentName, FObjectMaterials> CurrentDocMaterialsMap, Dictionary<int, FMaterial> MaterialsMap)
+		{
+			FObjectMaterials.LoadMaterialsWithUsers(InAsmDoc, swDisplayStateOpts_e.swThisDisplayState, null,
+				CurrentDocMaterialsMap, MaterialsMap);
+			yield return true;
+
+			foreach (FComponentName CompName in InComponentsSet.OrderBy(Name => Name.GetString()))
+			{
+				FObjectMaterials.UpdateComponentMaterials(InAsmDoc, CompName, CurrentDocMaterialsMap,
+					MaterialsMap);
+				yield return true;
+			}
+		}
+
+		public static void UpdateComponentMaterials(FAssemblyDocumentTracker InAsmDoc, FComponentName CompName,
+			Dictionary<FComponentName, FObjectMaterials> DocMaterials, Dictionary<int, FMaterial> MaterialsMap)
+		{
+			LogDebug($"Component '{CompName}' - adding DocMaterials");
+			LogIndent();
+
+			if (InAsmDoc.SyncState.CollectedComponentsMap.TryGetValue(CompName, out Component2 Comp))
+			{
+				if (!DocMaterials.ContainsKey(CompName))
+				{
+					LogDebug($"Adding to DocMaterials");
+
+					LogIndent();
+					FObjectMaterials ComponentMaterials = FObjectMaterials.LoadComponentMaterials(InAsmDoc,
+						Comp, swDisplayStateOpts_e.swThisDisplayState, null, MaterialsMap);
+
+					LogDedent();
+
+					if (ComponentMaterials == null)
+					{
+						if (Comp.GetModelDoc2() is PartDoc)
+						{
+							LogIndent();
+							ComponentMaterials = FObjectMaterials.LoadPartMaterials(InAsmDoc,
+								Comp.GetModelDoc2() as PartDoc, swDisplayStateOpts_e.swThisDisplayState, null,
+								MaterialsMap);
+							LogDedent();
+						}
+					}
+
+					if (ComponentMaterials != null)
+					{
+						DocMaterials[CompName] = ComponentMaterials;
+					}
+				}
+				else
+				{
+					LogDebug($"...already in DocMaterials");
+				}
+			}
+			else
+			{
+				LogDebug($"...not in CollectedComponentsMap");
+			}
+
+			LogDedent();
+		}
+
+		// Parse materials within document and return components using them
+		public static void LoadMaterialsWithUsers(FAssemblyDocumentTracker InAsmDoc,
+			swDisplayStateOpts_e InDisplayState,
+			string[] InDisplayStateNames, Dictionary<FComponentName, FObjectMaterials> DocMaterials,
+			Dictionary<int, FMaterial> MaterialsMap)
+		{
 			IModelDocExtension Ext = InAsmDoc.SwDoc.Extension;
 			int NumMaterials = Ext.GetRenderMaterialsCount2((int)InDisplayState, InDisplayStateNames);
 
 			LogDebug($"  Assembly materials={NumMaterials}");
 
-			ConcurrentDictionary<FComponentName, FObjectMaterials> DocMaterials = new ConcurrentDictionary<FComponentName, FObjectMaterials>();
-
 			if (NumMaterials > 0)
 			{
 				object[] ObjMaterials = Ext.GetRenderMaterials2((int)InDisplayState, InDisplayStateNames);
@@ -471,7 +562,7 @@ namespace DatasmithSolidworks
 				foreach (object ObjMat in ObjMaterials)
 				{
 					RenderMaterial RenderMat = ObjMat as RenderMaterial;
-					
+
 					LogDebug($"    Material: '{RenderMat.FileName}' ");
 
 					int NumUsers = RenderMat.GetEntitiesCount();
@@ -489,9 +580,8 @@ namespace DatasmithSolidworks
 						if (ObjUser is IEntity Entity)
 						{
 							LogDebug($"    User: type '{Entity.GetType()}' ");
-							
 						}
-						
+
 						switch (ObjUser)
 						{
 							case Component2 Comp:
@@ -500,7 +590,7 @@ namespace DatasmithSolidworks
 								if (Comp.GetModelDoc2() is PartDoc CompDoc)
 								{
 									LogDebug($"        registering material for Part Component in DocMaterials");
-									FObjectMaterials ComponentMaterials = new FObjectMaterials(InAsmDoc, CompDoc, ref InAsmDoc.ExportedMaterialsMap);
+									FObjectMaterials ComponentMaterials = new FObjectMaterials(CompDoc, MaterialsMap);
 
 									LogIndent();
 									ComponentMaterials.SetComponentMaterial(RenderMat, InAsmDoc.SwDoc);
@@ -508,6 +598,7 @@ namespace DatasmithSolidworks
 
 									DocMaterials[new FComponentName(Comp)] = ComponentMaterials;
 								}
+
 								break;
 							}
 							case IPartDoc Doc:
@@ -515,67 +606,25 @@ namespace DatasmithSolidworks
 								LogDebug($"      Part Doc at '{(Doc as ModelDoc2).GetPathName()}' ");
 
 								LogIndent();
-								FObjectMaterials PartMaterials = LoadPartMaterials(InAsmDoc, Doc as PartDoc, InDisplayState, InDisplayStateNames);
+								FObjectMaterials PartMaterials = LoadPartMaterials(InAsmDoc, Doc as PartDoc, InDisplayState,
+									InDisplayStateNames, MaterialsMap);
 								LogDedent();
 
 								if (PartMaterials != null)
 								{
 									// todo: why trying to present PathName as component name?
 									//  how it can be a PartDoc in Assembly function
-									DocMaterials[FComponentName.FromCustomString((Doc as ModelDoc2).GetPathName())] = PartMaterials;
+									DocMaterials[FComponentName.FromCustomString((Doc as ModelDoc2).GetPathName())] =
+										PartMaterials;
 								}
+
 								break;
 							}
 						}
 					}
 				}
 			}
-
-			LogDebug($"Load Each Component's materials");
-			LogIndent();
-			foreach(FComponentName CompName in InComponentsSet.OrderBy(Name => Name.GetString()))
-			{
-				Component2 Comp = null;
-
-				LogDebug($"Component '{CompName}' - adding DocMaterials");
-				LogIndent();
-
-				if (InAsmDoc.SyncState.ExportedComponentsMap.TryGetValue(CompName, out Comp))
-				{
-					if (!DocMaterials.ContainsKey(CompName))
-					{
-						LogDebug($"Adding to DocMaterials");
-
-						LogIndent();
-						FObjectMaterials ComponentMaterials = FObjectMaterials.LoadComponentMaterials(InAsmDoc, Comp, swDisplayStateOpts_e.swThisDisplayState, null);
-						LogDedent();
-
-						if (ComponentMaterials == null && Comp.GetModelDoc2() is PartDoc)
-						{
-							LogIndent();
-							ComponentMaterials = FObjectMaterials.LoadPartMaterials(InAsmDoc, Comp.GetModelDoc2() as PartDoc, swDisplayStateOpts_e.swThisDisplayState, null);
-							LogDedent();
-						}
-						if (ComponentMaterials != null)
-						{
-							DocMaterials.TryAdd(CompName, ComponentMaterials);
-						}
-					}
-					else
-					{
-						LogDebug($"...already in DocMaterials");
-					}
-				}
-				else
-				{
-					LogDebug($"...not in ExportedComponentsMap");
-				}
-
-				LogDedent();
-			}
-			LogDedent();
-
-			return DocMaterials.Count > 0 ? DocMaterials : null;
 		}
+
 	}
 }
diff --git a/Engine/Source/Programs/Enterprise/Datasmith/DatasmithSolidworksExporter/Private/PartDocument.cs b/Engine/Source/Programs/Enterprise/Datasmith/DatasmithSolidworksExporter/Private/PartDocument.cs
--- a/Engine/Source/Programs/Enterprise/Datasmith/DatasmithSolidworksExporter/Private/PartDocument.cs
+++ b/Engine/Source/Programs/Enterprise/Datasmith/DatasmithSolidworksExporter/Private/PartDocument.cs
@@ -10,9 +10,136 @@ using DatasmithSolidworks.Names;
 using static DatasmithSolidworks.FDatasmithExporter;
 using System.Diagnostics;
 using System.Linq;
+using static DatasmithSolidworks.FAssemblyDocumentTracker;
 
 namespace DatasmithSolidworks
 {
+	// Checks if scene is updated iteratively
+	//
+	// Solidworks notifications lack a number of notifications for material changes - 
+	// material (re)assignment to various entity types and change in material properties themselves
+	// This class lazily iterates over all the components to identify these changes
+	// This lazy process it performed in small portions ove multiple 'ticks' in the main thread.
+	// It's done in the main thread(not in a separate thread) for simple reason that SW api 
+	// is very slow to use from another thread(although it may be convenient for simple tasks)
+	// Doing this in a thread takes dozens of seconds for models which take a second to fully export.
+	// Making this unusable. And also unstable as any change to the model in the main thread may result in a crash in the thread
+	public abstract class FLazyUpdateCheckerBase
+	{
+		private readonly FDocumentTracker DocumentTracker;
+		private int MaxTickTimeMilliseconds = 100;
+		private IEnumerator<bool> Updating = null;
+
+		protected FLazyUpdateCheckerBase(FDocumentTracker InDocumentTracker)
+		{
+			DocumentTracker = InDocumentTracker;
+		}
+
+		public void Reset()
+		{
+			if (Updating != null)
+			{
+				Updating.Dispose();
+				Updating = null;
+			}
+		}
+
+		// Advances checking for updates in the model
+		public void Tick()
+		{
+			if (DocumentTracker.GetDirty())
+			{
+				// If document is already market dirty don't need to check for modifications
+				Reset();
+				return;
+			}
+
+			if (Updating == null)
+			{
+				Restart();
+				Updating = CheckForModificationsEnum().GetEnumerator();
+			}
+
+			Stopwatch Watch = Stopwatch.StartNew();
+
+			bool bPostponed = false;
+			try
+			{
+				while (Updating.MoveNext())
+				{
+					if (Watch.ElapsedMilliseconds > MaxTickTimeMilliseconds)
+					{
+						// Postpone the reset of update if it has already taken too long for this tick, so the application doesn't stutter
+						bPostponed = true;
+						break;
+					}
+				}
+			}
+			catch
+			{
+				// This will just restart the checker in case any exception occurs during update
+			}
+
+			if (HasUpdates())
+			{
+				DocumentTracker.SetDirty(true);
+			}
+
+			if (!bPostponed)
+			{
+				Reset();
+			}
+		}
+
+		/// Reset state of updates check
+		public abstract void Restart();
+
+		/// Get state of updates check
+		public abstract bool HasUpdates();
+
+		///  Resumable check for modifications
+		public abstract IEnumerable<bool> CheckForModificationsEnum();
+	};
+
+	public class FPartLazyUpdateChecker: FLazyUpdateCheckerBase
+	{
+		private FPartDocumentTracker PartDocumentTracker;
+		private bool bHasDirtyMaterials;
+
+		public FPartLazyUpdateChecker(FPartDocumentTracker InDocumentTracker) : base(InDocumentTracker)
+		{
+			PartDocumentTracker = InDocumentTracker;
+		}
+
+		public override void Restart()
+		{
+			bHasDirtyMaterials = false;
+		}
+
+		public override bool HasUpdates()
+		{
+			return bHasDirtyMaterials;
+		}
+
+		public override IEnumerable<bool> CheckForModificationsEnum()
+		{
+			FObjectMaterials CurrentMaterials =
+				FObjectMaterials.LoadPartMaterials(PartDocumentTracker, PartDocumentTracker.SwPartDoc, swDisplayStateOpts_e.swThisDisplayState, null, PartDocumentTracker.ExportedMaterialsMap);
+			if (CurrentMaterials == null && PartDocumentTracker.ExportedPartMaterials == null)
+			{
+				yield break;
+			}
+			else if ((CurrentMaterials != null && PartDocumentTracker.ExportedPartMaterials == null) ||
+			         (CurrentMaterials == null && PartDocumentTracker.ExportedPartMaterials != null))
+			{
+				bHasDirtyMaterials = true;
+				yield break;
+			}
+
+			bHasDirtyMaterials = !CurrentMaterials.EqualMaterials(PartDocumentTracker.ExportedPartMaterials);
+		}
+	};
+
 	public class FPartDocumentTracker : FDocumentTracker
 	{
 		public PartDoc SwPartDoc { get; private set; } = null;
@@ -27,6 +154,8 @@ namespace DatasmithSolidworks
 
 		public string PathName { get; private set; }
 
+		private FPartLazyUpdateChecker LazyUpdateChecker;
+
 		public FPartDocumentTracker(FPartDocument InDoc, PartDoc InPartDoc, FDatasmithExporter InExporter,
 			FAssemblyDocumentTracker InAsmDoc, FComponentName InComponentName)
 			: base(InDoc, InExporter)
@@ -217,23 +346,6 @@ namespace DatasmithSolidworks
 			
 		}
 
-		public override bool HasMaterialUpdates()
-		{
-			FObjectMaterials CurrentMaterials =
-				FObjectMaterials.LoadPartMaterials(this, SwPartDoc, swDisplayStateOpts_e.swThisDisplayState, null);
-			if (CurrentMaterials == null && ExportedPartMaterials == null)
-			{
-				return false;
-			}
-			else if ((CurrentMaterials != null && ExportedPartMaterials == null) ||
-			         (CurrentMaterials == null && ExportedPartMaterials != null))
-			{
-				return true;
-			}
-
-			return !CurrentMaterials.EqualMaterials(ExportedPartMaterials);
-		}
-
 		public override FObjectMaterials GetComponentMaterials(Component2 Comp)
 		{
 			return null;
@@ -259,7 +371,8 @@ namespace DatasmithSolidworks
 		// Extract mesh data for current active configuration. MeshData is ready for export to DatasmithMesh 
 		public FMeshData ExtractPartMeshData()
 		{
-			FObjectMaterials PartMaterials = FObjectMaterials.LoadPartMaterials(this, SwPartDoc, swDisplayStateOpts_e.swThisDisplayState, null);
+			FObjectMaterials PartMaterials = FObjectMaterials.LoadPartMaterials(this, SwPartDoc, swDisplayStateOpts_e.swThisDisplayState, null, ExportedMaterialsMap);
+			ExportedPartMaterials = PartMaterials;
 			ConcurrentBag<FBody> Bodies = FBody.FetchBodies(SwPartDoc);
 			return FStripGeometry.CreateMeshData(Bodies, PartMaterials);
 		}
@@ -268,10 +381,20 @@ namespace DatasmithSolidworks
 		{
 		}
 
-		public override ConcurrentDictionary<FComponentName, FObjectMaterials> LoadDocumentMaterials(HashSet<FComponentName> ComponentNamesToExportSet)
+		public override Dictionary<FComponentName, FObjectMaterials> LoadDocumentMaterials(HashSet<FComponentName> ComponentNamesToExportSet)
 		{
 			return null;
 		}
+
+		public void Tick()
+		{
+			LazyUpdateChecker?.Tick();
+		}
+
+		public void TrackChanges()
+		{
+			LazyUpdateChecker = new FPartLazyUpdateChecker(this);
+		}
 	}
 
 	public class FPartDocumentEvents
@@ -459,7 +582,6 @@ namespace DatasmithSolidworks
 	{
 		private readonly FPartDocumentTracker DocumentTracker;
 		private readonly FPartDocumentEvents Events;
-		private FMaterialUpdateChecker MaterialUpdateChecker;
 
 		public FPartDocumentNotifications(int InDocId, PartDoc InPartDoc, FPartDocumentTracker InDocumentTracker)
 		{
@@ -469,26 +591,26 @@ namespace DatasmithSolidworks
 
 		public void Destroy()
 		{
-			MaterialUpdateChecker?.Destroy();
 			Events.Destroy();
 		}
 
+
 		public void Start()
 		{
-			if (MaterialUpdateChecker == null)
-			{
-				MaterialUpdateChecker = new FMaterialUpdateChecker(DocumentTracker);
-			}
+			DocumentTracker.TrackChanges();
 		}
 
 		public void Resume()
 		{
-			MaterialUpdateChecker?.Resume();
 		}
 
 		public void Pause()
 		{
-			MaterialUpdateChecker?.Pause();
+		}
+
+		public void Idle()
+		{
+			DocumentTracker.Tick();
 		}
 	}
 	
@@ -509,6 +631,11 @@ namespace DatasmithSolidworks
 			DocumentTracker.Destroy();
 		}
 
+		public void Idle()
+		{
+			Notifications.Idle();
+		}
+
 		public void Start()
 		{
 			Notifications.Start();
