Added C# facade to DatasmithSDK package
Updated Mac and Linux build to match Windows.

#rnx

diff --git a/Engine/Source/Developer/Datasmith/DatasmithExporterUI/DatasmithExporterUI.Build.cs b/Engine/Source/Developer/Datasmith/DatasmithExporterUI/DatasmithExporterUI.Build.cs
--- a/Engine/Source/Developer/Datasmith/DatasmithExporterUI/DatasmithExporterUI.Build.cs
+++ b/Engine/Source/Developer/Datasmith/DatasmithExporterUI/DatasmithExporterUI.Build.cs
@@ -4,7 +4,7 @@ using System;
 
 namespace UnrealBuildTool.Rules
 {
-	[SupportedPlatforms("Win64", "Mac")]
+	[SupportedPlatforms("Win64", "Mac", "Linux")]
 	public class DatasmithExporterUI : ModuleRules
 	{
 		public DatasmithExporterUI(ReadOnlyTargetRules Target)
diff --git a/Engine/Source/Programs/Enterprise/Datasmith/DatasmithSDK/DatasmithSDK.Build.cs b/Engine/Source/Programs/Enterprise/Datasmith/DatasmithSDK/DatasmithSDK.Build.cs
--- a/Engine/Source/Programs/Enterprise/Datasmith/DatasmithSDK/DatasmithSDK.Build.cs
+++ b/Engine/Source/Programs/Enterprise/Datasmith/DatasmithSDK/DatasmithSDK.Build.cs
@@ -15,7 +15,9 @@ namespace UnrealBuildTool.Rules
 				{
 					"DatasmithCore",
 					"DatasmithExporter",
-					"UdpMessaging", // required for DirectLink networking
+					"DatasmithExporterUI",
+					// Network layer
+					"UdpMessaging",
 				}
 			);
 		}
diff --git a/Engine/Source/Programs/Enterprise/Datasmith/DatasmithSDK/DatasmithSDK.Target.cs b/Engine/Source/Programs/Enterprise/Datasmith/DatasmithSDK/DatasmithSDK.Target.cs
--- a/Engine/Source/Programs/Enterprise/Datasmith/DatasmithSDK/DatasmithSDK.Target.cs
+++ b/Engine/Source/Programs/Enterprise/Datasmith/DatasmithSDK/DatasmithSDK.Target.cs
@@ -2,7 +2,11 @@
 
 using UnrealBuildTool;
 using System.Collections.Generic;
+using System.IO;
 
+public delegate void PostCopyFunc(string SrcPath, string DestPath);
+
+[SupportedPlatforms("Win64", "Mac", "Linux")]
 public class DatasmithSDKTarget : TargetRules
 {
 	public DatasmithSDKTarget(TargetInfo Target)
@@ -34,11 +38,14 @@ public class DatasmithSDKTarget : TargetRules
 		bHasExports = true;
 		bIsBuildingConsoleApplication = true;
 
-		if (Platform == UnrealTargetPlatform.Win64)
+		PostCopyFunc PostBuildCopy = PostBuildCopyWin64;
+		if (Platform != UnrealTargetPlatform.Win64)
 		{
-			AddWindowsPostBuildSteps();
+			PostBuildCopy = PostBuildCopyUnix;
 		}
 
+		AddPostBuildSteps(PostBuildCopy);
+
 		// Enable UDP in shipping (used by DirectLink)
 		if (BuildEnvironment == TargetBuildEnvironment.Unique)
 		{
@@ -46,14 +53,22 @@ public class DatasmithSDKTarget : TargetRules
 			GlobalDefinitions.Add("PLATFORM_SUPPORTS_MESSAGEBUS=1"); // required to enable the default MessageBus in MessagingModule.cpp
 		}
 	}
-
-	public void PostBuildCopy(string SrcPath, string DestPath)
+	public void PostBuildCopyWin64(string SrcPath, string DestPath)
 	{
 		PostBuildSteps.Add(string.Format("echo Copying {0} to {1}", SrcPath, DestPath));
 		PostBuildSteps.Add(string.Format("xcopy \"{0}\" \"{1}\" /R /Y /S /Q", SrcPath, DestPath));
 	}
 
-	public void AddWindowsPostBuildSteps()
+	public void PostBuildCopyUnix(string SrcPath, string DestPath)
+	{
+		SrcPath = SrcPath.Replace("\\", "/");
+		DestPath = DestPath.Replace("\\", "/");
+		PostBuildSteps.Add(string.Format("echo Copying {0} to {1}\n", SrcPath, DestPath));
+		PostBuildSteps.Add(string.Format("mkdir -p {0}\n", DestPath));
+		PostBuildSteps.Add(string.Format("cp -R -f {0} {1}\n", SrcPath, DestPath));
+	}
+
+	public void AddPostBuildSteps(PostCopyFunc PostBuildCopy)
 	{
 		// Copy the documentation
 		PostBuildCopy(
@@ -77,6 +92,11 @@ public class DatasmithSDKTarget : TargetRules
 			@"$(EngineDir)\Binaries\$(TargetPlatform)\DatasmithSDK\Public\"
 		);
 
+		PostBuildCopy(
+			@"$(EngineDir)\Source\Developer\Datasmith\DatasmithExporterUI\Public\*.h",
+			@"$(EngineDir)\Binaries\$(TargetPlatform)\DatasmithSDK\Public\"
+		);
+
 		PostBuildCopy(
 			@"$(EngineDir)\Extras\VisualStudioDebugging\Unreal.natvis",
 			@"$(EngineDir)\Binaries\$(TargetPlatform)\DatasmithSDK\"
