Added Engine files to run UE's UI

diff --git a/Engine/Source/Programs/Enterprise/Datasmith/DatasmithSDK/CopyCmd.bat b/Engine/Source/Programs/Enterprise/Datasmith/DatasmithSDK/CopyCmd.bat
new file mode 100755
--- /dev/null
+++ b/Engine/Source/Programs/Enterprise/Datasmith/DatasmithSDK/CopyCmd.bat
@@ -0,0 +1,5 @@
+@echo off
+
+robocopy %1 %2 %3 /J /Z /S /R:1 /W:0 /NS /NC /NFL /NDL /NP > nul
+
+exit /B 0
\ No newline at end of file
diff --git a/Engine/Source/Programs/Enterprise/Datasmith/DatasmithSDK/DatasmithSDK.Target.cs b/Engine/Source/Programs/Enterprise/Datasmith/DatasmithSDK/DatasmithSDK.Target.cs
--- a/Engine/Source/Programs/Enterprise/Datasmith/DatasmithSDK/DatasmithSDK.Target.cs
+++ b/Engine/Source/Programs/Enterprise/Datasmith/DatasmithSDK/DatasmithSDK.Target.cs
@@ -3,8 +3,11 @@
 using UnrealBuildTool;
 using System.Collections.Generic;
 using System.IO;
+using System;
+using UnrealBuildTool.Rules;
 
-public delegate void PostCopyFunc(string SrcPath, string DestPath);
+public delegate void PostCopyFilesFunc(string SrcPath, string DestPath, string Files);
+public delegate void PostCopyDirFunc(string SrcPath, string DestPath);
 
 [SupportedPlatforms("Win64", "Mac", "Linux")]
 public class DatasmithSDKTarget : TargetRules
@@ -38,13 +41,36 @@ public class DatasmithSDKTarget : TargetRules
 		bHasExports = true;
 		bIsBuildingConsoleApplication = true;
 
-		PostCopyFunc PostBuildCopy = PostBuildCopyWin64;
+		KeyValuePair<string, string>[] FilesToCopy = new KeyValuePair<string, string>[]
+		{
+
+		};
+
+		PostCopyFilesFunc PostBuildCopy = PostBuildCopyWin64;
 		if (Platform != UnrealTargetPlatform.Win64)
 		{
+			List<string> Directories = new List<string>(
+				new string[]
+				{
+					@"$(EngineDir)/Binaries/$(TargetPlatform)/DatasmithSDK/Documentation/",
+					@"$(EngineDir)/Binaries/$(TargetPlatform)/DatasmithSDK/Public/",
+					@"$(EngineDir)/Binaries/$(TargetPlatform)/DatasmithSDK/Private/",
+					@"$(EngineDir)/Binaries/$(TargetPlatform)/DatasmithSDK/Engine/",
+					@"$(EngineDir)/Binaries/$(TargetPlatform)/DatasmithSDK/Engine/Shaders/StandaloneRenderer/D3D/",
+					@"$(EngineDir)/Binaries/$(TargetPlatform)/DatasmithSDK/Engine/Content/",
+				}
+			);
+
+			PostBuildSteps.Add(string.Format("echo Creating directories to export to\n"));
+			foreach (string Directory in Directories)
+			{
+				PostBuildSteps.Add(string.Format("mkdir -p {0}\n", Directory));
+			}
+
 			PostBuildCopy = PostBuildCopyUnix;
 		}
 
-		AddPostBuildSteps(PostBuildCopy);
+		AddPostBuildSteps();
 
 		// Enable UDP in shipping (used by DirectLink)
 		if (BuildEnvironment == TargetBuildEnvironment.Unique)
@@ -53,83 +79,86 @@ public class DatasmithSDKTarget : TargetRules
 			GlobalDefinitions.Add("PLATFORM_SUPPORTS_MESSAGEBUS=1"); // required to enable the default MessageBus in MessagingModule.cpp
 		}
 	}
-	public void PostBuildCopyWin64(string SrcPath, string DestPath)
+	public void PostBuildCopyWin64(string SrcPath, string DestPath, string Files)
+	{
+		PostBuildSteps.Add(string.Format("echo Copying {0}\\{2} to {1}", SrcPath, DestPath, Files));
+		PostBuildSteps.Add(string.Format(@"call $(EngineDir)\Source\Programs\Enterprise\Datasmith\DatasmithSDK\CopyCmd.bat {0} {1} {2}", SrcPath, DestPath, Files));
+	}
+
+	public void PostBuildCopyUnix(string SrcPath, string DestPath, string Files)
 	{
-		PostBuildSteps.Add(string.Format("echo Copying {0} to {1}", SrcPath, DestPath));
-		PostBuildSteps.Add(string.Format("xcopy \"{0}\" \"{1}\" /R /Y /S /Q", SrcPath, DestPath));
+		SrcPath = SrcPath.Replace("\\", "/");
+		DestPath = DestPath.Replace("\\", "/");
+		PostBuildSteps.Add(string.Format("echo Copying \"{0}\" to {1}\n", SrcPath, DestPath));
+		PostBuildSteps.Add(string.Format("cp -R -f {0}/{2} {1}\n", SrcPath, DestPath, Files));
 	}
 
 	public void PostBuildCopyUnix(string SrcPath, string DestPath)
 	{
 		SrcPath = SrcPath.Replace("\\", "/");
 		DestPath = DestPath.Replace("\\", "/");
-		PostBuildSteps.Add(string.Format("echo Copying {0} to {1}\n", SrcPath, DestPath));
-		PostBuildSteps.Add(string.Format("mkdir -p {0}\n", DestPath));
+		PostBuildSteps.Add(string.Format("echo Copying \"{0}\" to {1}\n", SrcPath, DestPath));
 		PostBuildSteps.Add(string.Format("cp -R -f {0} {1}\n", SrcPath, DestPath));
 	}
 
-	public void AddPostBuildSteps(PostCopyFunc PostBuildCopy)
+	public void AddPostBuildSteps()
 	{
-		// Copy the documentation
-		PostBuildCopy(
-			@"$(EngineDir)\Source\Programs\Enterprise\Datasmith\DatasmithSDK\Documentation\*.*",
-			@"$(EngineDir)\Binaries\$(TargetPlatform)\DatasmithSDK\Documentation\"
-		);
-
-		// Package our public headers
-		PostBuildCopy(
-			@"$(EngineDir)\Source\Runtime\Datasmith\DatasmithCore\Public\*.h",
-			@"$(EngineDir)\Binaries\$(TargetPlatform)\DatasmithSDK\Public\"
-		);
-
-		PostBuildCopy(
-			@"$(EngineDir)\Source\Runtime\Datasmith\DirectLink\Public\*.h",
-			@"$(EngineDir)\Binaries\$(TargetPlatform)\DatasmithSDK\Public\"
-		);
-
-		PostBuildCopy(
-			@"$(EngineDir)\Source\Developer\Datasmith\DatasmithExporter\Public\*.h",
-			@"$(EngineDir)\Binaries\$(TargetPlatform)\DatasmithSDK\Public\"
-		);
-
-		PostBuildCopy(
-			@"$(EngineDir)\Source\Developer\Datasmith\DatasmithExporterUI\Public\*.h",
-			@"$(EngineDir)\Binaries\$(TargetPlatform)\DatasmithSDK\Public\"
-		);
-
-		PostBuildCopy(
-			@"$(EngineDir)\Extras\VisualStudioDebugging\Unreal.natvis",
-			@"$(EngineDir)\Binaries\$(TargetPlatform)\DatasmithSDK\"
-		);
-
-		// Other headers we depend on, but that are not part of our public API:
-		PostBuildCopy(
-			@"$(EngineDir)\Source\Runtime\TraceLog\Public\*.h",
-			@"$(EngineDir)\Binaries\$(TargetPlatform)\DatasmithSDK\Private\"
-		);
-		PostBuildCopy(
-			@"$(EngineDir)\Source\Runtime\TraceLog\Public\*.inl",
-			@"$(EngineDir)\Binaries\$(TargetPlatform)\DatasmithSDK\Private\"
-		);
-
-		PostBuildCopy(
-			@"$(EngineDir)\Source\Runtime\Messaging\Public\*.h",
-			@"$(EngineDir)\Binaries\$(TargetPlatform)\DatasmithSDK\Private\"
-		);
-
-		PostBuildCopy(
-			@"$(EngineDir)\Source\Runtime\Core\Public\*.h",
-			@"$(EngineDir)\Binaries\$(TargetPlatform)\DatasmithSDK\Private\"
-		);
-
-		PostBuildCopy(
-			@"$(EngineDir)\Source\Runtime\Core\Public\*.inl",
-			@"$(EngineDir)\Binaries\$(TargetPlatform)\DatasmithSDK\Private\"
-		);
-
-		PostBuildCopy(
-			@"$(EngineDir)\Source\Runtime\CoreUObject\Public\*.h",
-			@"$(EngineDir)\Binaries\$(TargetPlatform)\DatasmithSDK\Private\"
-		);
+		PostCopyFilesFunc PostBuildCopy = PostBuildCopyWin64;
+
+		string SrcDatasmithSDK = @"$(EngineDir)\Source\Programs\Enterprise\Datasmith\DatasmithSDK";
+		string DstDatasmithSDK = @"$(EngineDir)\Binaries\$(TargetPlatform)\DatasmithSDK";
+
+		if (Platform != UnrealTargetPlatform.Win64)
+		{
+			List<string> Directories = new List<string>(
+				new string[]
+				{
+					@"$(EngineDir)/Binaries/$(TargetPlatform)/DatasmithSDK/Documentation/",
+					@"$(EngineDir)/Binaries/$(TargetPlatform)/DatasmithSDK/Public/",
+					@"$(EngineDir)/Binaries/$(TargetPlatform)/DatasmithSDK/Private/",
+					@"$(EngineDir)/Binaries/$(TargetPlatform)/DatasmithSDK/Engine/Binary",
+					@"$(EngineDir)/Binaries/$(TargetPlatform)/DatasmithSDK/Engine/Content/",
+					@"$(EngineDir)/Binaries/$(TargetPlatform)/DatasmithSDK/Engine/Shaders/StandaloneRenderer/",
+				}
+			);
+
+			PostBuildSteps.Add(string.Format("echo Creating directories to export to\n"));
+			foreach (string Directory in Directories)
+			{
+				PostBuildSteps.Add(string.Format("mkdir -p {0}\n", Directory));
+			}
+
+			PostBuildCopy = PostBuildCopyUnix;
+
+			PostBuildCopyUnix(@"$(EngineDir)/Content/Slate", Path.Combine(DstDatasmithSDK, "Engine/Content/"));
+			PostBuildCopyUnix(@"$(EngineDir)/Shaders/StandaloneRenderer/OpenGL", Path.Combine(DstDatasmithSDK, "Engine/Shaders/StandaloneRenderer"));
+		}
+		else
+		{
+			PostBuildSteps.Add(string.Format("mkdir {0}\n", Path.Combine(DstDatasmithSDK, "Engine\\Content\\Binary")));
+			PostBuildCopyWin64(@"$(EngineDir)\Content\Slate\", Path.Combine(DstDatasmithSDK, "Engine\\Content\\"), "");
+			PostBuildCopyWin64(@"$(EngineDir)\Shaders\StandaloneRenderer\D3D\", Path.Combine(DstDatasmithSDK, "Engine\\Shaders\\StandaloneRenderer\\D3D\\"), "");
+		}
+
+		(string, string, string)[] FilesToCopy = new (string, string, string)[]
+		{
+			(Path.Combine(SrcDatasmithSDK,"Documentation\\"), Path.Combine(DstDatasmithSDK,"Documentation\\"), "*.*"),
+			(@"$(EngineDir)\Source\Runtime\Datasmith\DatasmithCore\Public\", Path.Combine(DstDatasmithSDK,"Public\\"), "*.h"),
+			(@"$(EngineDir)\Source\Runtime\Datasmith\DirectLink\Public\", Path.Combine(DstDatasmithSDK,"Public\\"), "*.h"),
+			(@"$(EngineDir)\Source\Developer\Datasmith\DatasmithExporter\Public\", Path.Combine(DstDatasmithSDK,"Public\\"), "*.h"),
+			(@"$(EngineDir)\Source\Developer\Datasmith\DatasmithExporterUI\Public\", Path.Combine(DstDatasmithSDK,"Public\\"), "*.h"),
+			(@"$(EngineDir)\Extras\VisualStudioDebugging\", DstDatasmithSDK, "Unreal.natvis"),
+			(@"$(EngineDir)\Source\Runtime\TraceLog\Public\", Path.Combine(DstDatasmithSDK,"Private\\"), "*.h"),
+			(@"$(EngineDir)\Source\Runtime\TraceLog\Public\", Path.Combine(DstDatasmithSDK,"Private\\"), "*.inl"),
+			(@"$(EngineDir)\Source\Runtime\Messaging\Public\", Path.Combine(DstDatasmithSDK,"Private\\"), "*.h"),
+			(@"$(EngineDir)\Source\Runtime\Core\Public\", Path.Combine(DstDatasmithSDK,"Private\\"), "*.h"),
+			(@"$(EngineDir)\Source\Runtime\Core\Public\", Path.Combine(DstDatasmithSDK,"Private\\"), "*.inl"),
+			(@"$(EngineDir)\Source\Runtime\CoreUObject\Public\", Path.Combine(DstDatasmithSDK,"Private\\"), "*.h"),
+		};
+
+		foreach ((string, string, string) ToCopy in FilesToCopy)
+		{
+			PostBuildCopy(ToCopy.Item1, ToCopy.Item2, ToCopy.Item3);
+		}
 	}
 }
