[Solidworks] Change naming of configurations in Variant Manager

diff --git a/Engine/Source/Programs/Enterprise/Datasmith/DatasmithSolidworksExporter/Private/AssemblyDocument.cs b/Engine/Source/Programs/Enterprise/Datasmith/DatasmithSolidworksExporter/Private/AssemblyDocument.cs
--- a/Engine/Source/Programs/Enterprise/Datasmith/DatasmithSolidworksExporter/Private/AssemblyDocument.cs
+++ b/Engine/Source/Programs/Enterprise/Datasmith/DatasmithSolidworksExporter/Private/AssemblyDocument.cs
@@ -166,7 +166,7 @@ namespace DatasmithSolidworks
 		}
 
 		// Export to datasmith scene extracted configurations data
-		public override void ExportToDatasmithScene(FConfigurationExporter ConfigurationExporter)
+		public override void ExportToDatasmithScene(FConfigurationExporter ConfigurationExporter, FVariantName ActiveVariantName)
 		{
 			SetExportStatus("Actors");
 
@@ -176,7 +176,7 @@ namespace DatasmithSolidworks
 
 			// Configurations combined tree should have single child(root component of each config is the same)
 			Debug.Assert(ConfigurationExporter.CombinedTree.Children.Count == 1);
-			ExportComponentRecursive(CurrentConfig.Name, ConfigurationExporter,
+			ExportComponentRecursive(ActiveVariantName, ConfigurationExporter,
 				ConfigurationExporter.CombinedTree.Children[0], null);
 			SyncState.DirtyComponents.Clear();
 
@@ -327,7 +327,7 @@ namespace DatasmithSolidworks
 			return FConvertedTransform.Identity();
 		}
 
-		private void ExportComponentRecursive(string ActiveConfigName,
+		private void ExportComponentRecursive(FVariantName ActiveConfigName,
 			FConfigurationExporter ConfigurationExporter,
 			FConfigurationTree.FComponentTreeNode InNode, FConfigurationTree.FComponentTreeNode InParent)
 		{
diff --git a/Engine/Source/Programs/Enterprise/Datasmith/DatasmithSolidworksExporter/Private/ConfigurationExporter.cs b/Engine/Source/Programs/Enterprise/Datasmith/DatasmithSolidworksExporter/Private/ConfigurationExporter.cs
--- a/Engine/Source/Programs/Enterprise/Datasmith/DatasmithSolidworksExporter/Private/ConfigurationExporter.cs
+++ b/Engine/Source/Programs/Enterprise/Datasmith/DatasmithSolidworksExporter/Private/ConfigurationExporter.cs
@@ -162,17 +162,6 @@ namespace DatasmithSolidworks
 			return Configuration;
 		}
 
-		// Add variant which inherits geometry from one of the 'real' Solidworks configurations
-		// This could be (linked)display state or exploded view
-		public FConfiguration AddSubConfiguration(FConfiguration MeshesConfiguration, string ConfigurationName)
-		{
-			Debug.Assert(!Configurations.ContainsKey(ConfigurationName));
-
-			FConfiguration Configuration = new FConfiguration(this, ConfigurationName, MainConfigurationName==ConfigurationName, InParentConfiguration: MeshesConfiguration);
-			Configurations.Add(ConfigurationName, Configuration);
-			return Configuration;
-		}
-
 		public bool DoesComponentHaveVisibleButDifferentMeshesInConfigs(FComponentName ComponentName, string ConfigNameA, string ConfigNameB)
 		{
 			bool bHasMeshA = GetMeshesConfiguration(ConfigNameA).GetMeshForComponent(ComponentName, out FMesh MeshA);
@@ -351,7 +340,7 @@ namespace DatasmithSolidworks
 
 			string OriginalConfigurationName = ConfigManager.ActiveConfiguration.Name;
 
-			Dictionary<string, List<string>> ExportedConfigurationNames = new Dictionary<string, List<string>>();
+			Dictionary<string, List<FVariantName>> ExportedVariantNames = new Dictionary<string, List<FVariantName>>();
 
 			// Enumerate configurations starting with Active - to avoid redundant config switching(which takes time) and,
 			// additionally, to export meshes in Active configuration without configuration name suffixes
@@ -377,7 +366,12 @@ namespace DatasmithSolidworks
 
 				List<string> ExplodedViews = InitExplodedViews(InDoc, CfgName);
 
-				string ConfigName = CfgName;
+				// Variant name to base on for linked display states and exploded views
+				FVariantName BaseCfgVariantName = new FVariantName(swConfiguration, CfgName);  
+
+				// Name identifying Datasmith Variant to be exported
+				// variants sets consists of SW Configurations, derived configurations, display states, exploded views
+				FVariantName VariantName = BaseCfgVariantName;
 
 				string[] DisplayStates = null;
 				int DisplayStateCount = 0;
@@ -386,15 +380,16 @@ namespace DatasmithSolidworks
 				{
 					DisplayStateCount = swConfiguration.GetDisplayStatesCount();
 
+					// When display states are linked to configuration modify variant name to include current display state name
 					if (ConfigManager.LinkDisplayStatesToConfigurations && DisplayStateCount > 1)
 					{
 						DisplayStates = swConfiguration.GetDisplayStates();
-						ConfigName = $"{CfgName}_{DisplayStates[0]}";
+						VariantName = VariantName.LinkedDisplayStateVariant(DisplayStates[0]);
 					}
 				}
 
 				FConfigurationTree.FComponentTreeNode ConfigNode = new FConfigurationTree.FComponentTreeNode();
-				ConfigNode.ComponentInfo.ComponentName = FComponentName.FromCustomString(ConfigName);
+				ConfigNode.ComponentInfo.ComponentName = VariantName.GetRootComponentName();
 
 				ConfigNode.Children = new List<FConfigurationTree.FComponentTreeNode>();
 
@@ -406,11 +401,11 @@ namespace DatasmithSolidworks
 				LogDebug($"Components:");
 				CollectComponentsRecursive(InDoc, swConfiguration.GetRootComponent3(true), ConfigNode, MeshesConfiguration);
 
-				ExportedConfigurationNames.Add(CfgName, new List<string>(){ConfigName});
+				ExportedVariantNames.Add(CfgName, new List<FVariantName>(){VariantName});
 
 				LogDebug($"Materials:");
 				Dictionary<FComponentName, FObjectMaterials> MaterialsMap = GetComponentMaterials(InDoc, DisplayStates?[0], swConfiguration);
-				SetComponentTreeMaterials(ConfigNode, MaterialsMap, null, false);
+				SetComponentTreeMaterials(ConfigNode, MaterialsMap, new FVariantName(), false);
 
 				if (DisplayStates != null)
 				{
@@ -419,13 +414,14 @@ namespace DatasmithSolidworks
 					for (int Index = 1; Index < DisplayStateCount; ++Index)
 					{
 						string DisplayState = DisplayStates[Index];
-						string DisplayStateTreeName = $"{CfgName}_DisplayState_{FDatasmithExporter.SanitizeName(DisplayState)}";
-						LogDebug($"Linked DisplayState '{DisplayState}' Materials for variant '{DisplayStateTreeName}'");
 
-						Meshes.AddSubConfiguration(MeshesConfiguration, ConfigurationName: DisplayStateTreeName);
-						ExportedConfigurationNames[CfgName].Add(DisplayStateTreeName);
+						FVariantName DisplayStateVariantName = BaseCfgVariantName.LinkedDisplayStateVariant(DisplayState);
+
+						LogDebug($"Linked DisplayState '{DisplayState}' Materials for variant '{DisplayStateVariantName}'");
+
+						ExportedVariantNames[CfgName].Add(DisplayStateVariantName);
 						MaterialsMap = GetComponentMaterials(InDoc, DisplayState, swConfiguration);
-						SetComponentTreeMaterials(ConfigNode, MaterialsMap, DisplayStateTreeName, false);
+						SetComponentTreeMaterials(ConfigNode, MaterialsMap, DisplayStateVariantName, false);
 					}
 				}
 
@@ -458,9 +454,8 @@ namespace DatasmithSolidworks
 							}
 						}
 
-						string ExplodedViewVariantName = $"{CfgName}_Exploded_{FDatasmithExporter.SanitizeName(ExplodedViewName)}";
-						Meshes.AddSubConfiguration(MeshesConfiguration, ConfigurationName: ExplodedViewVariantName);
-						ExportedConfigurationNames[CfgName].Add(ExplodedViewVariantName);
+						FVariantName ExplodedViewVariantName = BaseCfgVariantName.ExplodedViewVariant(ExplodedViewName);
+						ExportedVariantNames[CfgName].Add(ExplodedViewVariantName);
 
 						// Add variant to each node with transform initialized from exploded state
 						foreach (FComponentTreeNode Child in ConfigNode.Children)
@@ -492,7 +487,7 @@ namespace DatasmithSolidworks
 
 				if (ModifiedComponentsMaterials != null)
 				{
-					foreach (var MatKVP in ModifiedComponentsMaterials)
+					foreach (KeyValuePair<FComponentName, FObjectMaterials> MatKVP in ModifiedComponentsMaterials)
 					{
 						InDoc.AddComponentMaterials(MatKVP.Key, MatKVP.Value);
 					}
@@ -504,7 +499,7 @@ namespace DatasmithSolidworks
 				InDoc.ProcessConfigurationMeshes(ExtractedMeshes, MeshesConfiguration);
 
 				// Combine separate scene trees into the single one with configuration-specific data
-				FConfigurationTree.Merge(CombinedTree, ConfigNode, CfgName, bIsActiveConfiguration);
+				FConfigurationTree.Merge(CombinedTree, ConfigNode, VariantName, bIsActiveConfiguration);
 
 				bIsActiveConfiguration = false;  // Next enumerated configurations will be switched to(they weren't active on export start)
 			}
@@ -541,8 +536,8 @@ namespace DatasmithSolidworks
 				}
 			}
 
-			List<string> DisplayStateConfigurations = new List<string>();
-
+			List<FVariantName> DisplayStateConfigurations = new List<FVariantName>();
+			
 			if (bExportDisplayStates && !ConfigManager.LinkDisplayStatesToConfigurations)
 			{
 				// Export display states as separate configurations
@@ -555,8 +550,8 @@ namespace DatasmithSolidworks
 					{
 						OriginalConfiguration.ApplyDisplayState(DisplayState);
 					}
-
-					string DisplayStateConfigName = $"DisplayState_{FDatasmithExporter.SanitizeName(DisplayState)}";
+					
+					FVariantName DisplayStateConfigName = FVariantName.DisplayStateVariant(DisplayState);
 					DisplayStateConfigurations.Add(DisplayStateConfigName);
 					LogDebug(
 						$"DisplayState '{DisplayState}' Materials for variant '{DisplayStateConfigName}'");
@@ -583,25 +578,25 @@ namespace DatasmithSolidworks
 			List<FConfigurationData> FlatConfigurationData = new List<FConfigurationData>();
 
 			// Order exported variants by original configurations order
-			IEnumerable<string> ExportedConfigurationNamesOrdered = ExportedConfigurationNames.OrderBy(V => ConfigurationOrder[V.Key]).SelectMany(V => V.Value);
+			IEnumerable<FVariantName> ExportedConfigurationNamesOrdered = ExportedVariantNames.OrderBy(V => ConfigurationOrder[V.Key]).SelectMany(V => V.Value);
 
 			// Add roots for configurations
-			foreach (string CfgName in ExportedConfigurationNamesOrdered)
+			foreach (FVariantName ConfigName in ExportedConfigurationNamesOrdered)
 			{
 				FConfigurationData CfgData = new FConfigurationData();
-				CfgData.Name = CfgName;
+				CfgData.Name = ConfigName.ToString();
 				CfgData.bIsDisplayStateConfiguration = false;
-				FConfigurationTree.FillConfigurationData(this, CombinedTree, CfgName, CfgData, false);
+				FConfigurationTree.FillConfigurationData(this, CombinedTree, ConfigName, CfgData, false);
 
 				if (!CfgData.IsEmpty())
 				{
 					FlatConfigurationData.Add(CfgData);
 				}
 			}
-			foreach (string DisplayStateConfigName in DisplayStateConfigurations)
+			foreach (FVariantName DisplayStateConfigName in DisplayStateConfigurations)
 			{
 				FConfigurationData CfgData = new FConfigurationData();
-				CfgData.Name = DisplayStateConfigName;
+				CfgData.Name = DisplayStateConfigName.ToString();
 				CfgData.bIsDisplayStateConfiguration = true;
 				FConfigurationTree.FillConfigurationData(this, CombinedTree, DisplayStateConfigName, CfgData, true);
 
@@ -771,13 +766,13 @@ namespace DatasmithSolidworks
 			return MaterialsMap;
 		}
 
-		private static void SetComponentTreeMaterials(FConfigurationTree.FComponentTreeNode InComponentTree, Dictionary<FComponentName, FObjectMaterials> InComponentMaterialsMap, string InConfigurationName, bool bIsDisplayState)
+		private static void SetComponentTreeMaterials(FConfigurationTree.FComponentTreeNode InComponentTree, Dictionary<FComponentName, FObjectMaterials> InComponentMaterialsMap, FVariantName InConfigurationName, bool bIsDisplayState)
 		{
 			LogDebug($"SetComponentTreeMaterials: '{InComponentTree.ComponentName}'");
 
 			FConfigurationTree.FComponentConfig TargetConfig = null;
 
-			if (InConfigurationName != null)
+			if (InConfigurationName.IsValid())
 			{
 				TargetConfig = InComponentTree.GetConfiguration(InConfigurationName, bIsDisplayState);
 				if (TargetConfig == null)
@@ -888,32 +883,32 @@ namespace DatasmithSolidworks
 			}
 		}
 
-		public bool DoesComponentHaveVisibleButDifferentMeshesInConfigs(FComponentName ComponentName, string ConfigNameA, string ConfigNameB)
+		public bool DoesComponentHaveVisibleButDifferentMeshesInConfigs(FComponentName ComponentName, FVariantName ConfigNameA, FVariantName ConfigNameB)
 		{
-			return Meshes.DoesComponentHaveVisibleButDifferentMeshesInConfigs(ComponentName, ConfigNameA, ConfigNameB);
+			return Meshes.DoesComponentHaveVisibleButDifferentMeshesInConfigs(ComponentName, ConfigNameA.CfgName, ConfigNameB.CfgName);
 		}
 
 		// name/label for variant mesh actor 
-		public FActorName GetMeshActorName(string ConfigName, FComponentName ComponentName)
+		public FActorName GetMeshActorName(FVariantName ConfigName, FComponentName ComponentName)
 		{
-			return FActorName.FromString($"{ComponentName}_{ConfigName}");
+			return FActorName.FromString($"{ComponentName}_{ConfigName.CfgName}");
 		}
 
-		public FMeshName GetMeshName(string ConfigName, FComponentName ComponentName)
+		public FMeshName GetMeshName(FVariantName ConfigName, FComponentName ComponentName)
 		{
-			return Meshes.GetMeshName(ConfigName, ComponentName);
+			return Meshes.GetMeshName(ConfigName.CfgName, ComponentName);
 		}
 
 		// Get mesh for component which has only single mesh configuration
 		public FMeshName GetMeshName(FComponentName ComponentName)
 		{
 			string ConfigName = Meshes.GetConfigurationNameForComponentMesh(ComponentName);
-			return GetMeshName(ConfigName, ComponentName);
+			return Meshes.GetMeshName(ConfigName, ComponentName);
 		}
 
 		// Register actor which needs mesh of a component in specific config
 		// this is needed for mesh actors for variants
-		public void AddActorForMesh(FActorName ActorName, string ConfigName, FComponentName ComponentName)
+		public void AddActorForMesh(FActorName ActorName, FVariantName ConfigName, FComponentName ComponentName)
 		{
 			LogDebug($"AddActorForMesh(ActorName='{ActorName}', ConfigName='{ConfigName}', ComponentName='{ComponentName}')");
 			AddActorForMesh(ActorName, ComponentName, GetMeshName(ConfigName, ComponentName));
diff --git a/Engine/Source/Programs/Enterprise/Datasmith/DatasmithSolidworksExporter/Private/ConfigurationTree.cs b/Engine/Source/Programs/Enterprise/Datasmith/DatasmithSolidworksExporter/Private/ConfigurationTree.cs
--- a/Engine/Source/Programs/Enterprise/Datasmith/DatasmithSolidworksExporter/Private/ConfigurationTree.cs
+++ b/Engine/Source/Programs/Enterprise/Datasmith/DatasmithSolidworksExporter/Private/ConfigurationTree.cs
@@ -47,7 +47,7 @@ namespace DatasmithSolidworks
 		// Component state for a single configuration
 		public class FComponentConfig
 		{
-			public readonly string ConfigName;
+			public readonly FVariantName ConfigName;
 			public readonly bool bIsMainActive;  // if this Configuration is the active config
 
 			public bool bVisible = true;
@@ -62,7 +62,7 @@ namespace DatasmithSolidworks
 
 			public List<Body2> Bodies = null;
 
-			public FComponentConfig(string InConfigName, bool bInIsMainActive)
+			public FComponentConfig(FVariantName InConfigName, bool bInIsMainActive)
 			{
 				ConfigName = InConfigName;
 				bIsMainActive = bInIsMainActive;
@@ -101,7 +101,7 @@ namespace DatasmithSolidworks
 			public string PartPath => ComponentInfo.PartPath;
 
 			// Common configuration data
-			public FComponentConfig CommonConfig = new FComponentConfig(null, true);
+			public FComponentConfig CommonConfig = new FComponentConfig(new FVariantName(), true);
 
 			// Per-configuration data
 			public List<FComponentConfig> Configurations = null;
@@ -156,10 +156,12 @@ namespace DatasmithSolidworks
 				Component = InComponent;
 			}
 
-			public FComponentConfig AddConfiguration(string InConfigurationName, bool bIsDisplayState, bool bIsActiveConfiguration)
+			public FComponentConfig AddConfiguration(FVariantName InConfigurationName, bool bIsDisplayState, bool bIsActiveConfiguration)
 			{
 				List<FComponentConfig> TargetList = null;
 
+				Debug.Assert(InConfigurationName.IsValid());
+
 				if (bIsDisplayState)
 				{
 					if (DisplayStates == null)
@@ -182,7 +184,7 @@ namespace DatasmithSolidworks
 				return Result;
 			}
 
-			public FComponentConfig GetConfiguration(string InConfigurationName, bool bIsDisplayState)
+			public FComponentConfig GetConfiguration(FVariantName InConfigurationName, bool bIsDisplayState)
 			{
 				List<FComponentConfig> TargetList = null;
 
@@ -223,7 +225,7 @@ namespace DatasmithSolidworks
 			}
 		};
 
-		static public void Merge(FComponentTreeNode OutCombined, FComponentTreeNode InTree, string InConfigurationName,
+		static public void Merge(FComponentTreeNode OutCombined, FComponentTreeNode InTree, FVariantName InConfigurationName,
 			bool bIsActiveConfiguration)
 		{
 			foreach (FComponentTreeNode Child in InTree.EnumChildren())
@@ -434,7 +436,7 @@ namespace DatasmithSolidworks
 			}
 		}
 
-		static public void FillConfigurationData(FConfigurationExporter ConfigurationExporter, FComponentTreeNode InNode, string InConfigurationName, FConfigurationData OutData, bool bIsDisplayState)
+		static public void FillConfigurationData(FConfigurationExporter ConfigurationExporter, FComponentTreeNode InNode, FVariantName InConfigurationName, FConfigurationData OutData, bool bIsDisplayState)
 		{
 			FComponentConfig NodeConfig = InNode.GetConfiguration(InConfigurationName, bIsDisplayState);
 
diff --git a/Engine/Source/Programs/Enterprise/Datasmith/DatasmithSolidworksExporter/Private/DatasmithExporter.cs b/Engine/Source/Programs/Enterprise/Datasmith/DatasmithSolidworksExporter/Private/DatasmithExporter.cs
--- a/Engine/Source/Programs/Enterprise/Datasmith/DatasmithSolidworksExporter/Private/DatasmithExporter.cs
+++ b/Engine/Source/Programs/Enterprise/Datasmith/DatasmithSolidworksExporter/Private/DatasmithExporter.cs
@@ -125,6 +125,80 @@ namespace DatasmithSolidworks
 			public override int GetHashCode() => Value.GetHashCode();
 		}
 
+		public struct FVariantName : IEquatable<FVariantName>
+		{
+			private readonly string Value;
+
+			// Solidworks configuration name this variant is based on
+			public readonly string CfgName; 
+
+			private FVariantName(string InCfgName, string InValue)
+			{
+				CfgName =  InCfgName;
+				Value = InValue;
+			}
+
+			public FVariantName(IConfiguration Configuration, string InCfgName)
+			{
+				Debug.Assert(Configuration.Name == InCfgName);
+
+				CfgName = InCfgName;
+
+				Value = Configuration.IsDerived() 
+					? $"{Configuration.GetParent().Name}_{CfgName}" 
+					: Configuration.Name;
+			}
+
+			public FVariantName(IConfiguration Configuration): this(Configuration, Configuration.Name)
+			{
+			}
+
+			public static FVariantName PartVariant(string CfgName)
+			{
+				return new FVariantName(CfgName, CfgName);
+			}
+
+			public FVariantName LinkedDisplayStateVariant(string DisplayStateName)
+			{
+				return new FVariantName(CfgName, $"{Value}_{FDatasmithExporter.SanitizeName(DisplayStateName)}");
+			}
+
+			public static FVariantName DisplayStateVariant(string DisplayStateName)
+			{
+				return new FVariantName((string)null, $"DisplayState_{FDatasmithExporter.SanitizeName(DisplayStateName)}");
+			}
+
+			public FVariantName ExplodedViewVariant(string ExplodedViewName)
+			{
+				return new FVariantName(CfgName, $"{Value}_{FDatasmithExporter.SanitizeName(ExplodedViewName)}");
+			}
+
+			public static FVariantName Invalid() => new FVariantName((string)null, null);
+
+			public FComponentName GetRootComponentName()
+			{
+				return FComponentName.FromCustomString(Value);
+			}
+
+			// Following is common between 'Names' but unfortunately C# has no way to generalize this for structs
+			public bool IsValid()
+			{
+				return !string.IsNullOrEmpty(Value);
+			}
+
+			public string GetString() => Value;
+
+			public override string ToString() => Value;
+
+			public override bool Equals(object Obj) => Obj is FVariantName Other && this == Other;
+			public bool Equals(FVariantName Other) => Value == Other.Value;
+			public static bool operator ==(FVariantName A, FVariantName B) => A.Equals(B);
+			public static bool operator !=(FVariantName A, FVariantName B) => !(A == B);
+
+			public override int GetHashCode() => Value.GetHashCode();
+
+		}
+
 	}
 }
 
diff --git a/Engine/Source/Programs/Enterprise/Datasmith/DatasmithSolidworksExporter/Private/Document.cs b/Engine/Source/Programs/Enterprise/Datasmith/DatasmithSolidworksExporter/Private/Document.cs
--- a/Engine/Source/Programs/Enterprise/Datasmith/DatasmithSolidworksExporter/Private/Document.cs
+++ b/Engine/Source/Programs/Enterprise/Datasmith/DatasmithSolidworksExporter/Private/Document.cs
@@ -163,7 +163,7 @@ namespace DatasmithSolidworks
 
 			bHasConfigurations = (Configs != null) && (Configs.Count != 0);
 
-			ExportToDatasmithScene(ConfigurationExporter);
+			ExportToDatasmithScene(ConfigurationExporter, new FVariantName(ConfigManager.ActiveConfiguration));
 
 			ExportLights();
 			ExportConfigurations(Configs);
@@ -193,7 +193,7 @@ namespace DatasmithSolidworks
 			FConfigurationExporter ConfigurationExporter = new FConfigurationExporter(Meshes, new []{ ActiveConfigurationName }, ActiveConfigurationName, bInExportDisplayStates: false, bInExportExplodedViews: false);
 			List<FConfigurationData> Configs = ConfigurationExporter.ExportConfigurations(this);
 			bHasConfigurations = (Configs != null) && (Configs.Count != 0);
-			ExportToDatasmithScene(ConfigurationExporter);
+			ExportToDatasmithScene(ConfigurationExporter, new FVariantName(ConfigManager.ActiveConfiguration));
 
 
 			ExportLights();
@@ -206,7 +206,7 @@ namespace DatasmithSolidworks
 		}
 
 		// todo: make abstract
-		public virtual void ExportToDatasmithScene(FConfigurationExporter ConfigurationExporter)
+		public virtual void ExportToDatasmithScene(FConfigurationExporter ConfigurationExporter, FVariantName ActiveVariantName)
 		{
 			throw new NotImplementedException();
 		}
diff --git a/Engine/Source/Programs/Enterprise/Datasmith/DatasmithSolidworksExporter/Private/PartDocument.cs b/Engine/Source/Programs/Enterprise/Datasmith/DatasmithSolidworksExporter/Private/PartDocument.cs
--- a/Engine/Source/Programs/Enterprise/Datasmith/DatasmithSolidworksExporter/Private/PartDocument.cs
+++ b/Engine/Source/Programs/Enterprise/Datasmith/DatasmithSolidworksExporter/Private/PartDocument.cs
@@ -10,7 +10,7 @@ using DatasmithSolidworks.Names;
 using static DatasmithSolidworks.FDatasmithExporter;
 using System.Diagnostics;
 using System.Linq;
-using static DatasmithSolidworks.FAssemblyDocumentTracker;
+
 
 namespace DatasmithSolidworks
 {
@@ -231,7 +231,7 @@ namespace DatasmithSolidworks
 			return Meshes ?? (Meshes = new FMeshes(ActiveConfigName));
 		}
 
-		public override void ExportToDatasmithScene(FConfigurationExporter ConfigurationExporter)
+		public override void ExportToDatasmithScene(FConfigurationExporter ConfigurationExporter, FVariantName ActiveVariantName)
 		{
 			FMeshes Meshes = ConfigurationExporter.Meshes;
 
@@ -258,7 +258,7 @@ namespace DatasmithSolidworks
 
 					// todo: not make extra copy - use existing(i.e. instead of bHasConfigurations) 
 					FActorName ActorName =
-						ConfigurationExporter.GetMeshActorName(ConfigurationName, RootComponentName);
+						ConfigurationExporter.GetMeshActorName(FVariantName.PartVariant(ConfigurationName), RootComponentName);
 
 					FMeshName MeshNameForConfiguration = new FMeshName(RootComponentName, ConfigurationName);
 
