[Naviswork] Update exporter for Naviswork 2024

diff --git a/Engine/Source/Programs/Enterprise/Datasmith/DatasmithNavisworksExporter/DatasmithNavisworksPlugin/DatasmithNavisworksPlugin.csproj b/Engine/Source/Programs/Enterprise/Datasmith/DatasmithNavisworksExporter/DatasmithNavisworksPlugin/DatasmithNavisworksPlugin.csproj
--- a/Engine/Source/Programs/Enterprise/Datasmith/DatasmithNavisworksExporter/DatasmithNavisworksPlugin/DatasmithNavisworksPlugin.csproj
+++ b/Engine/Source/Programs/Enterprise/Datasmith/DatasmithNavisworksExporter/DatasmithNavisworksPlugin/DatasmithNavisworksPlugin.csproj
@@ -8,6 +8,8 @@
     <ProjectGuid>{D16F2120-CCE4-4FCC-92C8-86124AB1FA25}</ProjectGuid>
     <TargetFrameworkProfile />
   </PropertyGroup>
+  <Import Project="..\Navisworks2024.props" Condition="'$(Configuration)|$(Platform)'=='Debug2024|AnyCPU'" />
+  <Import Project="..\Navisworks2024.props" Condition="'$(Configuration)|$(Platform)'=='Release2024|AnyCPU'" />
   <Import Project="..\Navisworks2023.props" Condition="'$(Configuration)|$(Platform)'=='Debug2023|AnyCPU'" />
   <Import Project="..\Navisworks2023.props" Condition="'$(Configuration)|$(Platform)'=='Release2023|AnyCPU'" />
   <Import Project="..\Navisworks2022.props" Condition="'$(Configuration)|$(Platform)'=='Debug2022|AnyCPU'" />
@@ -28,6 +30,30 @@
     <TargetFrameworkVersion>v4.8</TargetFrameworkVersion>
     <FileAlignment>512</FileAlignment>
   </PropertyGroup>
+  <PropertyGroup Condition=" '$(Configuration)|$(Platform)' == 'Debug2024|AnyCPU' ">
+    <DebugSymbols>true</DebugSymbols>
+    <DebugType>full</DebugType>
+    <Optimize>false</Optimize>
+    <OutputPath>bin\Debug2024\</OutputPath>
+    <DefineConstants>DEBUG;TRACE</DefineConstants>
+    <ErrorReport>prompt</ErrorReport>
+    <WarningLevel>4</WarningLevel>
+    <StartAction>Program</StartAction>
+    <StartProgram>$(ProgramW6432)\Autodesk\Navisworks Manage $(Navisworks_Version)\Roamer.exe</StartProgram>
+    <EnableUnmanagedDebugging>true</EnableUnmanagedDebugging>
+    <AllowUnsafeBlocks>true</AllowUnsafeBlocks>
+  </PropertyGroup>
+  <PropertyGroup Condition=" '$(Configuration)|$(Platform)' == 'Release2024|AnyCPU' ">
+    <DebugType>pdbonly</DebugType>
+    <Optimize>true</Optimize>
+    <OutputPath>bin\Release2024\</OutputPath>
+    <DefineConstants>TRACE</DefineConstants>
+    <ErrorReport>prompt</ErrorReport>
+    <WarningLevel>4</WarningLevel>
+    <StartAction>Program</StartAction>
+    <StartProgram>$(ProgramW6432)\Autodesk\Navisworks Manage $(Navisworks_Version)\Roamer.exe</StartProgram>
+    <AllowUnsafeBlocks>true</AllowUnsafeBlocks>
+  </PropertyGroup>
   <PropertyGroup Condition=" '$(Configuration)|$(Platform)' == 'Debug2023|AnyCPU' ">
     <DebugSymbols>true</DebugSymbols>
     <DebugType>full</DebugType>
diff --git a/Engine/Source/Programs/Enterprise/Datasmith/DatasmithNavisworksExporter/DatasmithNavisworksPlugin/PackageContents.xml b/Engine/Source/Programs/Enterprise/Datasmith/DatasmithNavisworksExporter/DatasmithNavisworksPlugin/PackageContents.xml
--- a/Engine/Source/Programs/Enterprise/Datasmith/DatasmithNavisworksExporter/DatasmithNavisworksPlugin/PackageContents.xml
+++ b/Engine/Source/Programs/Enterprise/Datasmith/DatasmithNavisworksExporter/DatasmithNavisworksPlugin/PackageContents.xml
@@ -21,4 +21,8 @@
 	    <RuntimeRequirements OS="Win64" Platform="NAVMAN|NAVSIM" SeriesMin="Nw20" SeriesMax="Nw20" />
         <ComponentEntry AppName="Datasmith Exporter for Autodesk® Navisworks® 2023" AppType="ManagedPlugin" ModuleName="./DatasmithNavisworksPlugin2023.dll" />
     </Components>
+    <Components Description="2024">
+	    <RuntimeRequirements OS="Win64" Platform="NAVMAN|NAVSIM" SeriesMin="Nw21" SeriesMax="Nw21" />
+        <ComponentEntry AppName="Datasmith Exporter for Autodesk® Navisworks® 2024" AppType="ManagedPlugin" ModuleName="./DatasmithNavisworksPlugin2024.dll" />
+    </Components>
 </ApplicationPackage>
diff --git a/Engine/Source/Programs/Enterprise/Datasmith/DatasmithNavisworksExporter/UBTModules/DatasmithNavisworks2024.Build.cs b/Engine/Source/Programs/Enterprise/Datasmith/DatasmithNavisworksExporter/UBTModules/DatasmithNavisworks2024.Build.cs
new file mode 100644
--- /dev/null
+++ b/Engine/Source/Programs/Enterprise/Datasmith/DatasmithNavisworksExporter/UBTModules/DatasmithNavisworks2024.Build.cs
@@ -0,0 +1,20 @@
+// Copyright Epic Games, Inc. All Rights Reserved.
+
+using System.IO;
+
+namespace UnrealBuildTool.Rules
+{
+	[SupportedPlatforms("Win64")]
+	public class DatasmithNavisworks2024 : DatasmithNavisworksBase
+	{
+		public DatasmithNavisworks2024(ReadOnlyTargetRules Target)
+			: base(Target)
+		{
+		}
+
+		public override string GetNavisworksVersion()
+		{
+			return "2024";
+		}
+	}
+}
diff --git a/Engine/Source/Programs/Enterprise/Datasmith/DatasmithNavisworksExporter/UBTModules/DatasmithNavisworks2024.Target.cs b/Engine/Source/Programs/Enterprise/Datasmith/DatasmithNavisworksExporter/UBTModules/DatasmithNavisworks2024.Target.cs
new file mode 100644
--- /dev/null
+++ b/Engine/Source/Programs/Enterprise/Datasmith/DatasmithNavisworksExporter/UBTModules/DatasmithNavisworks2024.Target.cs
@@ -0,0 +1,14 @@
+// Copyright Epic Games, Inc. All Rights Reserved.
+
+using UnrealBuildTool;
+using System.IO;
+
+public class DatasmithNavisworks2024Target : DatasmithNavisworksBaseTarget
+{
+	public DatasmithNavisworks2024Target(TargetInfo Target)
+		: base(Target)
+	{
+	}
+
+	public override string GetVersion() { return "2024"; }
+}
