[DatasmithExporter] - Fix inconsistent supported platform tag.

diff --git a/Engine/Source/Programs/Enterprise/Datasmith/DatasmithRhinoExporter/DatasmithRhino7.Build.cs b/Engine/Source/Programs/Enterprise/Datasmith/DatasmithRhinoExporter/DatasmithRhino7.Build.cs
--- a/Engine/Source/Programs/Enterprise/Datasmith/DatasmithRhinoExporter/DatasmithRhino7.Build.cs
+++ b/Engine/Source/Programs/Enterprise/Datasmith/DatasmithRhinoExporter/DatasmithRhino7.Build.cs
@@ -4,7 +4,7 @@ using System.IO;
 
 namespace UnrealBuildTool.Rules
 {
-	[SupportedPlatforms("Win64")]
+	[SupportedPlatforms("Win64", "Mac")]
 	public class DatasmithRhino7 : DatasmithRhinoBase
 	{
 		public DatasmithRhino7(ReadOnlyTargetRules Target)
