[3dsMax] 3ds Max 2024 support

diff --git a/Engine/Source/Programs/Enterprise/Datasmith/DatasmithMaxExporter/DatasmithMax2024.Build.cs b/Engine/Source/Programs/Enterprise/Datasmith/DatasmithMaxExporter/DatasmithMax2024.Build.cs
new file mode 100644
--- /dev/null
+++ b/Engine/Source/Programs/Enterprise/Datasmith/DatasmithMaxExporter/DatasmithMax2024.Build.cs
@@ -0,0 +1,16 @@
+// Copyright Epic Games, Inc. All Rights Reserved.
+
+using System.Diagnostics;
+
+namespace UnrealBuildTool.Rules
+{
+	[SupportedPlatforms("Win64")]
+	public class DatasmithMax2024 : DatasmithMaxBase
+	{
+		public DatasmithMax2024(ReadOnlyTargetRules Target)
+			: base(Target)
+		{
+		}
+		public override string GetMaxVersion() { return "2024"; }
+	}
+}
\ No newline at end of file
diff --git a/Engine/Source/Programs/Enterprise/Datasmith/DatasmithMaxExporter/DatasmithMax2024.Target.cs b/Engine/Source/Programs/Enterprise/Datasmith/DatasmithMaxExporter/DatasmithMax2024.Target.cs
new file mode 100644
--- /dev/null
+++ b/Engine/Source/Programs/Enterprise/Datasmith/DatasmithMaxExporter/DatasmithMax2024.Target.cs
@@ -0,0 +1,15 @@
+// Copyright Epic Games, Inc. All Rights Reserved.
+
+using UnrealBuildTool;
+using System.Collections.Generic;
+
+public class DatasmithMax2024Target : DatasmithMaxBaseTarget
+{
+	public DatasmithMax2024Target(TargetInfo Target) : base(Target)
+	{
+		LaunchModuleName = "DatasmithMax2024";
+		ExeBinariesSubFolder = @"3DSMax\2024";
+
+		AddCopyPostBuildStep(Target);
+	}
+}
diff --git a/Engine/Source/Programs/Enterprise/Datasmith/DatasmithMaxExporter/Plugin/PackageContents.xml b/Engine/Source/Programs/Enterprise/Datasmith/DatasmithMaxExporter/Plugin/PackageContents.xml
--- a/Engine/Source/Programs/Enterprise/Datasmith/DatasmithMaxExporter/Plugin/PackageContents.xml
+++ b/Engine/Source/Programs/Enterprise/Datasmith/DatasmithMaxExporter/Plugin/PackageContents.xml
@@ -16,7 +16,12 @@
 	<!--  todo: Icon="./DatasmithIcon32.png" -->
     <CompanyDetails Name="EpicGames, Inc." Url="www.epicgames.com"/>
 
-    <RuntimeRequirements OS="Win64" Platform="3ds Max" SeriesMin="2017" SeriesMax="2023" />
+    <RuntimeRequirements OS="Win64" Platform="3ds Max" SeriesMin="2017" SeriesMax="2024" />
+
+    <Components Description="plugins parts">
+        <RuntimeRequirements OS="Win64" Platform="3ds Max" SeriesMin="2024" SeriesMax="2024" />
+        <ComponentEntry ModuleName="./2024/DatasmithMax2024.gup" />
+    </Components>
 
     <Components Description="plugins parts">
         <RuntimeRequirements OS="Win64" Platform="3ds Max" SeriesMin="2023" SeriesMax="2023" />
@@ -54,18 +59,18 @@
     </Components>
 
     <Components Description="pre-start-up scripts parts">
-        <RuntimeRequirements OS="Win64" Platform="3ds Max" SeriesMin="2017" SeriesMax="2023" />
+        <RuntimeRequirements OS="Win64" Platform="3ds Max" SeriesMin="2017" SeriesMax="2024" />
         <ComponentEntry ModuleName="./UI/DatasmithPreStartUp.ms" />
         <ComponentEntry ModuleName="./UI/UnrealStudio_Datasmith_Attributes_Holder_Modifier.ms" />
     </Components>
 
     <Components Description="macroscripts parts">
-        <RuntimeRequirements OS="Win64" Platform="3ds Max" SeriesMin="2017" SeriesMax="2023" />
+        <RuntimeRequirements OS="Win64" Platform="3ds Max" SeriesMin="2017" SeriesMax="2024" />
         <ComponentEntry ModuleName="./UI/DatasmithMacros.mcr" />
     </Components>
 
     <Components Description="post-start-up scripts parts">
-        <RuntimeRequirements OS="Win64" Platform="3ds Max" SeriesMin="2017" SeriesMax="2023" />
+        <RuntimeRequirements OS="Win64" Platform="3ds Max" SeriesMin="2017" SeriesMax="2024" />
         <ComponentEntry ModuleName="./UI/DatasmithPostStartUp.ms" />
     </Components>
 
