Add support to geolocation to Datasmith API

diff --git a/Engine/Plugins/Enterprise/DatasmithImporter/Source/DatasmithImporter/Private/Utility/DatasmithImporterUtils.cpp b/Engine/Plugins/Enterprise/DatasmithImporter/Source/DatasmithImporter/Private/Utility/DatasmithImporterUtils.cpp
--- a/Engine/Plugins/Enterprise/DatasmithImporter/Source/DatasmithImporter/Private/Utility/DatasmithImporterUtils.cpp
+++ b/Engine/Plugins/Enterprise/DatasmithImporter/Source/DatasmithImporter/Private/Utility/DatasmithImporterUtils.cpp
@@ -197,6 +197,11 @@ ADatasmithSceneActor* FDatasmithImporterUtils::CreateImportSceneActor( FDatasmit
 		RootComponent->RegisterComponent();
 	}
 
+	FVector Geolocation = ImportContext.Scene->GetGeolocation();
+	UDatasmithAssetUserData::SetDatasmithUserDataValueForKey(SceneActor, TEXT("Geolocation_Latitude"),  FString::SanitizeFloat(Geolocation.X));
+	UDatasmithAssetUserData::SetDatasmithUserDataValueForKey(SceneActor, TEXT("Geolocation_Longitude"),  FString::SanitizeFloat(Geolocation.Y));
+	UDatasmithAssetUserData::SetDatasmithUserDataValueForKey(SceneActor, TEXT("Geolocation_Elevation"),  FString::SanitizeFloat(Geolocation.Z));
+
 	ImportContext.ActorsContext.ImportSceneActor = SceneActor;
 
 	return SceneActor;
diff --git a/Engine/Source/Developer/Datasmith/DatasmithFacade/Private/DatasmithFacadeScene.cpp b/Engine/Source/Developer/Datasmith/DatasmithFacade/Private/DatasmithFacadeScene.cpp
--- a/Engine/Source/Developer/Datasmith/DatasmithFacade/Private/DatasmithFacadeScene.cpp
+++ b/Engine/Source/Developer/Datasmith/DatasmithFacade/Private/DatasmithFacadeScene.cpp
@@ -483,6 +483,19 @@ const TCHAR* FDatasmithFacadeScene::GetAssetsOutputPath() const
 	return SceneExporterRef->GetAssetsOutputPath();
 }
 
+void FDatasmithFacadeScene::SetGeolocation(double Latitude, double Longitude, double Elevation)
+{
+	SceneRef->SetGeolocation(FVector(Latitude, Longitude, Elevation));
+}
+
+void FDatasmithFacadeScene::GetGeolocation(double& OutLatitude, double& OutLongitude, double& OutElevation) const
+{
+	FVector Geolocation = SceneRef->GetGeolocation();
+	OutLatitude = Geolocation.X;
+	OutLongitude = Geolocation.Y;
+	OutElevation = Geolocation.Z;
+}
+
 void FDatasmithFacadeScene::Shutdown()
 {
 	FDatasmithExporterManager::Shutdown();
diff --git a/Engine/Source/Developer/Datasmith/DatasmithFacade/Public/DatasmithFacadeScene.h b/Engine/Source/Developer/Datasmith/DatasmithFacade/Public/DatasmithFacadeScene.h
--- a/Engine/Source/Developer/Datasmith/DatasmithFacade/Public/DatasmithFacadeScene.h
+++ b/Engine/Source/Developer/Datasmith/DatasmithFacade/Public/DatasmithFacadeScene.h
@@ -236,7 +236,7 @@ public:
 	/** Set the Datasmith scene name */
 	void SetName(const TCHAR* InName);
 
-	/** Set the Datasmith scene file name */
+	/** Get the Datasmith scene file name */
 	const TCHAR* GetName() const;
 
 	/** Set the path to the folder where the .datasmith file will be saved. */
@@ -248,6 +248,12 @@ public:
 	/** Get the path were additional assets will be saved to. */
 	const TCHAR* GetAssetsOutputPath() const;
 
+	/** Set the Datasmith scene Geolocation data*/
+	void SetGeolocation(double Latitude, double Longitude, double Elevation);
+
+	/** Get the Datasmith scene Geolocation data*/
+	void GetGeolocation(double& OutLatitude, double& OutLongitude, double& OutElevation) const;
+
 	/** Instantiate an exporter and register export start time */
 	void PreExport();
 
diff --git a/Engine/Source/Programs/Enterprise/Datasmith/DatasmithFacadeCSharp/Private/DatasmithFacadeCSharp.cpp b/Engine/Source/Programs/Enterprise/Datasmith/DatasmithFacadeCSharp/Private/DatasmithFacadeCSharp.cpp
--- a/Engine/Source/Programs/Enterprise/Datasmith/DatasmithFacadeCSharp/Private/DatasmithFacadeCSharp.cpp
+++ b/Engine/Source/Programs/Enterprise/Datasmith/DatasmithFacadeCSharp/Private/DatasmithFacadeCSharp.cpp
@@ -6028,6 +6028,34 @@ SWIGEXPORT wchar_t * SWIGSTDCALL CSharp_FDatasmithFacadeScene_GetAssetsOutputPat
 }
 
 
+SWIGEXPORT void SWIGSTDCALL CSharp_FDatasmithFacadeScene_SetGeolocation(void * jarg1, double jarg2, double jarg3, double jarg4) {
+  FDatasmithFacadeScene *arg1 = (FDatasmithFacadeScene *) 0 ;
+  double arg2 ;
+  double arg3 ;
+  double arg4 ;
+  
+  arg1 = (FDatasmithFacadeScene *)jarg1; 
+  arg2 = (double)jarg2; 
+  arg3 = (double)jarg3; 
+  arg4 = (double)jarg4; 
+  (arg1)->SetGeolocation(arg2,arg3,arg4);
+}
+
+
+SWIGEXPORT void SWIGSTDCALL CSharp_FDatasmithFacadeScene_GetGeolocation(void * jarg1, double * jarg2, double * jarg3, double * jarg4) {
+  FDatasmithFacadeScene *arg1 = (FDatasmithFacadeScene *) 0 ;
+  double *arg2 = 0 ;
+  double *arg3 = 0 ;
+  double *arg4 = 0 ;
+  
+  arg1 = (FDatasmithFacadeScene *)jarg1; 
+  arg2 = (double *)jarg2; 
+  arg3 = (double *)jarg3; 
+  arg4 = (double *)jarg4; 
+  ((FDatasmithFacadeScene const *)arg1)->GetGeolocation(*arg2,*arg3,*arg4);
+}
+
+
 SWIGEXPORT void SWIGSTDCALL CSharp_FDatasmithFacadeScene_PreExport(void * jarg1) {
   FDatasmithFacadeScene *arg1 = (FDatasmithFacadeScene *) 0 ;
   
diff --git a/Engine/Source/Programs/Enterprise/Datasmith/DatasmithFacadeCSharp/Public/DatasmithFacadeCSharpPINVOKE.cs b/Engine/Source/Programs/Enterprise/Datasmith/DatasmithFacadeCSharp/Public/DatasmithFacadeCSharpPINVOKE.cs
--- a/Engine/Source/Programs/Enterprise/Datasmith/DatasmithFacadeCSharp/Public/DatasmithFacadeCSharpPINVOKE.cs
+++ b/Engine/Source/Programs/Enterprise/Datasmith/DatasmithFacadeCSharp/Public/DatasmithFacadeCSharpPINVOKE.cs
@@ -1667,6 +1667,12 @@ class DatasmithFacadeCSharpPINVOKE {
   [global::System.Runtime.InteropServices.DllImport("DatasmithFacadeCSharp", EntryPoint="CSharp_FDatasmithFacadeScene_GetAssetsOutputPath")]
   public static extern global::System.IntPtr FDatasmithFacadeScene_GetAssetsOutputPath(global::System.Runtime.InteropServices.HandleRef jarg1);
 
+  [global::System.Runtime.InteropServices.DllImport("DatasmithFacadeCSharp", EntryPoint="CSharp_FDatasmithFacadeScene_SetGeolocation")]
+  public static extern void FDatasmithFacadeScene_SetGeolocation(global::System.Runtime.InteropServices.HandleRef jarg1, double jarg2, double jarg3, double jarg4);
+
+  [global::System.Runtime.InteropServices.DllImport("DatasmithFacadeCSharp", EntryPoint="CSharp_FDatasmithFacadeScene_GetGeolocation")]
+  public static extern void FDatasmithFacadeScene_GetGeolocation(global::System.Runtime.InteropServices.HandleRef jarg1, out double jarg2, out double jarg3, out double jarg4);
+
   [global::System.Runtime.InteropServices.DllImport("DatasmithFacadeCSharp", EntryPoint="CSharp_FDatasmithFacadeScene_PreExport")]
   public static extern void FDatasmithFacadeScene_PreExport(global::System.Runtime.InteropServices.HandleRef jarg1);
 
diff --git a/Engine/Source/Programs/Enterprise/Datasmith/DatasmithFacadeCSharp/Public/FDatasmithFacadeScene.cs b/Engine/Source/Programs/Enterprise/Datasmith/DatasmithFacadeCSharp/Public/FDatasmithFacadeScene.cs
--- a/Engine/Source/Programs/Enterprise/Datasmith/DatasmithFacadeCSharp/Public/FDatasmithFacadeScene.cs
+++ b/Engine/Source/Programs/Enterprise/Datasmith/DatasmithFacadeCSharp/Public/FDatasmithFacadeScene.cs
@@ -348,6 +348,14 @@ public class FDatasmithFacadeScene : global::System.IDisposable {
     return ret;
   }
 
+  public void SetGeolocation(double Latitude, double Longitude, double Elevation) {
+    DatasmithFacadeCSharpPINVOKE.FDatasmithFacadeScene_SetGeolocation(swigCPtr, Latitude, Longitude, Elevation);
+  }
+
+  public void GetGeolocation(out double OutLatitude, out double OutLongitude, out double OutElevation) {
+    DatasmithFacadeCSharpPINVOKE.FDatasmithFacadeScene_GetGeolocation(swigCPtr, out OutLatitude, out OutLongitude, out OutElevation);
+  }
+
   public void PreExport() {
     DatasmithFacadeCSharpPINVOKE.FDatasmithFacadeScene_PreExport(swigCPtr);
   }
diff --git a/Engine/Source/Runtime/Datasmith/DatasmithCore/Private/DatasmithSceneElementsImpl.cpp b/Engine/Source/Runtime/Datasmith/DatasmithCore/Private/DatasmithSceneElementsImpl.cpp
--- a/Engine/Source/Runtime/Datasmith/DatasmithCore/Private/DatasmithSceneElementsImpl.cpp
+++ b/Engine/Source/Runtime/Datasmith/DatasmithCore/Private/DatasmithSceneElementsImpl.cpp
@@ -1060,6 +1060,7 @@ FDatasmithSceneImpl::FDatasmithSceneImpl(const TCHAR * InName)
 	Store.RegisterParameter(ProductVersion,     "ProductVersion"     );
 	Store.RegisterParameter(UserID,             "UserID"             );
 	Store.RegisterParameter(UserOS,             "UserOS"             );
+	Store.RegisterParameter(Geolocation,        "Geolocation"        );
 	Store.RegisterParameter(ExportDuration,     "ExportDuration"     );
 	Store.RegisterParameter(bUseSky,            "bUseSky"            );
 	FDatasmithSceneImpl::Reset();
@@ -1085,6 +1086,7 @@ void FDatasmithSceneImpl::Reset()
 	ProductVersion = TEXT("");
 	UserID = TEXT("");
 	UserOS = TEXT("");
+	Geolocation = FVector::ZeroVector;
 	ResourcePath = TEXT("");
 
 	ExportDuration = 0;
diff --git a/Engine/Source/Runtime/Datasmith/DatasmithCore/Private/DatasmithSceneElementsImpl.h b/Engine/Source/Runtime/Datasmith/DatasmithCore/Private/DatasmithSceneElementsImpl.h
--- a/Engine/Source/Runtime/Datasmith/DatasmithCore/Private/DatasmithSceneElementsImpl.h
+++ b/Engine/Source/Runtime/Datasmith/DatasmithCore/Private/DatasmithSceneElementsImpl.h
@@ -1609,6 +1609,9 @@ public:
 	virtual const TCHAR* GetUserOS() const override { return *(FString&)UserOS; }
 	virtual void SetUserOS(const TCHAR* InUserOS) override { UserOS = InUserOS; }
 
+	virtual FVector GetGeolocation() const override { return Geolocation; }
+	virtual void SetGeolocation(FVector InGeolocation) override { Geolocation = InGeolocation; }
+
 	virtual int32 GetExportDuration() const override { return ExportDuration; }
 	virtual void SetExportDuration(int32 InExportDuration) override { ExportDuration = InExportDuration; }
 
@@ -1704,6 +1707,7 @@ private:
 	TReflected<FString> ProductVersion;
 	TReflected<FString> UserID;
 	TReflected<FString> UserOS;
+	TReflected<FVector> Geolocation;
 	TReflected<FString> ResourcePath;
 
 	TReflected<uint32> ExportDuration;
diff --git a/Engine/Source/Runtime/Datasmith/DatasmithCore/Private/DatasmithSceneXmlReader.cpp b/Engine/Source/Runtime/Datasmith/DatasmithCore/Private/DatasmithSceneXmlReader.cpp
--- a/Engine/Source/Runtime/Datasmith/DatasmithCore/Private/DatasmithSceneXmlReader.cpp
+++ b/Engine/Source/Runtime/Datasmith/DatasmithCore/Private/DatasmithSceneXmlReader.cpp
@@ -1201,6 +1201,13 @@ bool FDatasmithSceneXmlReader::ParseXmlFile(TSharedRef< IDatasmithScene >& OutSc
 			OutScene->SetUserID(*Nodes[i]->GetAttribute(DATASMITH_USERID));
 			OutScene->SetUserOS(*Nodes[i]->GetAttribute(DATASMITH_USEROS));
 		}
+		else if (Nodes[i]->GetTag() == DATASMITH_GEOLOCATION)
+		{
+			OutScene->SetGeolocation(FVector(
+				ValueFromString<double>(Nodes[i]->GetAttribute(DATASMITH_GEOLOCATION_LATITUDE)),
+				ValueFromString<double>(Nodes[i]->GetAttribute(DATASMITH_GEOLOCATION_LONGITUDE)),
+				ValueFromString<double>(Nodes[i]->GetAttribute(DATASMITH_GEOLOCATION_ELEVATION))));
+		}
 		// STATIC MESHES
 		else if (Nodes[i]->GetTag() == DATASMITH_STATICMESHNAME)
 		{
diff --git a/Engine/Source/Runtime/Datasmith/DatasmithCore/Private/DatasmithSceneXmlWriter.cpp b/Engine/Source/Runtime/Datasmith/DatasmithCore/Private/DatasmithSceneXmlWriter.cpp
--- a/Engine/Source/Runtime/Datasmith/DatasmithCore/Private/DatasmithSceneXmlWriter.cpp
+++ b/Engine/Source/Runtime/Datasmith/DatasmithCore/Private/DatasmithSceneXmlWriter.cpp
@@ -1932,6 +1932,19 @@ void FDatasmithSceneXmlWriter::Serialize( TSharedRef< IDatasmithScene > Datasmit
 	XmlString = TEXT("<User ID=\"") + FPlatformMisc::GetLoginId() + TEXT("\" OS=\"") + OSVersion + TEXT("\"/>") + LINE_TERMINATOR;
 	FDatasmithSceneXmlWriterImpl::SerializeToArchive(Archive, XmlString);
 
+	FDatasmithSceneXmlWriterImpl::WriteIndent(Archive, 1);
+
+	FVector Geolocation = DatasmithScene->GetGeolocation();
+	
+	XmlString = FString::Printf( TEXT("<%s %s=\"%s\" %s=\"%s\" %s=\"%s\"/>"),
+		DATASMITH_GEOLOCATION,
+		DATASMITH_GEOLOCATION_LATITUDE, *FString::SanitizeFloat(Geolocation.X),
+		DATASMITH_GEOLOCATION_LONGITUDE, *FString::SanitizeFloat(Geolocation.Y),
+		DATASMITH_GEOLOCATION_ELEVATION, *FString::SanitizeFloat(Geolocation.Z)
+		) + LINE_TERMINATOR;
+
+	FDatasmithSceneXmlWriterImpl::SerializeToArchive(Archive, XmlString);
+
 	FDatasmithSceneXmlWriterImpl::WritePostProcessElement( DatasmithScene->GetPostProcess(), Archive, 1 );
 
 	bool bHasIlluminationEnvironment = false;
diff --git a/Engine/Source/Runtime/Datasmith/DatasmithCore/Public/DatasmithDefinitions.h b/Engine/Source/Runtime/Datasmith/DatasmithCore/Public/DatasmithDefinitions.h
--- a/Engine/Source/Runtime/Datasmith/DatasmithCore/Public/DatasmithDefinitions.h
+++ b/Engine/Source/Runtime/Datasmith/DatasmithCore/Public/DatasmithDefinitions.h
@@ -451,6 +451,12 @@ static const TCHAR* KeyValuePropertyTypeStrings[] = { TEXT("String"), TEXT("Colo
 #define DATASMITH_EXPORTDURATION				TEXT("Duration")
 #define DATASMITH_RESOURCEPATH					TEXT("ResourcePath")
 
+// SCENE GEOLOCATION
+#define DATASMITH_GEOLOCATION					TEXT("Geolocation")
+#define DATASMITH_GEOLOCATION_LATITUDE			TEXT("lat")
+#define DATASMITH_GEOLOCATION_LONGITUDE			TEXT("lon")
+#define DATASMITH_GEOLOCATION_ELEVATION			TEXT("ele")
+
 //ELEMENTS
 #define DATASMITH_HASH							TEXT("Hash")
 #define DATASMITH_ENABLED						TEXT("Enabled")
diff --git a/Engine/Source/Runtime/Datasmith/DatasmithCore/Public/IDatasmithSceneElements.h b/Engine/Source/Runtime/Datasmith/DatasmithCore/Public/IDatasmithSceneElements.h
--- a/Engine/Source/Runtime/Datasmith/DatasmithCore/Public/IDatasmithSceneElements.h
+++ b/Engine/Source/Runtime/Datasmith/DatasmithCore/Public/IDatasmithSceneElements.h
@@ -1554,6 +1554,12 @@ public:
 	 */
 	virtual void SetUserOS(const TCHAR*) = 0;
 
+	/** Get Geolocation data of the scene. Where X = Latitude, Y = Longitude, Z = Elevation */
+	virtual FVector GetGeolocation() const = 0; 
+
+	/** Set Geolocation data of the scene. Where X = Latitude, Y = Longitude, Z = Elevation */
+	virtual void SetGeolocation(FVector) = 0;
+
 	/** Returns the time taken to export the scene */
 	virtual int32 GetExportDuration() const = 0;
 
