[Sketchup] Datasmith Exporter to transfer the model geo-location information as metadata to UE

diff --git a/Engine/Source/Programs/Enterprise/Datasmith/DatasmithSketchUpRubyExporter/Private/DatasmithSketchUpComponent.cpp b/Engine/Source/Programs/Enterprise/Datasmith/DatasmithSketchUpRubyExporter/Private/DatasmithSketchUpComponent.cpp
--- a/Engine/Source/Programs/Enterprise/Datasmith/DatasmithSketchUpRubyExporter/Private/DatasmithSketchUpComponent.cpp
+++ b/Engine/Source/Programs/Enterprise/Datasmith/DatasmithSketchUpRubyExporter/Private/DatasmithSketchUpComponent.cpp
@@ -17,6 +17,7 @@
 #include <SketchUpAPI/model/entity.h>
 #include "SketchUpAPI/model/group.h"
 #include "SketchUpAPI/model/layer.h"
+#include "SketchUpAPI/model/location.h"
 #include "SketchUpAPI/model/model.h"
 #include "SketchUpAPI/model/component_instance.h"
 #include "SketchUpAPI/geometry/transformation.h"
@@ -345,6 +346,7 @@ void FModelDefinition::UpdateGeometry(FExportContext& Context)
 
 void FModelDefinition::UpdateMetadata(FExportContext& Context)
 {
+
 }
 
 void FModelDefinition::InvalidateInstancesGeometry(FExportContext& Context)
@@ -382,6 +384,31 @@ SUTransformation FModelDefinition::GetMeshBakedTransform()
 	return Transform;
 }
 
+bool FModelDefinition::UpdateModel(FExportContext& Context)
+{
+	// SketchUp API has no notification of Geolocation change so retrieve it every time and compare to check if we need to set DirectLink update
+	FVector Geolocation = Context.DatasmithScene->GetGeolocation();
+	FVector GeolocationNew = Geolocation;
+
+	SULocationRef Location = SU_INVALID;
+	if (SU_ERROR_NONE == SUModelGetLocation(Model, &Location))
+	{
+		double Latitude = 0;
+		double Longitude = 0;
+		SULocationGetLatLong(Location, &Latitude, &Longitude);
+
+		GeolocationNew = FVector(Latitude, Longitude, 0);
+	}
+
+	if ((Geolocation - GeolocationNew).IsNearlyZero(1e-10))
+	{
+		return false;
+	}
+
+	Context.DatasmithScene->SetGeolocation(GeolocationNew);
+	return true;
+}
+
 FString FModelDefinition::GetSketchupSourceGUID()
 {
 	return TEXT("MODEL");
diff --git a/Engine/Source/Programs/Enterprise/Datasmith/DatasmithSketchUpRubyExporter/Private/DatasmithSketchUpComponent.h b/Engine/Source/Programs/Enterprise/Datasmith/DatasmithSketchUpRubyExporter/Private/DatasmithSketchUpComponent.h
--- a/Engine/Source/Programs/Enterprise/Datasmith/DatasmithSketchUpRubyExporter/Private/DatasmithSketchUpComponent.h
+++ b/Engine/Source/Programs/Enterprise/Datasmith/DatasmithSketchUpRubyExporter/Private/DatasmithSketchUpComponent.h
@@ -220,7 +220,7 @@ namespace DatasmithSketchUp
 	public:
 		FModelDefinition(SUModelRef Model);
 
-		// Being FDefinition
+		// Beging FDefinition
 		virtual void Parse(FExportContext& Context) override;
 		
 		virtual void UpdateGeometry(FExportContext& Context) override;
@@ -240,6 +240,8 @@ namespace DatasmithSketchUp
 		virtual SUTransformation GetMeshBakedTransform() override;
 		// End FDefinition
 
+		bool UpdateModel(FExportContext& Context);
+
 	private:
 		SUModelRef Model;
 	};
diff --git a/Engine/Source/Programs/Enterprise/Datasmith/DatasmithSketchUpRubyExporter/Private/DatasmithSketchUpExportContext.cpp b/Engine/Source/Programs/Enterprise/Datasmith/DatasmithSketchUpRubyExporter/Private/DatasmithSketchUpExportContext.cpp
--- a/Engine/Source/Programs/Enterprise/Datasmith/DatasmithSketchUpRubyExporter/Private/DatasmithSketchUpExportContext.cpp
+++ b/Engine/Source/Programs/Enterprise/Datasmith/DatasmithSketchUpRubyExporter/Private/DatasmithSketchUpExportContext.cpp
@@ -166,6 +166,8 @@ bool FExportContext::Update(bool bModifiedHint)
 
 	bModified |= Scenes.Update();
 
+	bModified |= ModelDefinition->UpdateModel(*this);
+
 	if (!bModified)
 	{
 		// code below is not supposed to change Datasmith scene if not modification hint(i.e. no invalidated data) was present
