Eliminate warnings from install_name_tool

diff --git a/Engine/Source/Programs/Enterprise/Datasmith/DatasmithARCHICADExporter/Build/Config/setupdylibs.sh b/Engine/Source/Programs/Enterprise/Datasmith/DatasmithARCHICADExporter/Build/Config/setupdylibs.sh
--- a/Engine/Source/Programs/Enterprise/Datasmith/DatasmithARCHICADExporter/Build/Config/setupdylibs.sh
+++ b/Engine/Source/Programs/Enterprise/Datasmith/DatasmithARCHICADExporter/Build/Config/setupdylibs.sh
@@ -57,7 +57,7 @@ SetUpThirdPartyDll() {
 		echo "Copy $DylibName"
 		cp "$DylibPath" "$OurDylibFolder"
 		chmod +w "$OurDylibFolder/$DylibName"
-		install_name_tool -id @loader_path/$DylibName "$OurDylibFolder/$DylibName"
+		install_name_tool -id @loader_path/$DylibName "$OurDylibFolder/$DylibName" > /dev/null 2>&1
 	fi
 }
 
@@ -74,8 +74,8 @@ SetUpDll() {
 		if [ -f "$OriginalDylibPath" ]; then
 			echo "Copy $DylibName"
 			cp "$OriginalDylibPath" "$OurDylibFolder"
-			install_name_tool -id @loader_path/$DylibName "$OurDylibFolder/$DylibName"
-			install_name_tool -change @rpath/$dylibLibFreeImage @loader_path/$dylibLibFreeImage "$OurDylibFolder/$DylibName"
+			install_name_tool -id @loader_path/$DylibName "$OurDylibFolder/$DylibName" > /dev/null 2>&1
+			install_name_tool -change @rpath/$dylibLibFreeImage @loader_path/$dylibLibFreeImage "$OurDylibFolder/$DylibName" > /dev/null 2>&1
 		else
 			echo "Missing $DylibName"
 		fi
