[3ds Max] RailClone/ForestPack mesh sharing

diff --git a/Engine/Source/Developer/Datasmith/DatasmithFacade/Private/DatasmithFacadeMesh.cpp b/Engine/Source/Developer/Datasmith/DatasmithFacade/Private/DatasmithFacadeMesh.cpp
--- a/Engine/Source/Developer/Datasmith/DatasmithFacade/Private/DatasmithFacadeMesh.cpp
+++ b/Engine/Source/Developer/Datasmith/DatasmithFacade/Private/DatasmithFacadeMesh.cpp
@@ -9,6 +9,12 @@
 #include "DatasmithMesh.h"
 #include "DatasmithMeshExporter.h"
 
+void FDatasmithFacadeMesh::CalculateHash(TCHAR OutBuffer[33], size_t BufferSize) const
+{
+	FString HashString = LexToString(GetDatasmithMesh().CalculateHash());
+	FCString::Strncpy(OutBuffer, *HashString, BufferSize);
+}
+
 void FDatasmithFacadeMesh::SetVertex(int32 Index, float X, float Y, float Z)
 {
 	FVector3f Position(FDatasmithFacadeElement::ConvertPosition(X, Y, Z));
diff --git a/Engine/Source/Developer/Datasmith/DatasmithFacade/Public/DatasmithFacadeMesh.h b/Engine/Source/Developer/Datasmith/DatasmithFacade/Public/DatasmithFacadeMesh.h
--- a/Engine/Source/Developer/Datasmith/DatasmithFacade/Public/DatasmithFacadeMesh.h
+++ b/Engine/Source/Developer/Datasmith/DatasmithFacade/Public/DatasmithFacadeMesh.h
@@ -31,6 +31,9 @@ public:
 		return GetDatasmithMesh().GetName();
 	}
 
+	/** Return a MD5 hash of the content of the Datasmith Mesh */
+	void CalculateHash(TCHAR OutBuffer[33], size_t BufferSize) const;
+
 	//--------------------------
 	// Faces
 	//--------------------------
diff --git a/Engine/Source/Programs/Enterprise/Datasmith/DatasmithFacadeCSharp/Private/DatasmithFacadeCSharp.cpp b/Engine/Source/Programs/Enterprise/Datasmith/DatasmithFacadeCSharp/Private/DatasmithFacadeCSharp.cpp
--- a/Engine/Source/Programs/Enterprise/Datasmith/DatasmithFacadeCSharp/Private/DatasmithFacadeCSharp.cpp
+++ b/Engine/Source/Programs/Enterprise/Datasmith/DatasmithFacadeCSharp/Private/DatasmithFacadeCSharp.cpp
@@ -2597,6 +2597,18 @@ SWIGEXPORT wchar_t * SWIGSTDCALL CSharp_FDatasmithFacadeMesh_GetName(void * jarg
 }
 
 
+SWIGEXPORT void SWIGSTDCALL CSharp_FDatasmithFacadeMesh_InternalCalculateHash(void * jarg1, void * jarg2, unsigned long jarg3) {
+  FDatasmithFacadeMesh *arg1 = (FDatasmithFacadeMesh *) 0 ;
+  TCHAR *arg2 ;
+  size_t arg3 ;
+  
+  arg1 = (FDatasmithFacadeMesh *)jarg1; 
+  arg2 = (TCHAR *)jarg2; 
+  arg3 = (size_t)jarg3; 
+  ((FDatasmithFacadeMesh const *)arg1)->CalculateHash(arg2,arg3);
+}
+
+
 SWIGEXPORT void SWIGSTDCALL CSharp_FDatasmithFacadeMesh_SetFacesCount(void * jarg1, int jarg2) {
   FDatasmithFacadeMesh *arg1 = (FDatasmithFacadeMesh *) 0 ;
   int32 arg2 ;
diff --git a/Engine/Source/Programs/Enterprise/Datasmith/DatasmithFacadeCSharp/Public/DatasmithFacadeCSharpPINVOKE.cs b/Engine/Source/Programs/Enterprise/Datasmith/DatasmithFacadeCSharp/Public/DatasmithFacadeCSharpPINVOKE.cs
--- a/Engine/Source/Programs/Enterprise/Datasmith/DatasmithFacadeCSharp/Public/DatasmithFacadeCSharpPINVOKE.cs
+++ b/Engine/Source/Programs/Enterprise/Datasmith/DatasmithFacadeCSharp/Public/DatasmithFacadeCSharpPINVOKE.cs
@@ -803,6 +803,9 @@ class DatasmithFacadeCSharpPINVOKE {
   [global::System.Runtime.InteropServices.DllImport("DatasmithFacadeCSharp", EntryPoint="CSharp_FDatasmithFacadeMesh_GetName")]
   public static extern global::System.IntPtr FDatasmithFacadeMesh_GetName(global::System.Runtime.InteropServices.HandleRef jarg1);
 
+  [global::System.Runtime.InteropServices.DllImport("DatasmithFacadeCSharp", EntryPoint="CSharp_FDatasmithFacadeMesh_InternalCalculateHash")]
+  public static extern void FDatasmithFacadeMesh_InternalCalculateHash(global::System.Runtime.InteropServices.HandleRef jarg1, [global::System.Runtime.InteropServices.MarshalAs(global::System.Runtime.InteropServices.UnmanagedType.LPWStr)] System.Text.StringBuilder jarg2, ulong jarg3);
+
   [global::System.Runtime.InteropServices.DllImport("DatasmithFacadeCSharp", EntryPoint="CSharp_FDatasmithFacadeMesh_SetFacesCount")]
   public static extern void FDatasmithFacadeMesh_SetFacesCount(global::System.Runtime.InteropServices.HandleRef jarg1, int jarg2);
 
diff --git a/Engine/Source/Programs/Enterprise/Datasmith/DatasmithFacadeCSharp/Public/FDatasmithFacadeMesh.cs b/Engine/Source/Programs/Enterprise/Datasmith/DatasmithFacadeCSharp/Public/FDatasmithFacadeMesh.cs
--- a/Engine/Source/Programs/Enterprise/Datasmith/DatasmithFacadeCSharp/Public/FDatasmithFacadeMesh.cs
+++ b/Engine/Source/Programs/Enterprise/Datasmith/DatasmithFacadeCSharp/Public/FDatasmithFacadeMesh.cs
@@ -45,6 +45,13 @@ public class FDatasmithFacadeMesh : global::System.IDisposable {
     }
   }
 
+  public string CalculateHash()
+  {
+    System.Text.StringBuilder sb = new System.Text.StringBuilder(32);
+    InternalCalculateHash(sb, (ulong)sb.Capacity + 1);
+    return sb.ToString();
+  }
+
   public FDatasmithFacadeMesh() : this(DatasmithFacadeCSharpPINVOKE.new_FDatasmithFacadeMesh(), true) {
   }
 
@@ -57,6 +64,10 @@ public class FDatasmithFacadeMesh : global::System.IDisposable {
     return ret;
   }
 
+  private void InternalCalculateHash(System.Text.StringBuilder OutBuffer, ulong BufferSize) {
+    DatasmithFacadeCSharpPINVOKE.FDatasmithFacadeMesh_InternalCalculateHash(swigCPtr, OutBuffer, BufferSize);
+  }
+
   public void SetFacesCount(int NumFaces) {
     DatasmithFacadeCSharpPINVOKE.FDatasmithFacadeMesh_SetFacesCount(swigCPtr, NumFaces);
   }
diff --git a/Engine/Source/Programs/Enterprise/Datasmith/DatasmithMaxExporter/Private/DatasmithMaxDirectLinkItoo.cpp b/Engine/Source/Programs/Enterprise/Datasmith/DatasmithMaxExporter/Private/DatasmithMaxDirectLinkItoo.cpp
--- a/Engine/Source/Programs/Enterprise/Datasmith/DatasmithMaxExporter/Private/DatasmithMaxDirectLinkItoo.cpp
+++ b/Engine/Source/Programs/Enterprise/Datasmith/DatasmithMaxExporter/Private/DatasmithMaxDirectLinkItoo.cpp
@@ -47,6 +47,10 @@ bool ConvertRailClone(ISceneTracker& SceneTracker, FNodeTracker& NodeTracker)
 	}
 
 	RCStaticInterface->IRCRegisterEngine();
+	Interval ObjectValidity = RailCloneNode->GetObjectRef()->ObjectValidity(CurrentTime);
+
+	NodeTracker.Validity.NarrowValidityToInterval(ObjectValidity);
+
 	IRCInterface* RCInterface = GetRCInterface(RailCloneNode->GetObjectRef());
 	if (!RCInterface)
 	{
@@ -208,14 +212,13 @@ bool ConvertForest(ISceneTracker& Scene, FNodeTracker& NodeTracker)
 				FRenderMeshForConversion()
 			};
 
-			if (MeshSource.RenderMesh.IsValid())
+			if (MeshSource.IsValid())
 			{
 				for(const TPair<Mtl*, TArray<Matrix3>>& MaterialAndTransforms: Instances.InstancesTransformsForMaterial)
 				{
 					Scene.SetupDatasmithHISMForNode(NodeTracker, MeshSource, MaterialAndTransforms.Key, MeshIndex, MaterialAndTransforms.Value);
 				}
 
-
 				MeshIndex ++;
 			}
 		}
diff --git a/Engine/Source/Programs/Enterprise/Datasmith/DatasmithMaxExporter/Private/DatasmithMaxDirectLinkMaterial.cpp b/Engine/Source/Programs/Enterprise/Datasmith/DatasmithMaxExporter/Private/DatasmithMaxDirectLinkMaterial.cpp
--- a/Engine/Source/Programs/Enterprise/Datasmith/DatasmithMaxExporter/Private/DatasmithMaxDirectLinkMaterial.cpp
+++ b/Engine/Source/Programs/Enterprise/Datasmith/DatasmithMaxExporter/Private/DatasmithMaxDirectLinkMaterial.cpp
@@ -174,6 +174,11 @@ void FMaterialsCollectionTracker::UnSetMaterialsForMeshActor(const TSharedPtr<ID
 	MeshActorToUsedMaterialSlot.Remove(MeshActor.Get());
 }
 
+void FMaterialsCollectionTracker::SetMaterialsForMeshElement(FMeshConverted& MeshConverted, Mtl* Material)
+{
+	SetMaterialsForMeshElement(MeshConverted.GetDatasmithMeshElement(), Material, MeshConverted.SupportedChannels);
+}
+
 void FMaterialsCollectionTracker::SetMaterialsForMeshElement(const TSharedPtr<IDatasmithMeshElement>& MeshElement, Mtl* Material, const TSet<uint16>& SupportedChannels)
 {
 	if (FDatasmithMaxMatHelper::GetMaterialClass(Material) == EDSMaterialType::XRefMat)
diff --git a/Engine/Source/Programs/Enterprise/Datasmith/DatasmithMaxExporter/Private/DatasmithMaxDirectLinkMesh.cpp b/Engine/Source/Programs/Enterprise/Datasmith/DatasmithMaxExporter/Private/DatasmithMaxDirectLinkMesh.cpp
--- a/Engine/Source/Programs/Enterprise/Datasmith/DatasmithMaxExporter/Private/DatasmithMaxDirectLinkMesh.cpp
+++ b/Engine/Source/Programs/Enterprise/Datasmith/DatasmithMaxExporter/Private/DatasmithMaxDirectLinkMesh.cpp
@@ -563,19 +563,75 @@ int SetObjectParamValue(TimeValue CurrentTime, Object *Obj, const FString& Param
 namespace DatasmithMaxDirectLink
 {
 
+template<typename T>
+static void UpdateMD5SimpleType(FMD5& MD5, const T& Value)
+{
+	static_assert(TIsPODType<T>::Value, "Simple type required");
+	MD5.Update(reinterpret_cast<const uint8*>(&Value), sizeof(Value));
+}
+
+template<typename T>
+static void UpdateMD5Array(FMD5& MD5, TArray<T> Value)
+{
+	static_assert(TIsPODType<T>::Value, "This function requires POD array");
+	UpdateMD5SimpleType(MD5, Value.Num());
+	if (!Value.IsEmpty())
+	{
+		MD5.Update(reinterpret_cast<const uint8*>(Value.GetData()), Value.GetTypeSize()*Value.Num());
+	}
+}
+
+static void UpdateMD5(FMD5& MD5, const FDatasmithMesh& DatasmithMesh)
+{
+	FMD5Hash MeshHash = DatasmithMesh.CalculateHash();
+	MD5.Update(MeshHash.GetBytes(), MeshHash.GetSize());
+}
+
+FMD5Hash FDatasmithMeshConverter::ComputeHash()
+{
+	FMD5 MD5;
+
+	UpdateMD5(MD5, RenderMesh);
+
+	TArray<uint16> SupportedChannelsArray = SupportedChannels.Array();
+	SupportedChannelsArray.Sort();
+	UpdateMD5Array(MD5, SupportedChannelsArray);
+
+	UpdateMD5SimpleType(MD5, UVChannelsMap.Num());
+
+	if (!UVChannelsMap.IsEmpty())
+	{
+		TArray<TPair<int32, int32>> UVChannelsMapArray = UVChannelsMap.Array();
+		UVChannelsMapArray.Sort();
+		for (TPair<int32, int32> KVP : UVChannelsMapArray)
+		{
+			UpdateMD5SimpleType(MD5, KVP.Key);
+			UpdateMD5SimpleType(MD5, KVP.Value);
+		}
+	}
+
+	UpdateMD5SimpleType(MD5, SelectedLightmapUVChannel);
+
+	UpdateMD5SimpleType(MD5, bHasCollision);
+	if (bHasCollision)
+	{
+		UpdateMD5(MD5, CollisionMesh);
+	}
+	FMD5Hash Hash;
+	Hash.Set(MD5);
+	return Hash;
+}
+
+
 // todo: paralelize calls to ExportToUObject 
-bool ConvertMaxMeshToDatasmith(TimeValue CurrentTime, ISceneTracker& Scene, FMeshConverterSource& MeshSource, FMeshConverted& MeshConverted)
+bool ConvertMaxMeshToDatasmith(TimeValue CurrentTime, FMeshConverterSource& MeshSource, FDatasmithMeshConverter& DatasmithMeshConverter)
 {
 	TOptional<FDatasmithMaxStaticMeshAttributes> DatasmithAttributes = FDatasmithMaxStaticMeshAttributes::ExtractStaticMeshAttributes(MeshSource.Node);
 
-	MeshConverted.DatasmithMeshElement = FDatasmithSceneFactory::CreateMesh(*MeshSource.MeshName);
-
 	if (DatasmithAttributes && DatasmithAttributes->GetExportMode() == EStaticMeshExportMode::BoundingBox)
 	{
-		FDatasmithMesh DatasmithMesh;
-		if (GeomUtils::FillDatasmithMeshFromBoundingBox(CurrentTime, DatasmithMesh, MeshSource.RenderMesh))
+		if (GeomUtils::FillDatasmithMeshFromBoundingBox(CurrentTime, DatasmithMeshConverter.RenderMesh, MeshSource.RenderMesh))
 		{
-			Scene.AddMeshElement(MeshConverted.DatasmithMeshElement, DatasmithMesh, nullptr);
 			return true;
 		}
 		else
@@ -585,15 +641,13 @@ bool ConvertMaxMeshToDatasmith(TimeValue CurrentTime, ISceneTracker& Scene, FMes
 		}
 	}
 
-	FDatasmithMesh DatasmithMesh;
-
 	MeshConversionParams RenderMeshParams = {
 		MeshSource.RenderMesh.GetNode(),
 		MeshSource.RenderMesh,
 		MeshSource.bConsolidateMaterialIds
 	};
 
-	if (!GeomUtils::CreateDatasmithMeshFromMaxMesh(DatasmithMesh, RenderMeshParams, MeshConverted.SupportedChannels, MeshConverted.UVChannelsMap))
+	if (!GeomUtils::CreateDatasmithMeshFromMaxMesh(DatasmithMeshConverter.RenderMesh, RenderMeshParams, DatasmithMeshConverter.SupportedChannels, DatasmithMeshConverter.UVChannelsMap))
 	{
 		LogWarning(FString(TEXT("Invalid object: ")) + MeshSource.Node->GetName());
 		return false;
@@ -602,25 +656,9 @@ bool ConvertMaxMeshToDatasmith(TimeValue CurrentTime, ISceneTracker& Scene, FMes
 	// Mapping between the 3ds max channel and the exported mesh channel
 	if (DatasmithAttributes)
 	{
-		constexpr int32 MaxToUnrealUVOffset = -1;
-		constexpr int32 DefaultValue = -1;
-		const int32 SelectedLightmapUVChannel = DatasmithAttributes->GetLightmapUVChannel();
-		const int32* ExportedSelectedChannel = MeshConverted.UVChannelsMap.Find(SelectedLightmapUVChannel + MaxToUnrealUVOffset);
-
-		if (ExportedSelectedChannel)
-		{
-			MeshConverted.DatasmithMeshElement->SetLightmapCoordinateIndex(*ExportedSelectedChannel);
-		}
-		else if (SelectedLightmapUVChannel != DefaultValue)
-		{
-			LogWarning(*FString::Printf(TEXT("%s won't use the channel %i for its lightmap because it's not supported by the mesh. A new channel will be generated.")
-				, static_cast<const TCHAR*>(MeshSource.Node->GetName())
-				, SelectedLightmapUVChannel));
-		}
+		DatasmithMeshConverter.SelectedLightmapUVChannel = DatasmithAttributes->GetLightmapUVChannel();
 	}
 
-	FDatasmithMesh* DatasmithCollisionMeshPtr = nullptr;
-	FDatasmithMesh DatasmithCollisionMesh;
 	if (MeshSource.CollisionMesh.IsValid())
 	{
 		TSet<uint16> SupportedChannelsDummy; // ignore map channels for collision mesh
@@ -632,13 +670,12 @@ bool ConvertMaxMeshToDatasmith(TimeValue CurrentTime, ISceneTracker& Scene, FMes
 			true // Consolidate material ids into single mesh for collision
 		};
 
-		if (GeomUtils::CreateDatasmithMeshFromMaxMesh(DatasmithCollisionMesh, CollisionParams, SupportedChannelsDummy, UVChannelsMapDummy))
+		if (GeomUtils::CreateDatasmithMeshFromMaxMesh(DatasmithMeshConverter.CollisionMesh, CollisionParams, SupportedChannelsDummy, UVChannelsMapDummy))
 		{
-			DatasmithCollisionMeshPtr = &DatasmithCollisionMesh;
+			DatasmithMeshConverter.bHasCollision = true;
 		}
 	}
 
-	Scene.AddMeshElement(MeshConverted.DatasmithMeshElement, DatasmithMesh, DatasmithCollisionMeshPtr);
 	return true;
 }
 
diff --git a/Engine/Source/Programs/Enterprise/Datasmith/DatasmithMaxExporter/Private/DatasmithMaxExporterWithDirectLink.cpp b/Engine/Source/Programs/Enterprise/Datasmith/DatasmithMaxExporter/Private/DatasmithMaxExporterWithDirectLink.cpp
--- a/Engine/Source/Programs/Enterprise/Datasmith/DatasmithMaxExporter/Private/DatasmithMaxExporterWithDirectLink.cpp
+++ b/Engine/Source/Programs/Enterprise/Datasmith/DatasmithMaxExporter/Private/DatasmithMaxExporterWithDirectLink.cpp
@@ -64,6 +64,13 @@ MAX_INCLUDES_START
 	#include "maxscript/maxscript.h"
 MAX_INCLUDES_END
 
+// Make FMD5Hash usable in TMap as a key
+inline uint32 GetTypeHash(const FMD5Hash& Hash)
+{
+	uint32* HashAsInt32 = (uint32*)Hash.GetBytes();
+	return HashAsInt32[0] ^ HashAsInt32[1] ^ HashAsInt32[2] ^ HashAsInt32[3];
+}
+
 namespace DatasmithMaxDirectLink
 {
 
@@ -152,13 +159,13 @@ private:
 };
 
 // Every node which is resolved to the same object is considered an instance
-// This class holds all this nodes which resolve to the same object
+// This class holds all the nodes which resolve to the same object
 struct FInstances: FNoncopyable
 {
 	AnimHandle Handle; // Handle of anim that is instanced
 
 	Object* EvaluatedObj = nullptr;
-	Mtl* Material = nullptr; // Material assigned to Datasmith StaticMesh, used to check if a particular instance needs to verride it
+	Mtl* Material = nullptr; // Material assigned to Datasmith StaticMesh, used to check if a particular instance needs to override it
 
 	TSet<class FNodeTracker*> NodeTrackers;
 
@@ -170,18 +177,18 @@ struct FInstances: FNoncopyable
 
 	bool HasMesh()
 	{
-		return Converted.DatasmithMeshElement.IsValid();
+		return Converted.GetDatasmithMeshElement().IsValid();
 	}
 
 	const TCHAR* GetStaticMeshPathName()
 	{
-		return Converted.DatasmithMeshElement->GetName();
+		return Converted.GetDatasmithMeshElement()->GetName();
 	}
 
-	void AssignMaterialToStaticMesh(FMaterialsCollectionTracker& MaterialsCollectionTracker,  Mtl* InMaterial)
+	// Record wich material is assigned to static mesh
+	void AssignMaterialToStaticMesh(Mtl* InMaterial)
 	{
 		Material = InMaterial;
-		MaterialsCollectionTracker.SetMaterialsForMeshElement(Converted.DatasmithMeshElement, Material, Converted.SupportedChannels);
 	}
 };
 
@@ -2145,9 +2152,13 @@ public:
 	// Called when mesh element is not needed anymore and should be removed from the scene
 	virtual void ReleaseMeshElement(FMeshConverted& Converted) override
 	{
-		MaterialsCollectionTracker.UnSetMaterialsForMeshElement(Converted.DatasmithMeshElement);
-		GetDatasmithScene().RemoveMesh(Converted.DatasmithMeshElement);
-		Converted.ReleaseMeshConverted();
+		if (Meshes.Release(Converted))
+		{
+			MaterialsCollectionTracker.UnSetMaterialsForMeshElement(Converted.DatasmithMeshElement);
+			GetDatasmithScene().RemoveMesh(Converted.DatasmithMeshElement);
+
+			Meshes.Remove(Converted);
+		}
 	}
 
 	void ReleaseNodeTrackerFromLayer(FNodeTracker& NodeTracker)
@@ -2444,7 +2455,7 @@ public:
 			// Invalidated instances without actual instances left(all removed)
 
 			// NOTE: release mesh element and release usage of Instances pointer
-			//   BEFORE deallocating Instances instance from the map below
+			//   BEFORE RemoveInstances call(which deallocates Instances object)
 			ReleaseMeshElement(Instances.Converted);
 
 			InstancesManager.RemoveInstances(Instances);
@@ -2527,7 +2538,7 @@ public:
 
 			// todo: use single EnumProc instance to enumerate all nodes during update to:
 			//    - have single call to BeginEnumeration and EndEnumeration
-			//    - track all Begin'd nodes to End them together after all is updated(to prevent duplicated Begin's of referenced objects that might be shared by different noвes)
+			//    - track all Begin'd nodes to End them together after all is updated(to prevent duplicated Begin's of referenced objects that might be shared by different nodes)
 			NodesPreparer.PrepareNode(NodeTracker.Node);
 			SCENE_UPDATE_STAT_INC(UpdateInstances, GeometryUpdated);
 			UpdateInstancesGeometry(Instances, NodeTracker);
@@ -2537,7 +2548,8 @@ public:
 			static_cast<FMeshNodeConverter&>(NodeTracker.GetConverter()).bMaterialsAssignedToStaticMesh = true;
 			if (Mtl* Material = UpdateGeometryNodeMaterial(*this, Instances, NodeTracker))
 			{
-				Instances.AssignMaterialToStaticMesh(MaterialsCollectionTracker, Material);
+				MaterialsCollectionTracker.SetMaterialsForMeshElement(Instances.Converted, Material);
+				Instances.AssignMaterialToStaticMesh(Material);
 			}
 		}
 
@@ -2900,9 +2912,9 @@ public:
 			GeomUtils::GetMeshForCollision(CurrentSyncPoint.Time, *this, Node, Instances.bShouldBakePivot),
 		};
 
-		Instances.ValidityInterval = MeshSource.RenderMesh.GetValidityInterval() & MeshSource.CollisionMesh.GetValidityInterval();
+		Instances.ValidityInterval = MeshSource.GetValidityInterval();
 
-		if (MeshSource.RenderMesh.GetMesh())
+		if (MeshSource.IsValid())
 		{
 			bool bHasInstanceWithMultimat = false;
 			for (FNodeTracker* InstanceNodeTracker : Instances.NodeTrackers)
@@ -2918,17 +2930,19 @@ public:
 
 			MeshSource.bConsolidateMaterialIds = !bHasInstanceWithMultimat;
 
-			FMeshes::AddMesh(*this, MeshSource, Instances.Converted, [&](bool bHasConverted, FMeshConverted& MeshConverted)
+			// Not sharing meshes for regular nodes when they don't resolve to the same Max object -
+			// resusing identical(but different) mesh is done only for RailClone instances - RC tends to make
+			// different in different RC objects instancing same Max meshes
+			Meshes.AddMesh(*this, MeshSource, Instances.Converted, false, [&](bool bHasConverted, bool bWasReused)
 			{
 				if (bHasConverted)
 				{
-					MeshConverted.DatasmithMeshElement->SetLabel(*NodeTrackersNames.GetNodeName(NodeTracker));
+					Instances.Converted.GetDatasmithMeshElement()->SetLabel(*NodeTrackersNames.GetNodeName(NodeTracker));
 				}
 				else
 				{
-					MeshConverted.DatasmithMeshElement.Reset();
+					Instances.Converted.DatasmithMeshElement.Reset();
 				}
-				Instances.Converted = MeshConverted;
 			});
 		}
 		else
@@ -2966,15 +2980,135 @@ public:
 	class FMeshes
 	{
 	public:
+
+		TMap<FMD5Hash, TSharedPtr<FDatasmithMeshConverter>> CachedMeshes;
+		TMap<TSharedPtr<IDatasmithMeshElement>, FMD5Hash> CachedMeshHashes;
+		TMap<TSharedPtr<IDatasmithMeshElement>, int32> MeshUsers;
+
+		// CompletionCallback is called when mehs is processed, passing flag identifing if actual Datasmith mesh was created(sometimes it's not created - e.g. mesh is empty)
+		// or is a cached mesh was reused
 		FORCENOINLINE
-		static void AddMesh(FSceneTracker& Scene, FMeshConverterSource& MeshSource, FMeshConverted& MeshConverted, TFunction<void(bool, FMeshConverted&)> CompletionCallback)
+		void AddMesh(FSceneTracker& Scene, FMeshConverterSource& MeshSource, FMeshConverted& MeshConverted, bool bShouldCache, TFunction<void(bool bConverted, bool bFoundCached)> CompletionCallback)
+		{
+
+			TSharedPtr<FDatasmithMeshConverter> MeshConverter = MakeShared<FDatasmithMeshConverter>();
+			bool bConverted = ConvertMaxMeshToDatasmith(Scene.CurrentSyncPoint.Time, MeshSource, *MeshConverter);
+			bool bFoundCached = false;
+			if (bConverted)
+			{
+				FMD5Hash MeshHash;
+				if (bShouldCache)
+				{
+					MeshHash = MeshConverter->ComputeHash();
+
+					if (TSharedPtr<FDatasmithMeshConverter>* CachedMeshConverter = CachedMeshes.Find(MeshHash))
+					{
+						// If there's mesh that is the same cached reuse it
+						// For RailClone/ForestPack meshes we have no other way to find out if RC objects are using the same mesh than to compare those meshes
+						// For regular Max nodes it's mush easier to find instances but RC gives away different Mesh objects for identical instances
+						bFoundCached = true;
+						MeshConverter = *CachedMeshConverter;
+					}
+					else
+					{
+						CachedMeshes.Add(MeshHash, MeshConverter);
+					}
+				}
+
+				if (bFoundCached)
+				{
+					if (MeshConverter->Converted.GetDatasmithMeshElement() != MeshConverted.GetDatasmithMeshElement())
+					{
+						// Release old converted mesh and hold onto new if new mesh is different(same was already cached)
+						Scene.ReleaseMeshElement(MeshConverted);
+						Acquire(MeshConverter->Converted);
+					}
+				}
+				else
+				{
+					// Release old converted mesh before building new
+					Scene.ReleaseMeshElement(MeshConverted);
+	
+					MeshConverter->Converted.DatasmithMeshElement = FDatasmithSceneFactory::CreateMesh(*MeshSource.MeshName);
+					MeshConverter->Converted.SupportedChannels = MeshConverter->SupportedChannels;
+					MeshConverter->Converted.UVChannelsMap = MeshConverter->UVChannelsMap;
+					if (bShouldCache)
+					{
+						CachedMeshHashes.Add(MeshConverter->Converted.DatasmithMeshElement, MeshHash);
+					}
+
+					int32 SelectedLightmapUVChannel = MeshConverter->SelectedLightmapUVChannel;
+					if (SelectedLightmapUVChannel >= 0)
+					{
+						constexpr int32 MaxToUnrealUVOffset = -1;
+						constexpr int32 DefaultValue = -1;
+						const int32* ExportedSelectedChannel = MeshConverter->UVChannelsMap.Find(SelectedLightmapUVChannel + MaxToUnrealUVOffset);
+
+						if (ExportedSelectedChannel)
+						{
+							MeshConverter->Converted.DatasmithMeshElement->SetLightmapCoordinateIndex(*ExportedSelectedChannel);
+						}
+						else if (SelectedLightmapUVChannel != DefaultValue)
+						{
+							LogWarning(*FString::Printf(TEXT("%s won't use the channel %i for its lightmap because it's not supported by the mesh. A new channel will be generated.")
+							                            , static_cast<const TCHAR*>(MeshSource.Node->GetName())
+							                            , SelectedLightmapUVChannel));
+						}
+					}
+
+					// Actually export mesh to disk
+					Scene.AddMeshElement(MeshConverter->Converted.DatasmithMeshElement, MeshConverter->RenderMesh, MeshConverter->bHasCollision ? &MeshConverter->CollisionMesh : nullptr);
+					Acquire(MeshConverter->Converted);
+				}
+				MeshConverted = MeshConverter->Converted;
+			}
+			else
+			{
+				Scene.ReleaseMeshElement(MeshConverted);
+				MeshConverted = FMeshConverted();
+			}
+
+			CompletionCallback(bConverted, bFoundCached);
+		}
+
+		void Reset()
+		{
+			CachedMeshes.Empty();
+			CachedMeshHashes.Empty();
+			MeshUsers.Empty();
+		}
+
+		void Acquire(const FMeshConverted& Converted)
+		{
+			MeshUsers.FindOrAdd(Converted.DatasmithMeshElement, 0) += 1;
+		}
+
+		// Returns true if mesh has no users afterwards
+		bool Release(const FMeshConverted& Converted)
+		{
+			if (int32* Found = MeshUsers.Find(Converted.DatasmithMeshElement))
+			{
+				int32& UserCount = *Found;
+				UserCount -= 1;
+
+				if (UserCount == 0)
+				{
+					return true;
+				}
+			}
+			return false;
+		}
+
+		void Remove(FMeshConverted& Converted)
 		{
-			// Reset old mesh
-			// todo: potential mesh reuse - when DatasmithMeshElement allows to reset materials(as well as other params)
-			Scene.ReleaseMeshElement(MeshConverted);
+			MeshUsers.Remove(Converted.DatasmithMeshElement);
+			FMD5Hash MeshHash;
+			if (CachedMeshHashes.RemoveAndCopyValue(Converted.DatasmithMeshElement, MeshHash))
+			{
+				CachedMeshes.Remove(MeshHash);
+			}
 
-			bool bConverted = ConvertMaxMeshToDatasmith(Scene.CurrentSyncPoint.Time, Scene, MeshSource, MeshConverted);
-			CompletionCallback(bConverted, MeshConverted);
+			Converted.ReleaseMeshConverted();
 		}
 	};
 
@@ -2983,20 +3117,22 @@ public:
 		MeshSource.MeshName = FString::FromInt(NodeTracker.Node->GetHandle()) + TEXT("_") + FString::FromInt(MeshIndex);
 		// note: when export Mesh goes to other place due to parallelizing it's result would be unknown here so MeshIndex handling will change(i.e. increment for any mesh)
 
-		FMeshConverted MeshConvertedDummy; // todo: possible reuse of previous converted mesh
-		FMeshes::AddMesh(*this, MeshSource, MeshConvertedDummy, [&](bool bHasConverted, FMeshConverted& MeshConverted)
+		FMeshConverted MeshConverted; // todo: possible reuse of previous converted mesh
+		Meshes.AddMesh(*this, MeshSource, MeshConverted, true, [&](bool bHasConverted, bool bReused)
 		{
 			if (bHasConverted)
 			{
 				FHismNodeConverter& NodeConverter = static_cast<FHismNodeConverter&>(NodeTracker.GetConverter());
-
 				NodeConverter.Meshes.Add(MeshConverted);
 
 				RegisterNodeForMaterial(NodeTracker, Material);
-				MaterialsCollectionTracker.SetMaterialsForMeshElement(MeshConverted.DatasmithMeshElement, Material, MeshConverted.SupportedChannels);
-
 				FString MeshLabel = FString(MeshSource.Node->GetName()) + (TEXT("_") + FString::FromInt(MeshIndex));
-				MeshConverted.DatasmithMeshElement->SetLabel(*MeshLabel);
+
+				if (!bReused)
+				{
+					MaterialsCollectionTracker.SetMaterialsForMeshElement(MeshConverted, Material);
+					MeshConverted.GetDatasmithMeshElement()->SetLabel(*MeshLabel);
+				}
 
 				FDatasmithConverter Converter;
 
@@ -3004,12 +3140,13 @@ public:
 				TSharedPtr< IDatasmithActorElement > InversedHISMActor;
 				// todo: ExportHierarchicalInstanceStaticMeshActor CustomMeshNode only used for Material - can be simplified, Material anyway is dealt with outside too
 				TSharedRef<IDatasmithActorElement> HismActorElement = ExportHierarchicalInstanceStaticMeshActor(NodeTracker.Node, MeshSource.Node, *MeshLabel, MeshConverted.SupportedChannels,
-					Material, &Transforms, *MeshSource.MeshName, Converter.UnitToCentimeter, EStaticMeshExportMode::Default, InversedHISMActor);
+					Material, &Transforms, MeshConverted.GetDatasmithMeshElement()->GetName(), Converter.UnitToCentimeter, EStaticMeshExportMode::Default, InversedHISMActor);
 				NodeTracker.GetConverted().DatasmithActorElement->AddChild(HismActorElement, EDatasmithActorAttachmentRule::KeepWorldTransform);
 				if (InversedHISMActor)
 				{
 					NodeTracker.GetConverted().DatasmithActorElement->AddChild(InversedHISMActor, EDatasmithActorAttachmentRule::KeepWorldTransform);
 				}
+
 				MeshIndex++;
 			}
 		});
@@ -3382,6 +3519,8 @@ public:
 		LayersForAnimHandle.Reset();
 		NodesPerLayer.Reset();
 
+		Meshes.Reset();
+
 		MaterialsCollectionTracker.Reset();
 		MaterialsAssignedToNodes.Reset();
 
@@ -3411,6 +3550,8 @@ public:
 	TMap<AnimHandle, TUniquePtr<FLayerTracker>> LayersForAnimHandle;
 	TMap<FLayerTracker*, TSet<FNodeTracker*>> NodesPerLayer;
 
+	FMeshes Meshes;
+
 	FMaterialsCollectionTracker MaterialsCollectionTracker;
 	TMap<FMaterialTracker*, TSet<FNodeTracker*>> MaterialsAssignedToNodes;
 
diff --git a/Engine/Source/Programs/Enterprise/Datasmith/DatasmithMaxExporter/Public/DatasmithMaxConverters.h b/Engine/Source/Programs/Enterprise/Datasmith/DatasmithMaxExporter/Public/DatasmithMaxConverters.h
--- a/Engine/Source/Programs/Enterprise/Datasmith/DatasmithMaxExporter/Public/DatasmithMaxConverters.h
+++ b/Engine/Source/Programs/Enterprise/Datasmith/DatasmithMaxExporter/Public/DatasmithMaxConverters.h
@@ -34,6 +34,12 @@ class FMeshConverted
 {
 public:
 	TSharedPtr<IDatasmithMeshElement> DatasmithMeshElement;
+
+	[[nodiscard]] TSharedPtr<IDatasmithMeshElement> GetDatasmithMeshElement() const
+	{
+		return DatasmithMeshElement;
+	}
+
 	TSet<uint16> SupportedChannels;
 	TMap<int32, int32> UVChannelsMap;
 
diff --git a/Engine/Source/Programs/Enterprise/Datasmith/DatasmithMaxExporter/Public/DatasmithMaxDirectLink.h b/Engine/Source/Programs/Enterprise/Datasmith/DatasmithMaxExporter/Public/DatasmithMaxDirectLink.h
--- a/Engine/Source/Programs/Enterprise/Datasmith/DatasmithMaxExporter/Public/DatasmithMaxDirectLink.h
+++ b/Engine/Source/Programs/Enterprise/Datasmith/DatasmithMaxExporter/Public/DatasmithMaxDirectLink.h
@@ -326,6 +326,17 @@ public:
 	bool bConsolidateMaterialIds; 
 
 	GeomUtils::FRenderMeshForConversion CollisionMesh;
+
+	Interval GetValidityInterval()
+	{
+		return RenderMesh.GetValidityInterval() & CollisionMesh.GetValidityInterval();
+	}
+
+	// Whether this will make a mesh
+	bool IsValid()
+	{
+		return RenderMesh.IsValid();
+	}
 };
 
 
@@ -452,8 +463,25 @@ struct MeshConversionParams
 	bool bConsolidateMaterialIds;
 };
 
+/** Holds all the data to export mesh to DatasmithMeshElement */
+class FDatasmithMeshConverter
+{
+public:
+	FMD5Hash ComputeHash();
+
+	FDatasmithMesh RenderMesh;
+	TSet<uint16> SupportedChannels;
+	TMap<int32, int32> UVChannelsMap;
+	int32 SelectedLightmapUVChannel = -1;
+
+	FDatasmithMesh CollisionMesh;
+	bool bHasCollision = false;
+
+	FMeshConverted Converted;
+};
+
 /** Creates Mesh element and converts max mesh into it */
-bool ConvertMaxMeshToDatasmith(TimeValue CurrentTime, ISceneTracker& Scene, FMeshConverterSource& MeshSource, FMeshConverted& MeshConverted);
+bool ConvertMaxMeshToDatasmith(TimeValue CurrentTime, FMeshConverterSource& MeshSource, FDatasmithMeshConverter& MeshConverter);
 
 bool OpenDirectLinkUI();
 const TCHAR* GetDirectlinkCacheDirectory();
@@ -649,7 +677,8 @@ public:
 
 	/** If material is used on a mesh/actor to create a datamsith element */
 	bool IsMaterialUsed(Mtl* Material);
-	/** Set datasmith materials for static mesh for assigned max material */
+	/** Set datasmith materials to Datasmith static mesh according to assigned max material */
+	void SetMaterialsForMeshElement(FMeshConverted& MeshConverted, Mtl* Material);
 	void SetMaterialsForMeshElement(const TSharedPtr<IDatasmithMeshElement>& MeshElement, Mtl* Material, const TSet<uint16>& SupportedChannels);
 	void SetMaterialForMeshElement(const TSharedPtr<IDatasmithMeshElement>& MeshElement, Mtl* MaterialAssignedToNode, Mtl* ActualMaterial, uint16 SlotId);
 	void UnSetMaterialsForMeshElement(const TSharedPtr<IDatasmithMeshElement>& MeshElement);
diff --git a/Engine/Source/Runtime/Datasmith/DatasmithCore/Private/DatasmithMesh.cpp b/Engine/Source/Runtime/Datasmith/DatasmithCore/Private/DatasmithMesh.cpp
--- a/Engine/Source/Runtime/Datasmith/DatasmithCore/Private/DatasmithMesh.cpp
+++ b/Engine/Source/Runtime/Datasmith/DatasmithCore/Private/DatasmithMesh.cpp
@@ -7,6 +7,30 @@
 #include "Containers/Set.h"
 #include "Math/Color.h"
 
+namespace DatasmithMeshImpl
+{
+
+template<typename T>
+static void UpdateMD5SimpleType(FMD5& MD5, const T& Value)
+{
+	static_assert(TIsPODType<T>::Value, "Simple type required");
+	MD5.Update(reinterpret_cast<const uint8*>(&Value), sizeof(Value));
+}
+
+template<typename T>
+static void UpdateMD5Array(FMD5& MD5, TArray<T> Value)
+{
+	static_assert(TIsPODType<T>::Value, "This function requires POD array");
+	UpdateMD5SimpleType(MD5, Value.Num());
+	if (!Value.IsEmpty())
+	{
+		MD5.Update(reinterpret_cast<const uint8*>(Value.GetData()), Value.GetTypeSize()*Value.Num());
+	}
+}
+
+}
+
+
 class FDatasmithMesh::FDatasmithMeshImpl
 {
 public:
@@ -16,6 +40,8 @@ public:
 	bool GetIdInUse( int32 Id ) const;
 	int32 GetMaterialsCount() const { return IdsInUse.Num(); }
 
+	FMD5Hash CalculateHash() const;
+
 	FString Name;
 
 	TArray< FVector3f > Vertices;
@@ -57,6 +83,52 @@ bool FDatasmithMesh::FDatasmithMeshImpl::GetIdInUse( int32 Id ) const
 	return IdsInUse.Contains(Id);
 }
 
+FMD5Hash FDatasmithMesh::FDatasmithMeshImpl::CalculateHash() const
+{
+	FMD5 MD5;
+
+	using namespace DatasmithMeshImpl;
+
+	UpdateMD5Array(MD5, Vertices);
+	UpdateMD5Array(MD5, Indices);
+
+	UpdateMD5Array(MD5, Normals);
+	UpdateMD5Array(MD5, MaterialIndices);
+
+	int32 UVChannelCount = UVs.Num();
+	UpdateMD5SimpleType(MD5, UVChannelCount);
+	for (const TArray<FVector2D>& UVChannel: UVs)
+	{
+		UpdateMD5Array(MD5, UVChannel);
+	}
+
+	check(UVChannelCount == UVIndices.Num()); // UV indices are per-channel
+	for (const TArray<int32>& UVChannelIndices: UVIndices)
+	{
+		UpdateMD5Array(MD5, UVChannelIndices);
+	}
+
+	UpdateMD5Array(MD5, FaceSmoothingMasks);
+	UpdateMD5Array(MD5, IndicesColor);
+	UpdateMD5SimpleType(MD5, LightmapUVChannel);
+
+	TArray<int32> IdsInUseArray = IdsInUse.Array();
+	IdsInUseArray.Sort();
+	UpdateMD5Array(MD5, IdsInUseArray);
+
+	int32 LODCount = LODs.Num();
+	UpdateMD5SimpleType(MD5, LODCount);
+	for (const FDatasmithMesh& LODMesh : LODs)
+	{
+		FMD5Hash LODHash = LODMesh.CalculateHash();
+		MD5.Update(LODHash.GetBytes(), LODHash.GetSize());
+	}
+
+	FMD5Hash Hash;
+	Hash.Set(MD5);
+	return Hash;
+}
+
 FDatasmithMesh::FDatasmithMesh()
 	: Impl( new FDatasmithMeshImpl() )
 {
@@ -95,6 +167,12 @@ FDatasmithMesh& FDatasmithMesh::operator=( FDatasmithMesh&& Other )
 	return *this;
 }
 
+FMD5Hash FDatasmithMesh::CalculateHash() const
+{
+	return Impl->CalculateHash();
+}
+
+
 void FDatasmithMesh::SetName(const TCHAR* InName)
 {
 	Impl->Name = InName;
diff --git a/Engine/Source/Runtime/Datasmith/DatasmithCore/Public/DatasmithMesh.h b/Engine/Source/Runtime/Datasmith/DatasmithCore/Public/DatasmithMesh.h
--- a/Engine/Source/Runtime/Datasmith/DatasmithCore/Public/DatasmithMesh.h
+++ b/Engine/Source/Runtime/Datasmith/DatasmithCore/Public/DatasmithMesh.h
@@ -3,9 +3,11 @@
 
 #include "Math/Box.h"
 #include "Math/Color.h"
-#include "Math/UnrealMathSSE.h"
 #include "Math/Vector2D.h"
 
+#include "Misc/SecureHash.h"
+
+
 class DATASMITHCORE_API FDatasmithMesh
 {
 public:
@@ -18,6 +20,9 @@ public:
 	FDatasmithMesh& operator=( const FDatasmithMesh& Other );
 	FDatasmithMesh& operator=( FDatasmithMesh&& Other );
 
+	// Calculate mesh data hash(doesn't include name)
+	FMD5Hash CalculateHash() const;
+
 	void SetName(const TCHAR* InName);
 	const TCHAR* GetName() const;
 
