[Solidworks] Deleted sub-assemblies are removed from Datasmith scene on sync

diff --git a/Engine/Source/Programs/Enterprise/Datasmith/DatasmithSolidworksExporter/Private/AssemblyDocument.cs b/Engine/Source/Programs/Enterprise/Datasmith/DatasmithSolidworksExporter/Private/AssemblyDocument.cs
--- a/Engine/Source/Programs/Enterprise/Datasmith/DatasmithSolidworksExporter/Private/AssemblyDocument.cs
+++ b/Engine/Source/Programs/Enterprise/Datasmith/DatasmithSolidworksExporter/Private/AssemblyDocument.cs
@@ -113,8 +113,6 @@ namespace DatasmithSolidworks
 
 			public Dictionary<FComponentName, FObjectMaterials> ComponentsMaterialsMap = null;
 
-			public Dictionary<FComponentName, string> ComponentToPartMap = new Dictionary<FComponentName, string>();
-
 			public Dictionary<FComponentName, FConvertedTransform> ComponentsTransformsMap =
 				new Dictionary<FComponentName, FConvertedTransform>();
 
@@ -204,7 +202,6 @@ namespace DatasmithSolidworks
 			foreach (FComponentName CompName in SyncState.ComponentsToDelete)
 			{
 				FActorName ActorName = Exporter.GetComponentActorName(CompName);
-				SyncState.ComponentToPartMap.Remove(CompName);
 				SyncState.CollectedComponentsMap.Remove(CompName);
 
 				ReleaseComponentMeshes(CompName);
@@ -472,8 +469,6 @@ namespace DatasmithSolidworks
 				FPartDocument PartTracker = AddTrackedAssemblyPart(Component, ComponentName);
 				SyncState.PartsMap[PartPath] = PartTracker;
 			}
-
-			SyncState.ComponentToPartMap[ComponentName] = PartPath;
 		}
 
 		// Tracks changes in a Part that is a component in the assembly
@@ -512,7 +507,7 @@ namespace DatasmithSolidworks
 
 		public void ComponentDeleted(FComponentName ComponentName)
 		{
-			if (SyncState.ComponentToPartMap.ContainsKey(ComponentName) &&
+			if (SyncState.CollectedComponentsMap.ContainsKey(ComponentName) &&
 				!SyncState.ComponentsToDelete.Contains(ComponentName))
 			{
 				SyncState.ComponentsToDelete.Add(ComponentName);
