[SketchUp] Updated plugin windows installer with required dependency dll

diff --git a/Engine/Source/Programs/Enterprise/Datasmith/DatasmithSketchUpRubyExporter/DatasmithSketchUpRuby2020.Target.cs b/Engine/Source/Programs/Enterprise/Datasmith/DatasmithSketchUpRubyExporter/DatasmithSketchUpRuby2020.Target.cs
--- a/Engine/Source/Programs/Enterprise/Datasmith/DatasmithSketchUpRubyExporter/DatasmithSketchUpRuby2020.Target.cs
+++ b/Engine/Source/Programs/Enterprise/Datasmith/DatasmithSketchUpRubyExporter/DatasmithSketchUpRuby2020.Target.cs
@@ -57,7 +57,7 @@ public abstract class DatasmithSketchUpRubyBaseTarget : TargetRules
 
 		// Copy support dlls
 		PostBuildSteps.Add(string.Format("echo F|xcopy /Y /R /F \"{0}\" \"{1}\"",
-			string.Format(@"$(EngineDir)/Binaries/Win64/{0}/tbb.dll", ExeBinariesSubFolder),
+			string.Format(@"$(EngineDir)/Binaries/Win64/{0}/tbbmalloc.dll", ExeBinariesSubFolder),
 			string.Format(@"$(EngineDir)/Binaries/Win64/{0}/Plugin/UnrealDatasmithSketchUp", ExeBinariesSubFolder)
 		));
 
