[Datasmith] Geolocation api supports 'unset' value

-  Sketchup doesn't provide elevation so not setting elevation won't have it in .udatasmith file and imported into Unreal as metadata

diff --git a/Engine/Plugins/Enterprise/DatasmithImporter/Source/DatasmithImporter/Private/Utility/DatasmithImporterUtils.cpp b/Engine/Plugins/Enterprise/DatasmithImporter/Source/DatasmithImporter/Private/Utility/DatasmithImporterUtils.cpp
--- a/Engine/Plugins/Enterprise/DatasmithImporter/Source/DatasmithImporter/Private/Utility/DatasmithImporterUtils.cpp
+++ b/Engine/Plugins/Enterprise/DatasmithImporter/Source/DatasmithImporter/Private/Utility/DatasmithImporterUtils.cpp
@@ -198,9 +198,21 @@ ADatasmithSceneActor* FDatasmithImporterUtils::CreateImportSceneActor( FDatasmit
 	}
 
 	FVector Geolocation = ImportContext.Scene->GetGeolocation();
-	UDatasmithAssetUserData::SetDatasmithUserDataValueForKey(SceneActor, TEXT("Geolocation_Latitude"),  FString::SanitizeFloat(Geolocation.X));
-	UDatasmithAssetUserData::SetDatasmithUserDataValueForKey(SceneActor, TEXT("Geolocation_Longitude"),  FString::SanitizeFloat(Geolocation.Y));
-	UDatasmithAssetUserData::SetDatasmithUserDataValueForKey(SceneActor, TEXT("Geolocation_Elevation"),  FString::SanitizeFloat(Geolocation.Z));
+
+	if (Geolocation.X != TNumericLimits<double>::Max())
+	{
+		UDatasmithAssetUserData::SetDatasmithUserDataValueForKey(SceneActor, TEXT("Geolocation_Latitude"),  FString::SanitizeFloat(Geolocation.X));
+	}
+
+	if (Geolocation.Y != TNumericLimits<double>::Max())
+	{
+		UDatasmithAssetUserData::SetDatasmithUserDataValueForKey(SceneActor, TEXT("Geolocation_Longitude"),  FString::SanitizeFloat(Geolocation.Y));
+	}
+
+	if (Geolocation.Z != TNumericLimits<double>::Max())
+	{
+		UDatasmithAssetUserData::SetDatasmithUserDataValueForKey(SceneActor, TEXT("Geolocation_Elevation"),  FString::SanitizeFloat(Geolocation.Z));
+	}
 
 	ImportContext.ActorsContext.ImportSceneActor = SceneActor;
 
diff --git a/Engine/Source/Developer/Datasmith/DatasmithFacade/Private/DatasmithFacadeScene.cpp b/Engine/Source/Developer/Datasmith/DatasmithFacade/Private/DatasmithFacadeScene.cpp
--- a/Engine/Source/Developer/Datasmith/DatasmithFacade/Private/DatasmithFacadeScene.cpp
+++ b/Engine/Source/Developer/Datasmith/DatasmithFacade/Private/DatasmithFacadeScene.cpp
@@ -483,9 +483,19 @@ const TCHAR* FDatasmithFacadeScene::GetAssetsOutputPath() const
 	return SceneExporterRef->GetAssetsOutputPath();
 }
 
-void FDatasmithFacadeScene::SetGeolocation(double Latitude, double Longitude, double Elevation)
+void FDatasmithFacadeScene::SetGeolocationLatitude(double Latitude)
 {
-	SceneRef->SetGeolocation(FVector(Latitude, Longitude, Elevation));
+	SceneRef->SetGeolocationLatitude(Latitude);
+}
+
+void FDatasmithFacadeScene::SetGeolocationLongitude(double Longitude)
+{
+	SceneRef->SetGeolocationLongitude(Longitude);
+}
+
+void FDatasmithFacadeScene::SetGeolocationElevation(double Elevation)
+{
+	SceneRef->SetGeolocationElevation(Elevation);
 }
 
 void FDatasmithFacadeScene::GetGeolocation(double& OutLatitude, double& OutLongitude, double& OutElevation) const
diff --git a/Engine/Source/Developer/Datasmith/DatasmithFacade/Public/DatasmithFacadeScene.h b/Engine/Source/Developer/Datasmith/DatasmithFacade/Public/DatasmithFacadeScene.h
--- a/Engine/Source/Developer/Datasmith/DatasmithFacade/Public/DatasmithFacadeScene.h
+++ b/Engine/Source/Developer/Datasmith/DatasmithFacade/Public/DatasmithFacadeScene.h
@@ -248,12 +248,14 @@ public:
 	/** Get the path were additional assets will be saved to. */
 	const TCHAR* GetAssetsOutputPath() const;
 
-	/** Set the Datasmith scene Geolocation data*/
-	void SetGeolocation(double Latitude, double Longitude, double Elevation);
-
 	/** Get the Datasmith scene Geolocation data*/
 	void GetGeolocation(double& OutLatitude, double& OutLongitude, double& OutElevation) const;
 
+	/** Set the Datasmith scene Geolocation data*/
+	void SetGeolocationLatitude(double Latitude);
+	void SetGeolocationLongitude(double Longitude);
+	void SetGeolocationElevation(double Elevation);
+
 	/** Instantiate an exporter and register export start time */
 	void PreExport();
 
diff --git a/Engine/Source/Programs/Enterprise/Datasmith/DatasmithFacadeCSharp/Private/DatasmithFacadeCSharp.cpp b/Engine/Source/Programs/Enterprise/Datasmith/DatasmithFacadeCSharp/Private/DatasmithFacadeCSharp.cpp
--- a/Engine/Source/Programs/Enterprise/Datasmith/DatasmithFacadeCSharp/Private/DatasmithFacadeCSharp.cpp
+++ b/Engine/Source/Programs/Enterprise/Datasmith/DatasmithFacadeCSharp/Private/DatasmithFacadeCSharp.cpp
@@ -6040,20 +6040,6 @@ SWIGEXPORT wchar_t * SWIGSTDCALL CSharp_FDatasmithFacadeScene_GetAssetsOutputPat
 }
 
 
-SWIGEXPORT void SWIGSTDCALL CSharp_FDatasmithFacadeScene_SetGeolocation(void * jarg1, double jarg2, double jarg3, double jarg4) {
-  FDatasmithFacadeScene *arg1 = (FDatasmithFacadeScene *) 0 ;
-  double arg2 ;
-  double arg3 ;
-  double arg4 ;
-  
-  arg1 = (FDatasmithFacadeScene *)jarg1; 
-  arg2 = (double)jarg2; 
-  arg3 = (double)jarg3; 
-  arg4 = (double)jarg4; 
-  (arg1)->SetGeolocation(arg2,arg3,arg4);
-}
-
-
 SWIGEXPORT void SWIGSTDCALL CSharp_FDatasmithFacadeScene_GetGeolocation(void * jarg1, double * jarg2, double * jarg3, double * jarg4) {
   FDatasmithFacadeScene *arg1 = (FDatasmithFacadeScene *) 0 ;
   double *arg2 = 0 ;
@@ -6068,6 +6054,36 @@ SWIGEXPORT void SWIGSTDCALL CSharp_FDatasmithFacadeScene_GetGeolocation(void * j
 }
 
 
+SWIGEXPORT void SWIGSTDCALL CSharp_FDatasmithFacadeScene_SetGeolocationLatitude(void * jarg1, double jarg2) {
+  FDatasmithFacadeScene *arg1 = (FDatasmithFacadeScene *) 0 ;
+  double arg2 ;
+  
+  arg1 = (FDatasmithFacadeScene *)jarg1; 
+  arg2 = (double)jarg2; 
+  (arg1)->SetGeolocationLatitude(arg2);
+}
+
+
+SWIGEXPORT void SWIGSTDCALL CSharp_FDatasmithFacadeScene_SetGeolocationLongitude(void * jarg1, double jarg2) {
+  FDatasmithFacadeScene *arg1 = (FDatasmithFacadeScene *) 0 ;
+  double arg2 ;
+  
+  arg1 = (FDatasmithFacadeScene *)jarg1; 
+  arg2 = (double)jarg2; 
+  (arg1)->SetGeolocationLongitude(arg2);
+}
+
+
+SWIGEXPORT void SWIGSTDCALL CSharp_FDatasmithFacadeScene_SetGeolocationElevation(void * jarg1, double jarg2) {
+  FDatasmithFacadeScene *arg1 = (FDatasmithFacadeScene *) 0 ;
+  double arg2 ;
+  
+  arg1 = (FDatasmithFacadeScene *)jarg1; 
+  arg2 = (double)jarg2; 
+  (arg1)->SetGeolocationElevation(arg2);
+}
+
+
 SWIGEXPORT void SWIGSTDCALL CSharp_FDatasmithFacadeScene_PreExport(void * jarg1) {
   FDatasmithFacadeScene *arg1 = (FDatasmithFacadeScene *) 0 ;
   
diff --git a/Engine/Source/Programs/Enterprise/Datasmith/DatasmithFacadeCSharp/Public/DatasmithFacadeCSharpPINVOKE.cs b/Engine/Source/Programs/Enterprise/Datasmith/DatasmithFacadeCSharp/Public/DatasmithFacadeCSharpPINVOKE.cs
--- a/Engine/Source/Programs/Enterprise/Datasmith/DatasmithFacadeCSharp/Public/DatasmithFacadeCSharpPINVOKE.cs
+++ b/Engine/Source/Programs/Enterprise/Datasmith/DatasmithFacadeCSharp/Public/DatasmithFacadeCSharpPINVOKE.cs
@@ -1670,12 +1670,18 @@ class DatasmithFacadeCSharpPINVOKE {
   [global::System.Runtime.InteropServices.DllImport("DatasmithFacadeCSharp", EntryPoint="CSharp_FDatasmithFacadeScene_GetAssetsOutputPath")]
   public static extern global::System.IntPtr FDatasmithFacadeScene_GetAssetsOutputPath(global::System.Runtime.InteropServices.HandleRef jarg1);
 
-  [global::System.Runtime.InteropServices.DllImport("DatasmithFacadeCSharp", EntryPoint="CSharp_FDatasmithFacadeScene_SetGeolocation")]
-  public static extern void FDatasmithFacadeScene_SetGeolocation(global::System.Runtime.InteropServices.HandleRef jarg1, double jarg2, double jarg3, double jarg4);
-
   [global::System.Runtime.InteropServices.DllImport("DatasmithFacadeCSharp", EntryPoint="CSharp_FDatasmithFacadeScene_GetGeolocation")]
   public static extern void FDatasmithFacadeScene_GetGeolocation(global::System.Runtime.InteropServices.HandleRef jarg1, out double jarg2, out double jarg3, out double jarg4);
 
+  [global::System.Runtime.InteropServices.DllImport("DatasmithFacadeCSharp", EntryPoint="CSharp_FDatasmithFacadeScene_SetGeolocationLatitude")]
+  public static extern void FDatasmithFacadeScene_SetGeolocationLatitude(global::System.Runtime.InteropServices.HandleRef jarg1, double jarg2);
+
+  [global::System.Runtime.InteropServices.DllImport("DatasmithFacadeCSharp", EntryPoint="CSharp_FDatasmithFacadeScene_SetGeolocationLongitude")]
+  public static extern void FDatasmithFacadeScene_SetGeolocationLongitude(global::System.Runtime.InteropServices.HandleRef jarg1, double jarg2);
+
+  [global::System.Runtime.InteropServices.DllImport("DatasmithFacadeCSharp", EntryPoint="CSharp_FDatasmithFacadeScene_SetGeolocationElevation")]
+  public static extern void FDatasmithFacadeScene_SetGeolocationElevation(global::System.Runtime.InteropServices.HandleRef jarg1, double jarg2);
+
   [global::System.Runtime.InteropServices.DllImport("DatasmithFacadeCSharp", EntryPoint="CSharp_FDatasmithFacadeScene_PreExport")]
   public static extern void FDatasmithFacadeScene_PreExport(global::System.Runtime.InteropServices.HandleRef jarg1);
 
diff --git a/Engine/Source/Programs/Enterprise/Datasmith/DatasmithFacadeCSharp/Public/FDatasmithFacadeScene.cs b/Engine/Source/Programs/Enterprise/Datasmith/DatasmithFacadeCSharp/Public/FDatasmithFacadeScene.cs
--- a/Engine/Source/Programs/Enterprise/Datasmith/DatasmithFacadeCSharp/Public/FDatasmithFacadeScene.cs
+++ b/Engine/Source/Programs/Enterprise/Datasmith/DatasmithFacadeCSharp/Public/FDatasmithFacadeScene.cs
@@ -348,14 +348,22 @@ public class FDatasmithFacadeScene : global::System.IDisposable {
     return ret;
   }
 
-  public void SetGeolocation(double Latitude, double Longitude, double Elevation) {
-    DatasmithFacadeCSharpPINVOKE.FDatasmithFacadeScene_SetGeolocation(swigCPtr, Latitude, Longitude, Elevation);
-  }
-
   public void GetGeolocation(out double OutLatitude, out double OutLongitude, out double OutElevation) {
     DatasmithFacadeCSharpPINVOKE.FDatasmithFacadeScene_GetGeolocation(swigCPtr, out OutLatitude, out OutLongitude, out OutElevation);
   }
 
+  public void SetGeolocationLatitude(double Latitude) {
+    DatasmithFacadeCSharpPINVOKE.FDatasmithFacadeScene_SetGeolocationLatitude(swigCPtr, Latitude);
+  }
+
+  public void SetGeolocationLongitude(double Longitude) {
+    DatasmithFacadeCSharpPINVOKE.FDatasmithFacadeScene_SetGeolocationLongitude(swigCPtr, Longitude);
+  }
+
+  public void SetGeolocationElevation(double Elevation) {
+    DatasmithFacadeCSharpPINVOKE.FDatasmithFacadeScene_SetGeolocationElevation(swigCPtr, Elevation);
+  }
+
   public void PreExport() {
     DatasmithFacadeCSharpPINVOKE.FDatasmithFacadeScene_PreExport(swigCPtr);
   }
diff --git a/Engine/Source/Programs/Enterprise/Datasmith/DatasmithSketchUpRubyExporter/Private/DatasmithSketchUpComponent.cpp b/Engine/Source/Programs/Enterprise/Datasmith/DatasmithSketchUpRubyExporter/Private/DatasmithSketchUpComponent.cpp
--- a/Engine/Source/Programs/Enterprise/Datasmith/DatasmithSketchUpRubyExporter/Private/DatasmithSketchUpComponent.cpp
+++ b/Engine/Source/Programs/Enterprise/Datasmith/DatasmithSketchUpRubyExporter/Private/DatasmithSketchUpComponent.cpp
@@ -387,8 +387,10 @@ SUTransformation FModelDefinition::GetMeshBakedTransform()
 bool FModelDefinition::UpdateModel(FExportContext& Context)
 {
 	// SketchUp API has no notification of Geolocation change so retrieve it every time and compare to check if we need to set DirectLink update
-	FVector Geolocation = Context.DatasmithScene->GetGeolocation();
-	FVector GeolocationNew = Geolocation;
+	FVector GeolocationDatasmith = Context.DatasmithScene->GetGeolocation();
+
+	FVector2d Geolocation(GeolocationDatasmith);
+	FVector2d GeolocationNew = Geolocation;
 
 	SULocationRef Location = SU_INVALID;
 	if (SU_ERROR_NONE == SUModelGetLocation(Model, &Location))
@@ -397,7 +399,7 @@ bool FModelDefinition::UpdateModel(FExportContext& Context)
 		double Longitude = 0;
 		SULocationGetLatLong(Location, &Latitude, &Longitude);
 
-		GeolocationNew = FVector(Latitude, Longitude, 0);
+		GeolocationNew = FVector2d(Latitude, Longitude);
 	}
 
 	if ((Geolocation - GeolocationNew).IsNearlyZero(1e-10))
@@ -405,7 +407,8 @@ bool FModelDefinition::UpdateModel(FExportContext& Context)
 		return false;
 	}
 
-	Context.DatasmithScene->SetGeolocation(GeolocationNew);
+	Context.DatasmithScene->SetGeolocationLatitude(GeolocationNew.X);
+	Context.DatasmithScene->SetGeolocationLongitude(GeolocationNew.Y);
 	return true;
 }
 
diff --git a/Engine/Source/Runtime/Datasmith/DatasmithCore/Private/DatasmithSceneElementsImpl.cpp b/Engine/Source/Runtime/Datasmith/DatasmithCore/Private/DatasmithSceneElementsImpl.cpp
--- a/Engine/Source/Runtime/Datasmith/DatasmithCore/Private/DatasmithSceneElementsImpl.cpp
+++ b/Engine/Source/Runtime/Datasmith/DatasmithCore/Private/DatasmithSceneElementsImpl.cpp
@@ -1086,7 +1086,7 @@ void FDatasmithSceneImpl::Reset()
 	ProductVersion = TEXT("");
 	UserID = TEXT("");
 	UserOS = TEXT("");
-	Geolocation = FVector::ZeroVector;
+	Geolocation = FVector(TNumericLimits<double>::Max(), TNumericLimits<double>::Max(), TNumericLimits<double>::Max());  // max double indicates values are not set
 	ResourcePath = TEXT("");
 
 	ExportDuration = 0;
diff --git a/Engine/Source/Runtime/Datasmith/DatasmithCore/Private/DatasmithSceneElementsImpl.h b/Engine/Source/Runtime/Datasmith/DatasmithCore/Private/DatasmithSceneElementsImpl.h
--- a/Engine/Source/Runtime/Datasmith/DatasmithCore/Private/DatasmithSceneElementsImpl.h
+++ b/Engine/Source/Runtime/Datasmith/DatasmithCore/Private/DatasmithSceneElementsImpl.h
@@ -1610,7 +1610,9 @@ public:
 	virtual void SetUserOS(const TCHAR* InUserOS) override { UserOS = InUserOS; }
 
 	virtual FVector GetGeolocation() const override { return Geolocation; }
-	virtual void SetGeolocation(FVector InGeolocation) override { Geolocation = InGeolocation; }
+	virtual void SetGeolocationLatitude(double InLatitude) override {  Geolocation.Get().X = InLatitude;}
+	virtual void SetGeolocationLongitude(double InLongitude) override {  Geolocation.Get().Y = InLongitude; }
+	virtual void SetGeolocationElevation(double InElevation) override {  Geolocation.Get().Z = InElevation; }
 
 	virtual int32 GetExportDuration() const override { return ExportDuration; }
 	virtual void SetExportDuration(int32 InExportDuration) override { ExportDuration = InExportDuration; }
diff --git a/Engine/Source/Runtime/Datasmith/DatasmithCore/Private/DatasmithSceneXmlReader.cpp b/Engine/Source/Runtime/Datasmith/DatasmithCore/Private/DatasmithSceneXmlReader.cpp
--- a/Engine/Source/Runtime/Datasmith/DatasmithCore/Private/DatasmithSceneXmlReader.cpp
+++ b/Engine/Source/Runtime/Datasmith/DatasmithCore/Private/DatasmithSceneXmlReader.cpp
@@ -1203,10 +1203,23 @@ bool FDatasmithSceneXmlReader::ParseXmlFile(TSharedRef< IDatasmithScene >& OutSc
 		}
 		else if (Nodes[i]->GetTag() == DATASMITH_GEOLOCATION)
 		{
-			OutScene->SetGeolocation(FVector(
-				ValueFromString<double>(Nodes[i]->GetAttribute(DATASMITH_GEOLOCATION_LATITUDE)),
-				ValueFromString<double>(Nodes[i]->GetAttribute(DATASMITH_GEOLOCATION_LONGITUDE)),
-				ValueFromString<double>(Nodes[i]->GetAttribute(DATASMITH_GEOLOCATION_ELEVATION))));
+			FString LatitudeStr = Nodes[i]->GetAttribute(DATASMITH_GEOLOCATION_LATITUDE);
+			if (!LatitudeStr.IsEmpty())
+			{
+				OutScene->SetGeolocationLatitude(ValueFromString<double>(LatitudeStr));
+			}
+
+			FString LongitudeStr = Nodes[i]->GetAttribute(DATASMITH_GEOLOCATION_LONGITUDE);
+			if (!LongitudeStr.IsEmpty())
+			{
+				OutScene->SetGeolocationLongitude(ValueFromString<double>(LongitudeStr));
+			}
+
+			FString ElevationStr = Nodes[i]->GetAttribute(DATASMITH_GEOLOCATION_ELEVATION);
+			if (!ElevationStr.IsEmpty())
+			{
+				OutScene->SetGeolocationElevation(ValueFromString<double>(ElevationStr));
+			}
 		}
 		// STATIC MESHES
 		else if (Nodes[i]->GetTag() == DATASMITH_STATICMESHNAME)
diff --git a/Engine/Source/Runtime/Datasmith/DatasmithCore/Private/DatasmithSceneXmlWriter.cpp b/Engine/Source/Runtime/Datasmith/DatasmithCore/Private/DatasmithSceneXmlWriter.cpp
--- a/Engine/Source/Runtime/Datasmith/DatasmithCore/Private/DatasmithSceneXmlWriter.cpp
+++ b/Engine/Source/Runtime/Datasmith/DatasmithCore/Private/DatasmithSceneXmlWriter.cpp
@@ -1934,16 +1934,27 @@ void FDatasmithSceneXmlWriter::Serialize( TSharedRef< IDatasmithScene > Datasmit
 
 	FDatasmithSceneXmlWriterImpl::WriteIndent(Archive, 1);
 
-	FVector Geolocation = DatasmithScene->GetGeolocation();
-	
-	XmlString = FString::Printf( TEXT("<%s %s=\"%s\" %s=\"%s\" %s=\"%s\"/>"),
-		DATASMITH_GEOLOCATION,
-		DATASMITH_GEOLOCATION_LATITUDE, *FString::SanitizeFloat(Geolocation.X),
-		DATASMITH_GEOLOCATION_LONGITUDE, *FString::SanitizeFloat(Geolocation.Y),
-		DATASMITH_GEOLOCATION_ELEVATION, *FString::SanitizeFloat(Geolocation.Z)
-		) + LINE_TERMINATOR;
+	// Geolocation
+	{
+		FVector Geolocation = DatasmithScene->GetGeolocation();
 
-	FDatasmithSceneXmlWriterImpl::SerializeToArchive(Archive, XmlString);
+		XmlString = FString::Printf( TEXT("<%s"), DATASMITH_GEOLOCATION);
+
+		auto WriteGeolocationComponent = [&XmlString](const TCHAR* AttrName, double Value)
+		{
+			if (Value != TNumericLimits<double>::Max())
+			{
+				XmlString += FString::Printf( TEXT(" %s=\"%s\""), AttrName, *FString::SanitizeFloat(Value));
+			}
+		};
+		WriteGeolocationComponent(DATASMITH_GEOLOCATION_LATITUDE, Geolocation.X);
+		WriteGeolocationComponent(DATASMITH_GEOLOCATION_LONGITUDE, Geolocation.Y);
+		WriteGeolocationComponent(DATASMITH_GEOLOCATION_ELEVATION, Geolocation.Z);
+		XmlString += TEXT("/>");
+		XmlString += LINE_TERMINATOR;
+
+		FDatasmithSceneXmlWriterImpl::SerializeToArchive(Archive, XmlString);
+	}
 
 	FDatasmithSceneXmlWriterImpl::WritePostProcessElement( DatasmithScene->GetPostProcess(), Archive, 1 );
 
diff --git a/Engine/Source/Runtime/Datasmith/DatasmithCore/Public/IDatasmithSceneElements.h b/Engine/Source/Runtime/Datasmith/DatasmithCore/Public/IDatasmithSceneElements.h
--- a/Engine/Source/Runtime/Datasmith/DatasmithCore/Public/IDatasmithSceneElements.h
+++ b/Engine/Source/Runtime/Datasmith/DatasmithCore/Public/IDatasmithSceneElements.h
@@ -1554,11 +1554,14 @@ public:
 	 */
 	virtual void SetUserOS(const TCHAR*) = 0;
 
-	/** Get Geolocation data of the scene. Where X = Latitude, Y = Longitude, Z = Elevation */
+	/** Get Geolocation data of the scene. Where X = Latitude, Y = Longitude, Z = Elevation
+	 *   Components are initialized to TNumericLimits<double>::Max() to indicate they are not "set"
+	 */
 	virtual FVector GetGeolocation() const = 0; 
 
-	/** Set Geolocation data of the scene. Where X = Latitude, Y = Longitude, Z = Elevation */
-	virtual void SetGeolocation(FVector) = 0;
+	virtual void SetGeolocationLatitude(double) = 0;
+	virtual void SetGeolocationLongitude(double) = 0;
+	virtual void SetGeolocationElevation(double) = 0;
 
 	/** Returns the time taken to export the scene */
 	virtual int32 GetExportDuration() const = 0;
