Fixed DatasmithSDKSample/DirectLink Tests - Geolocation components are exported only when set explicitly

diff --git a/Engine/Source/Programs/Enterprise/Datasmith/DatasmithSDKSample/Scenes/DefaultScene.cpp b/Engine/Source/Programs/Enterprise/Datasmith/DatasmithSDKSample/Scenes/DefaultScene.cpp
--- a/Engine/Source/Programs/Enterprise/Datasmith/DatasmithSDKSample/Scenes/DefaultScene.cpp
+++ b/Engine/Source/Programs/Enterprise/Datasmith/DatasmithSDKSample/Scenes/DefaultScene.cpp
@@ -44,6 +44,13 @@ TSharedPtr<IDatasmithScene> FDefaultScene::Export()
 	TSharedRef<IDatasmithScene> DatasmithScene = FDatasmithSceneFactory::CreateScene(*GetName());
 	SetupSharedSceneProperties(DatasmithScene);
 
+	// Set Geolocation to the Scene
+	{
+		DatasmithScene->SetGeolocationLatitude(10);
+		DatasmithScene->SetGeolocationLongitude(20);
+		// Skip setting Elevation to test it's absent
+	}
+
 	// Creates a mesh asset
 	{
 		FDatasmithMesh BaseMesh;
