[Solidworks] Fixed deletion of sub-component of a sub-assembly

diff --git a/Engine/Source/Programs/Enterprise/Datasmith/DatasmithSolidworksExporter/Private/AssemblyDocument.cs b/Engine/Source/Programs/Enterprise/Datasmith/DatasmithSolidworksExporter/Private/AssemblyDocument.cs
--- a/Engine/Source/Programs/Enterprise/Datasmith/DatasmithSolidworksExporter/Private/AssemblyDocument.cs
+++ b/Engine/Source/Programs/Enterprise/Datasmith/DatasmithSolidworksExporter/Private/AssemblyDocument.cs
@@ -50,9 +50,18 @@ namespace DatasmithSolidworks
 			Dictionary<int, FMaterial> MaterialsMap = new Dictionary<int, FMaterial>();
 			Dictionary<FComponentName, FObjectMaterials> CurrentDocMaterialsMap  = new Dictionary<FComponentName, FObjectMaterials>();
 
-			foreach (bool _ in FObjectMaterials.LoadAssemblyMaterialsEnum(AsmDocumentTracker, AllExportedComponents, CurrentDocMaterialsMap, MaterialsMap))
+			List<FComponentName> InvalidComponents = new List<FComponentName>();
+			foreach (bool _ in FObjectMaterials.LoadAssemblyMaterialsEnum(AsmDocumentTracker, AllExportedComponents, CurrentDocMaterialsMap, MaterialsMap, InvalidComponents))
 			{
-				yield return _;
+				yield return true;
+			}
+
+			// Workaround for lack of notifications for some types of component deletion:
+			//   - a component was deleted which is internal to a subassembly of another assembly. As opposed to a component representing subassembly instance in the parent assembly which is notified when deleted.
+			//      in UI: right-click on such a subcomponent, select Delete and dialog should appear whether to delete the whole subassembly or the component in this subassembly. Select deleting just the component
+			foreach (FComponentName ComponentName in InvalidComponents)
+			{
+				AsmDocumentTracker.ComponentDeleted(ComponentName);	
 			}
 
 			HashSet<FComponentName> Components =  SyncState.ComponentsMaterialsMap == null ? new HashSet<FComponentName>() : new HashSet<FComponentName>(SyncState.ComponentsMaterialsMap.Keys);
diff --git a/Engine/Source/Programs/Enterprise/Datasmith/DatasmithSolidworksExporter/Private/ObjectMaterials.cs b/Engine/Source/Programs/Enterprise/Datasmith/DatasmithSolidworksExporter/Private/ObjectMaterials.cs
--- a/Engine/Source/Programs/Enterprise/Datasmith/DatasmithSolidworksExporter/Private/ObjectMaterials.cs
+++ b/Engine/Source/Programs/Enterprise/Datasmith/DatasmithSolidworksExporter/Private/ObjectMaterials.cs
@@ -482,7 +482,7 @@ namespace DatasmithSolidworks
 			return DocMaterials.Count > 0 ? DocMaterials : null;
 		}
 
-		public static IEnumerable<bool> LoadAssemblyMaterialsEnum(FAssemblyDocumentTracker InAsmDoc, HashSet<FComponentName> InComponentsSet, Dictionary<FComponentName, FObjectMaterials> CurrentDocMaterialsMap, Dictionary<int, FMaterial> MaterialsMap)
+		public static IEnumerable<bool> LoadAssemblyMaterialsEnum(FAssemblyDocumentTracker InAsmDoc, HashSet<FComponentName> InComponentsSet, Dictionary<FComponentName, FObjectMaterials> CurrentDocMaterialsMap, Dictionary<int, FMaterial> MaterialsMap, List<FComponentName> InvalidComponents)
 		{
 			FObjectMaterials.LoadMaterialsWithUsers(InAsmDoc, swDisplayStateOpts_e.swThisDisplayState, null,
 				CurrentDocMaterialsMap, MaterialsMap);
@@ -490,13 +490,18 @@ namespace DatasmithSolidworks
 
 			foreach (FComponentName CompName in InComponentsSet.OrderBy(Name => Name.GetString()))
 			{
-				FObjectMaterials.UpdateComponentMaterials(InAsmDoc, CompName, CurrentDocMaterialsMap,
+				bool bIsComponentValid = FObjectMaterials.UpdateComponentMaterials(InAsmDoc, CompName, CurrentDocMaterialsMap,
 					MaterialsMap);
+				if (!bIsComponentValid)
+				{
+					InvalidComponents.Add(CompName);
+				}
 				yield return true;
 			}
 		}
 
-		public static void UpdateComponentMaterials(FAssemblyDocumentTracker InAsmDoc, FComponentName CompName,
+		/// Returns whether component was found valid for export(not suppressed/not deleted)
+		public static bool UpdateComponentMaterials(FAssemblyDocumentTracker InAsmDoc, FComponentName CompName,
 			Dictionary<FComponentName, FObjectMaterials> DocMaterials, Dictionary<int, FMaterial> MaterialsMap)
 		{
 			LogDebug($"Component '{CompName}' - adding DocMaterials");
@@ -508,6 +513,13 @@ namespace DatasmithSolidworks
 				{
 					LogDebug($"Adding to DocMaterials");
 
+					if (Comp.GetModelDoc2() == null)
+					{
+						// Component's model doc might be null if component is suppressed/lightweight (in which case we treat it as hidden)
+						// or has already been deleted
+						return false;
+					}
+
 					LogIndent();
 					FObjectMaterials ComponentMaterials = FObjectMaterials.LoadComponentMaterials(InAsmDoc,
 						Comp, swDisplayStateOpts_e.swThisDisplayState, null, MaterialsMap);
@@ -542,6 +554,8 @@ namespace DatasmithSolidworks
 			}
 
 			LogDedent();
+
+			return true;
 		}
 
 		// Parse materials within document and return components using them
