[Navisworks] Tooltips are now working when docking/undocking plugin UI. Export tooltip added.

diff --git a/Engine/Source/Programs/Enterprise/Datasmith/DatasmithNavisworksExporter/DatasmithNavisworksPlugin/DockPaneControl.cs b/Engine/Source/Programs/Enterprise/Datasmith/DatasmithNavisworksExporter/DatasmithNavisworksPlugin/DockPaneControl.cs
--- a/Engine/Source/Programs/Enterprise/Datasmith/DatasmithNavisworksExporter/DatasmithNavisworksPlugin/DockPaneControl.cs
+++ b/Engine/Source/Programs/Enterprise/Datasmith/DatasmithNavisworksExporter/DatasmithNavisworksPlugin/DockPaneControl.cs
@@ -75,18 +75,41 @@ namespace DatasmithNavisworks
 			const string MetadataTooltip = "Include metadata export(might affect export time)";
 			toolTip1.SetToolTip(MetadataCheckBox, MetadataTooltip);
 
-			Presenter.OnLoaded();
+			const string ExportTooltip = "Export the current 3D model to Datasmith file";
+			toolTip1.SetToolTip(ExportButton, ExportTooltip);
+
+			toolTip1.ShowAlways = true;
+
+			Presenter?.OnLoaded();
 		}
 
 		private void ExportButton_Click(object Sender, EventArgs EventArgs)
 		{
-			Presenter.OnExportClicked();
+			Presenter?.OnExportClicked();
 		}
 
 		private void PickButton_Click(object Sender, EventArgs EventArgs)
 		{
-			Presenter.OnPickOriginClicked();
+			Presenter?.OnPickOriginClicked();
+		}
+
+		protected override void OnEnter(EventArgs EventArgs)
+		{
+			base.OnEnter(EventArgs);
+			RepairTooltips();
 		}
 
+		protected override void OnValidated(EventArgs EventArgs)
+		{
+			base.OnValidated(EventArgs);
+			RepairTooltips();
+		}
+
+		private void RepairTooltips()
+		{
+			// Weird hack to repair disappeared tooltips  when docking/undocking plugin Panel in Navisworks
+			// Seems like nothing else works but to show some (empty in this case)tooltip on both Enter and Validated events
+			toolTip1.Show("", tableLayoutPanel1, 0, 0, 1);
+		}
 	}
 }
