[Rhino] Fixed export crash when exporting a model

- exportimg meshes with multiple uv sets

diff --git a/Engine/Source/Programs/Enterprise/Datasmith/DatasmithRhinoExporter/Private/ElementExporters/DatasmithRhinoMeshExporter.cs b/Engine/Source/Programs/Enterprise/Datasmith/DatasmithRhinoExporter/Private/ElementExporters/DatasmithRhinoMeshExporter.cs
--- a/Engine/Source/Programs/Enterprise/Datasmith/DatasmithRhinoExporter/Private/ElementExporters/DatasmithRhinoMeshExporter.cs
+++ b/Engine/Source/Programs/Enterprise/Datasmith/DatasmithRhinoExporter/Private/ElementExporters/DatasmithRhinoMeshExporter.cs
@@ -92,11 +92,35 @@ namespace DatasmithRhino.ElementExporters
 
 		private static void ParseMesh(FDatasmithFacadeMeshElement DatasmithMeshElement, FDatasmithFacadeMesh DatasmithMesh, DatasmithMeshInfo MeshInfo, List<DatasmithMaterialInfo> MaterialInfos)
 		{
+			List<Mesh> MeshSections = MeshInfo.RhinoMeshes;
+
+			// UVs need to be fixed first, before parsing topology
+			// Since calling SetTextureCoordinates may sometimes re-tessellate mesh changing vertices count!
+			for (int MeshIndex = 0; MeshIndex < MeshSections.Count; ++MeshIndex)
+			{
+				Mesh RhinoMesh = MeshSections[MeshIndex];
+
+				if (MeshInfo.TextureMappings.Count > 0)
+				{
+					foreach (DatasmithTextureMappingData TextureMappingData in MeshInfo.TextureMappings)
+					{
+						// Since we are offsetting the meshes before exporting them, we must apply the same correction to the UV transform.
+						Transform InverveOffsetTranform;
+						MeshInfo.OffsetTransform.TryGetInverse(out InverveOffsetTranform);
+						Transform CorrectedTransform = InverveOffsetTranform * TextureMappingData.ObjectTransform;
+
+						// Rhino gives no guarantee on the state of the texture mapping in a given mesh.
+						// We must make sure that UV is set to the channel we are exporting.
+						const bool bLazyLoad = false;
+						RhinoMesh.SetTextureCoordinates(TextureMappingData.RhinoTextureMapping, CorrectedTransform, bLazyLoad);
+					}
+				}
+			}
+
 			int VertexIndexOffset = 0;
 			int FaceIndexOffset = 0;
 			int UVIndexOffset = 0;
 			int NumberOfUVChannels = System.Math.Max(1, MeshInfo.TextureMappings.Count);
-			List<Mesh> MeshSections = MeshInfo.RhinoMeshes;
 			List<DatasmithMaterialInfo> UniqueMaterialInfo = new List<DatasmithMaterialInfo>();
 			InitializeDatasmithMesh(DatasmithMesh, MeshSections, NumberOfUVChannels);
 
@@ -179,17 +203,6 @@ namespace DatasmithRhino.ElementExporters
 				{
 					for (int UVChannel = 0; UVChannel < MeshInfo.TextureMappings.Count; ++UVChannel)
 					{
-						// Since we are offsetting the meshes before exporting them, we must apply the same correction to the UV transform.
-						Transform InverveOffsetTranform;
-						MeshInfo.OffsetTransform.TryGetInverse(out InverveOffsetTranform);
-						DatasmithTextureMappingData TextureMappingData = MeshInfo.TextureMappings[UVChannel];
-						Transform CorrectedTransform = InverveOffsetTranform * TextureMappingData.ObjectTransform;
-
-						// Rhino gives no guarantee on the state of the texture mapping in a given mesh.
-						// We must make sure that UV is set to the channel we are exporting.
-						const bool bLazyLoad = false;
-						RhinoMesh.SetTextureCoordinates(TextureMappingData.RhinoTextureMapping, CorrectedTransform, bLazyLoad);
-
 						// Add the UV coordinates to the current channel.
 						AddUVsToMesh(DatasmithMesh, RhinoMesh, UVChannel, UVIndexOffset);
 					}
