[Revit] Datasmith Exporter for Revit 2018.3 - Direct Link error

diff --git a/Engine/Source/Programs/Enterprise/Datasmith/DatasmithRevitExporter/Private/DatasmithRevitDirectLink.cs b/Engine/Source/Programs/Enterprise/Datasmith/DatasmithRevitExporter/Private/DatasmithRevitDirectLink.cs
--- a/Engine/Source/Programs/Enterprise/Datasmith/DatasmithRevitExporter/Private/DatasmithRevitDirectLink.cs
+++ b/Engine/Source/Programs/Enterprise/Datasmith/DatasmithRevitExporter/Private/DatasmithRevitDirectLink.cs
@@ -314,8 +314,7 @@ namespace DatasmithRevitExporter
 						if (ElementInView.Category != null && (BuiltInCategory)ElementInView.Category.Id.IntegerValue != BuiltInCategory.OST_Levels)
 #endif
 						{
-							IList<ElementId> DependentElements = ElementInView.GetDependentElements(null);
-							foreach (ElementId DependentElementId in DependentElements)
+							foreach (ElementId DependentElementId in FUtils.GetAllDependentElements(ElementInView))
 							{
 								if (DependentElementId == ElemId)
 								{
diff --git a/Engine/Source/Programs/Enterprise/Datasmith/DatasmithRevitExporter/Private/DatasmithRevitDocumentData.cs b/Engine/Source/Programs/Enterprise/Datasmith/DatasmithRevitExporter/Private/DatasmithRevitDocumentData.cs
--- a/Engine/Source/Programs/Enterprise/Datasmith/DatasmithRevitExporter/Private/DatasmithRevitDocumentData.cs
+++ b/Engine/Source/Programs/Enterprise/Datasmith/DatasmithRevitExporter/Private/DatasmithRevitDocumentData.cs
@@ -1387,8 +1387,7 @@ namespace DatasmithRevitExporter
 			{
 				//Check if any of its children is Decal:
 				//If so track them, so that we can use the owner object elementId to update the Decals (when modified, for ep: changing its location)
-				IList<ElementId> DependentElements = InElement.GetDependentElements(null);
-				foreach (ElementId DependentElementId in DependentElements)
+				foreach (ElementId DependentElementId in FUtils.GetAllDependentElements(InElement))
 				{
 					if (FUtils.IsElementIdDecal(InElement.Document, DependentElementId))
 					{
@@ -1462,14 +1461,6 @@ namespace DatasmithRevitExporter
 						ElementData.BaseElementType = InElement.Document.GetElement(InElement.GetTypeId()) as ElementType;
 						ElementData.CurrentElement = InElement;
 						ElementData.MeshMaterialsMap.Clear();
-
-						if (ElementData.bIsDecalElement)
-						{
-							if (!DecalElementsMap.ContainsKey(InElement.Id))
-							{
-								DecalElementsMap.Add(InElement.Id, ElementData);
-							}
-						}
 					}
 					else
 					{
@@ -1480,11 +1471,6 @@ namespace DatasmithRevitExporter
 				else
 				{
 					ElementData = new FElementData(InElement, InWorldTransform, this);
-
-					if (ElementData.bIsDecalElement)
-					{
-						DecalElementsMap.Add(InElement.Id, ElementData);
-					}
 				}
 
 				ElementDataStack.Push(ElementData);
@@ -1493,6 +1479,15 @@ namespace DatasmithRevitExporter
 				// GetActorName depends on this 
 				ElementData.InitializePivotPlacement(ref InWorldTransform);
 				ElementData.InitializeElement(InWorldTransform, ElementData);
+
+				if (ElementData.bIsDecalElement)
+				{
+					if (!DecalElementsMap.ContainsKey(InElement.Id))
+					{
+						DecalElementsMap.Add(InElement.Id, ElementData);
+					}
+				}
+
 			}
 			else
 			{
diff --git a/Engine/Source/Programs/Enterprise/Datasmith/DatasmithRevitExporter/Private/DatasmithRevitUtils.cs b/Engine/Source/Programs/Enterprise/Datasmith/DatasmithRevitExporter/Private/DatasmithRevitUtils.cs
--- a/Engine/Source/Programs/Enterprise/Datasmith/DatasmithRevitExporter/Private/DatasmithRevitUtils.cs
+++ b/Engine/Source/Programs/Enterprise/Datasmith/DatasmithRevitExporter/Private/DatasmithRevitUtils.cs
@@ -212,5 +212,11 @@ namespace DatasmithRevitExporter
 
 			return false;
 		}
+
+		public static IList<ElementId> GetAllDependentElements(Element ElementInView)
+		{
+			// Revit 2018 doesn't support passing 'null' to GetDependentElements so just make a simple filter
+			return ElementInView.GetDependentElements(new ExclusionFilter(new List<ElementId>(){ElementInView.Id}));
+		}
 	}
 }
