[3dsMax] Exporter build fix after switch to c++20 for UE

diff --git a/Engine/Source/Programs/Enterprise/Datasmith/DatasmithMaxExporter/DatasmithMax2017.Target.cs b/Engine/Source/Programs/Enterprise/Datasmith/DatasmithMaxExporter/DatasmithMax2017.Target.cs
--- a/Engine/Source/Programs/Enterprise/Datasmith/DatasmithMaxExporter/DatasmithMax2017.Target.cs
+++ b/Engine/Source/Programs/Enterprise/Datasmith/DatasmithMaxExporter/DatasmithMax2017.Target.cs
@@ -33,6 +33,32 @@ public abstract class DatasmithMaxBaseTarget : TargetRules
 
 		GlobalDefinitions.Add("UE_EXTERNAL_PROFILING_ENABLED=0"); // For DirectLinkUI (see FDatasmithExporterManager::FInitOptions)
 
+		// Setting CppStandard to Cpp17 as new CppStandardVersion.Default is Cpp20 and this implies "/permissive-" for MSVC which
+		// doesn't work well with 3ds Max includes(lots of non-conformance like const to non-const string conversion, temp to &, etc)
+		// 
+		// For some reason, adding following doesn't help 100%. Even though docs says that should override "/permissive-"'s settings like C2102:
+		// AdditionalCompilerArguments = " /Zc:referenceBinding- /Zc:strictStrings- /Zc:rvalueCast- ";
+		// 
+		// 3dsMax uses MSVC extension which allows taking address of a class rvalue returned from a function:
+		// class V
+		// {
+		// 	int dummy;
+		// };
+		//
+		// class C {
+		// public:
+		// 	V get() const { return V(); }
+		// };
+		// void f(V* ){}
+		// int main() {
+		// 	C c;
+		// 	f(&(c.get()));// In conformance mode C2102 '&' requires l-value, but MSVC allows this by default
+		// } 
+		// This behavior is disabled when c++20 enabled and is not repaired with any /Zc option. Seems like it only works with "/permissive" enabled. But...
+		// Reverting permissive breaks Engine compilation because of operator resolution ambiguity(which doesn't happen with stricter checks) and other issues:
+		// AdditionalCompilerArguments = "  /permissive ";
+		CppStandard = CppStandardVersion.Cpp17;
+
 		// todo: remove?
 		// bSupportEditAndContinue = true;
 	}
