=> the control of the import time cannot be activated as the size of all the files are not know...

diff --git a/Engine/Source/Programs/Enterprise/Datasmith/DatasmithCADWorker/Private/DatasmithCADWorkerImpl.cpp b/Engine/Source/Programs/Enterprise/Datasmith/DatasmithCADWorker/Private/DatasmithCADWorkerImpl.cpp
--- a/Engine/Source/Programs/Enterprise/Datasmith/DatasmithCADWorker/Private/DatasmithCADWorkerImpl.cpp
+++ b/Engine/Source/Programs/Enterprise/Datasmith/DatasmithCADWorker/Private/DatasmithCADWorkerImpl.cpp
@@ -123,7 +123,7 @@ void FDatasmithCADWorkerImpl::ProcessCommand(const FImportParametersCommand& Imp
 	ImportParameters = ImportParametersCommand.ImportParameters;
 }
 
-uint64 DefineMaximumAllowedDuration(const CADLibrary::FFileDescriptor& FileDescriptor)
+uint64 DefineMaximumAllowedDuration(const CADLibrary::FFileDescriptor& FileDescriptor, bool& bEnableTimeControl)
 {
 	FFileStatData FileStatData = IFileManager::Get().GetStatData(*FileDescriptor.GetSourcePath());
 	double MaxTimePerMb = 5e-6;
@@ -132,6 +132,11 @@ uint64 DefineMaximumAllowedDuration(const CADLibrary::FFileDescriptor& FileDescr
 	CADLibrary::ECADFormat Format = FileDescriptor.GetFileFormat();
 	switch (Format)
 	{
+	case CADLibrary::ECADFormat::JT:
+	case CADLibrary::ECADFormat::INVENTOR:
+		MaxTimePerMb = 1.;
+		bEnableTimeControl = false;
+		break;
 	case CADLibrary::ECADFormat::SOLIDWORKS:
 	case CADLibrary::ECADFormat::CATIA_3DXML:
 		MaxTimePerMb = 1e-5;
@@ -162,9 +167,15 @@ void FDatasmithCADWorkerImpl::ProcessCommand(const FRunTaskCommand& RunTaskComma
 	FCompletedTaskCommand CompletedTask;
 
 	bProcessIsRunning = true;
-	int64 MaxDuration = DefineMaximumAllowedDuration(FileToProcess);
+
+	bool bEnableTimeControl = CADLibrary::FImportParameters::bGEnableTimeControl;
+	int64 MaxDuration = DefineMaximumAllowedDuration(FileToProcess, bEnableTimeControl);
+
 	TArray<UE::Tasks::FTask> Checkers;
-	Checkers.Emplace(UE::Tasks::Launch(TEXT("TimeChecker"), [&FileToProcess, &MaxDuration]() { CheckDuration(FileToProcess, MaxDuration); }));
+	if(bEnableTimeControl)
+	{
+		Checkers.Emplace(UE::Tasks::Launch(TEXT("TimeChecker"), [&FileToProcess, &MaxDuration]() { CheckDuration(FileToProcess, MaxDuration); }));
+	}
 	Checkers.Emplace(UE::Tasks::Launch(TEXT("MemoryChecker"), []() { CheckMemory(); }));
 
 	FImportParameters FileImporParameters(ImportParameters, RunTaskCommand.Mesher);
@@ -216,11 +227,6 @@ void FDatasmithCADWorkerImpl::ProcessCommand(const FRunTaskCommand& RunTaskComma
 
 void FDatasmithCADWorkerImpl::CheckDuration(const CADLibrary::FFileDescriptor& FileToProcess, const int64 MaxDuration)
 {
-	if (!CADLibrary::FImportParameters::bGEnableTimeControl)
-	{
-		return;
-	}
-
 	const uint64 StartTime = FPlatformTime::Cycles64();
 	const uint64 MaxCycles = MaxDuration / FPlatformTime::GetSecondsPerCycle64() + StartTime;
 
