Fixed build break following changes from CL 27235710

diff --git a/Engine/Source/Programs/Enterprise/Datasmith/DatasmithSDK/DatasmithSDK.Target.cs b/Engine/Source/Programs/Enterprise/Datasmith/DatasmithSDK/DatasmithSDK.Target.cs
--- a/Engine/Source/Programs/Enterprise/Datasmith/DatasmithSDK/DatasmithSDK.Target.cs
+++ b/Engine/Source/Programs/Enterprise/Datasmith/DatasmithSDK/DatasmithSDK.Target.cs
@@ -80,7 +80,7 @@ public class DatasmithSDKTarget : TargetRules
 	}
 	public void PostBuildCopyWin64(string SrcPath, string DestPath, string Files)
 	{
-		PostBuildSteps.Add(string.Format("echo Copying {0}\\{2} to {1}", SrcPath, DestPath, Files));
+		PostBuildSteps.Add(string.Format("echo Copying {0}{2} to {1}", SrcPath, DestPath, Files));
 		PostBuildSteps.Add(string.Format(@"call $(EngineDir)\Source\Programs\Enterprise\Datasmith\DatasmithSDK\CopyCmd.bat {0} {1} {2}", SrcPath, DestPath, Files));
 	}
 
@@ -146,6 +146,7 @@ public class DatasmithSDKTarget : TargetRules
 			(@"$(EngineDir)\Source\Runtime\Datasmith\DirectLink\Public\", Path.Combine(DstDatasmithSDK,"Public\\"), "*.h"),
 			(@"$(EngineDir)\Source\Developer\Datasmith\DatasmithExporter\Public\", Path.Combine(DstDatasmithSDK,"Public\\"), "*.h"),
 			(@"$(EngineDir)\Source\Developer\Datasmith\DatasmithExporterUI\Public\", Path.Combine(DstDatasmithSDK,"Public\\"), "*.h"),
+			(@"$(EngineDir)\Intermediate\Build\Win64\DatasmithSDK\Inc\DatasmithCore\UHT\", Path.Combine(DstDatasmithSDK,"Public\\"), "*.generated.h"),
 			(@"$(EngineDir)\Extras\VisualStudioDebugging\", DstDatasmithSDK, "Unreal.natvis"),
 			(@"$(EngineDir)\Source\Runtime\TraceLog\Public\", Path.Combine(DstDatasmithSDK,"Private\\"), "*.h"),
 			(@"$(EngineDir)\Source\Runtime\TraceLog\Public\", Path.Combine(DstDatasmithSDK,"Private\\"), "*.inl"),
