[Datasmith] Change DatasmithSDKSample solution targetting to C++20 and VS2022.

This fixes compilation issues.

diff --git a/Engine/Source/Programs/Enterprise/Datasmith/DatasmithSDKSample/README.txt b/Engine/Source/Programs/Enterprise/Datasmith/DatasmithSDKSample/README.txt
--- a/Engine/Source/Programs/Enterprise/Datasmith/DatasmithSDKSample/README.txt
+++ b/Engine/Source/Programs/Enterprise/Datasmith/DatasmithSDKSample/README.txt
@@ -6,4 +6,4 @@ In this project, we use the Datasmith SDK to generate some content, expose this
 
 Dependency
 - In order to build this project, a Datasmith SDK must be provided. The default configuration searches for an SDK generated by the main Unreal solution (see project DatasmithSDK), but you could edit DatasmithSDK.props to point to another SDK.
-- Unreal requires a c++17 compiler
\ No newline at end of file
+- Unreal requires a c++20 compiler
\ No newline at end of file
