[Datasmith Sketchup] Fix build warning on Mac.

diff --git a/Engine/Source/Programs/Enterprise/Datasmith/DatasmithSketchUpRubyExporter/Build/build_mac.py b/Engine/Source/Programs/Enterprise/Datasmith/DatasmithSketchUpRubyExporter/Build/build_mac.py
--- a/Engine/Source/Programs/Enterprise/Datasmith/DatasmithSketchUpRubyExporter/Build/build_mac.py
+++ b/Engine/Source/Programs/Enterprise/Datasmith/DatasmithSketchUpRubyExporter/Build/build_mac.py
@@ -37,7 +37,7 @@ SU_VERSION = args.sketchup_version
 TARGET_ARCHS = args.target_archs;
 
 arch_target_dict = {
-    'x86_64': 'x86_64-apple-macos10.15',
+    'x86_64': 'x86_64-apple-macos11',
     'arm64': 'arm64-apple-macos11'
 }
 
diff --git a/Engine/Source/Programs/Enterprise/Datasmith/DatasmithSketchUpRubyExporter/Build/build_mac.sh b/Engine/Source/Programs/Enterprise/Datasmith/DatasmithSketchUpRubyExporter/Build/build_mac.sh
--- a/Engine/Source/Programs/Enterprise/Datasmith/DatasmithSketchUpRubyExporter/Build/build_mac.sh
+++ b/Engine/Source/Programs/Enterprise/Datasmith/DatasmithSketchUpRubyExporter/Build/build_mac.sh
@@ -46,9 +46,9 @@ SetupIntermediateDirectory() {
     chmod 777 "$Dylibs/$dylibLibTbb"
     chmod 777 "$Dylibs/$dylibLibTbbMalloc"
 
-    install_name_tool -id @loader_path/Dylibs/DatasmithSDK.dylib "$Dylibs/DatasmithSDK.dylib"
-    install_name_tool -id @loader_path/Dylibs/$dylibLibFreeImage "$Dylibs/$dylibLibFreeImage"
-    install_name_tool -change @rpath/$dylibLibFreeImage @loader_path/$dylibLibFreeImage "$Dylibs/DatasmithSDK.dylib"
+    install_name_tool -id @loader_path/Dylibs/DatasmithSDK.dylib "$Dylibs/DatasmithSDK.dylib" > /dev/null 2>&1
+    install_name_tool -id @loader_path/Dylibs/$dylibLibFreeImage "$Dylibs/$dylibLibFreeImage" > /dev/null 2>&1
+    install_name_tool -change @rpath/$dylibLibFreeImage @loader_path/$dylibLibFreeImage "$Dylibs/DatasmithSDK.dylib" > /dev/null 2>&1
 }
 
 # Technicaly our dylibs are already Universal Arm64+x86_64 binaries, so they shouldn't really go into the x86_64 folder,
