Fixed build break for SketchUp exporter on the Mac

#rnx

diff --git a/Engine/Source/Programs/Enterprise/Datasmith/DatasmithSketchUpRubyExporter/Build/build_mac.py b/Engine/Source/Programs/Enterprise/Datasmith/DatasmithSketchUpRubyExporter/Build/build_mac.py
--- a/Engine/Source/Programs/Enterprise/Datasmith/DatasmithSketchUpRubyExporter/Build/build_mac.py
+++ b/Engine/Source/Programs/Enterprise/Datasmith/DatasmithSketchUpRubyExporter/Build/build_mac.py
@@ -123,7 +123,7 @@ UE_Core = f'{UE_Runtime}/Core/Public'
 UE_Project = f'{UE_Runtime}/Projects/Public'
 
 UE_ThirdParty = f'{UE_Engine}/Source/ThirdParty'
-UE_Imath = f'{UE_ThirdParty}/Imath/Deploy/Imath-3.1.3'
+UE_Imath = f'{UE_ThirdParty}/Imath/Deploy/Imath-3.1.9'
 
 UE_INCLUDES_PATH = f'''"{UE_Core}" "{UE_Core}/Internationalization" "{UE_Core}/Async" "{UE_Core}/Containers" "{UE_Core}/Delegates" "{UE_Core}/GenericPlatform" "{UE_Core}/HAL" "{UE_Core}/Logging" "{UE_Core}/Math" "{UE_Core}/Misc" "{UE_Core}/Modules" "{UE_Core}/Modules/Boilerplate" "{UE_Core}/ProfilingDebugging" "{UE_Core}/Serialization" "{UE_Core}/Serialization/Csv" "{UE_Core}/Stats" "{UE_Core}/Templates" "{UE_Core}/UObject" "{UE_Runtime}/CoreUObject/Public" "{UE_Project}/Interfaces" "{UE_Runtime}/TraceLog/Public"  "{UE_Runtime}/Datasmith/DatasmithCore/Public" "{UE_Runtime}/Datasmith/DirectLink/Public" "{UE_Runtime}/Messaging/Public" "{UE_Runtime}/Launch/Resources" "{UE_Engine}/Source/Developer/Datasmith/DatasmithExporter/Public" "{UE_Engine}/Source/Developer/Datasmith/DatasmithExporterUI/Public"'''
 
