Added support for ArchiCAD 27 - bis

diff --git a/Engine/Source/Programs/Enterprise/Datasmith/DatasmithARCHICADExporter/Build/Config/Notes.txt b/Engine/Source/Programs/Enterprise/Datasmith/DatasmithARCHICADExporter/Build/Config/Notes.txt
--- a/Engine/Source/Programs/Enterprise/Datasmith/DatasmithARCHICADExporter/Build/Config/Notes.txt
+++ b/Engine/Source/Programs/Enterprise/Datasmith/DatasmithARCHICADExporter/Build/Config/Notes.txt
@@ -9,6 +9,8 @@ To build you need to define some user environment variables
 			UE_SDKs_Root/HostWin64/Win64/Archicad/23/Support/...
 			UE_SDKs_Root/HostWin64/Win64/Archicad/24/Support/...
 			UE_SDKs_Root/HostWin64/Win64/Archicad/25/Support/...
+			UE_SDKs_Root/HostWin64/Win64/Archicad/26/Support/...
+			UE_SDKs_Root/HostWin64/Win64/Archicad/27/Support/...
 	UE_DatasmithArchicadDeveloperId='Organization (UserId)'
 		Released by Apple "Developer ID Certification Authority"
 			Must be valid for Application and for Installer
diff --git a/Engine/Source/Programs/Enterprise/Datasmith/DatasmithARCHICADExporter/Build/Config/UE_AC.xcconfig b/Engine/Source/Programs/Enterprise/Datasmith/DatasmithARCHICADExporter/Build/Config/UE_AC.xcconfig
--- a/Engine/Source/Programs/Enterprise/Datasmith/DatasmithARCHICADExporter/Build/Config/UE_AC.xcconfig
+++ b/Engine/Source/Programs/Enterprise/Datasmith/DatasmithARCHICADExporter/Build/Config/UE_AC.xcconfig
@@ -2,7 +2,7 @@
 
 #include? "SDKsRoot.xcconfig"
 
-AC_SDK_PATH = $(UESDKRoot)/HostMac/Mac/Archicad/$(AC_VERSION)
+AC_SDK_PATH = $(UESDKRoot)/HostMac/Mac/Archicad/$(AC_SDK_DIR)
 
 AC_INCLUDES_PATH = "$(AC_SDK_PATH)/Support/Inc" "$(AC_SDK_PATH)/Support/Modules/GSRoot/STL/imp" "$(AC_SDK_PATH)/Support/Modules/GSRoot" "$(AC_SDK_PATH)/Support/Modules/InputOutput"/** "$(AC_SDK_PATH)/Support/Modules/GSModelDevLib"/** "$(AC_SDK_PATH)/Support/Modules/GSUtils"/** "$(AC_SDK_PATH)/Support/Modules/GSModeler"/** "$(AC_SDK_PATH)/Support/Modules/Model3D"/** "$(AC_SDK_PATH)/Support/Modules/MPIInterfaceDevLib"/** "$(AC_SDK_PATH)/Support/Modules/Geometry"/** "$(AC_SDK_PATH)/Support/Modules/DGLib"/** "$(AC_SDK_PATH)/Support/Modules/GDL"/** "$(AC_SDK_PATH)/Support/Modules/GXImage"/** "$(AC_SDK_PATH)/Support/Modules/GX"/** "$(AC_SDK_PATH)/Support/Modules/DGGraphix"/** "$(AC_SDK_PATH)/Support/Modules/Graphix"/** "$(AC_SDK_PATH)/Support/Modules/TextEngine"/** "$(AC_SDK_PATH)/Support/Modules/GXImageBase"/** "$(AC_SDK_PATH)/Support/Modules/PointCloudManager" "$(AC_SDK_PATH)/Support/Modules/PointCloud" "$(AC_SDK_PATH)/Support/Modules/Pattern" "$(AC_SDK_PATH)/Support/Modules/Brep" "$(AC_SDK_PATH)/Support/Modules/VBAttributes" "$(AC_SDK_PATH)/Support/Modules/GSModeler2D" "$(AC_SDK_PATH)/Support/Modules/ObjectDatabase" "$(AC_SDK_PATH)/Support/Modules/UDLib" "$(AC_SDK_PATH)/Support/Modules/RS" "$(UE_Engine)/Binaries/Mac/DatasmithSDK"
 
diff --git a/Engine/Source/Programs/Enterprise/Datasmith/DatasmithARCHICADExporter/Build/Config/UE_AC_ARM.xcconfig b/Engine/Source/Programs/Enterprise/Datasmith/DatasmithARCHICADExporter/Build/Config/UE_AC_ARM.xcconfig
--- a/Engine/Source/Programs/Enterprise/Datasmith/DatasmithARCHICADExporter/Build/Config/UE_AC_ARM.xcconfig
+++ b/Engine/Source/Programs/Enterprise/Datasmith/DatasmithARCHICADExporter/Build/Config/UE_AC_ARM.xcconfig
@@ -2,9 +2,9 @@
 
 #include? "SDKsRoot.xcconfig"
 
-AC_SDK_PATH = $(UESDKRoot)/HostMac/Mac/Archicad/$(AC_VERSION)
+AC_SDK_PATH = $(UESDKRoot)/HostMac/Mac/Archicad/$(AC_SDK_DIR)
 
-AC_INCLUDES_PATH = "$(AC_SDK_PATH)/Support/Inc" "$(AC_SDK_PATH)/Support/Modules/GSRoot/STL/imp" "$(AC_SDK_PATH)/Support/Modules/GSRoot" "$(AC_SDK_PATH)/Support/Modules/InputOutput"/** "$(AC_SDK_PATH)/Support/Modules/GSModelDevLib"/** "$(AC_SDK_PATH)/Support/Modules/GSUtils"/** "$(AC_SDK_PATH)/Support/Modules/GSModeler"/** "$(AC_SDK_PATH)/Support/Modules/Model3D"/** "$(AC_SDK_PATH)/Support/Modules/MPIInterfaceDevLib"/** "$(AC_SDK_PATH)/Support/Modules/Geometry"/** "$(AC_SDK_PATH)/Support/Modules/DGLib"/** "$(AC_SDK_PATH)/Support/Modules/GDL"/** "$(AC_SDK_PATH)/Support/Modules/GXImage"/** "$(AC_SDK_PATH)/Support/Modules/GX"/** "$(AC_SDK_PATH)/Support/Modules/DGGraphix"/** "$(AC_SDK_PATH)/Support/Modules/Graphix"/** "$(AC_SDK_PATH)/Support/Modules/TextEngine"/** "$(AC_SDK_PATH)/Support/Modules/GXImageBase"/** "$(AC_SDK_PATH)/Support/Modules/PointCloudManager" "$(AC_SDK_PATH)/Support/Modules/PointCloud" "$(AC_SDK_PATH)/Support/Modules/Pattern" "$(AC_SDK_PATH)/Support/Modules/Brep" "$(AC_SDK_PATH)/Support/Modules/VBAttributes" "$(AC_SDK_PATH)/Support/Modules/GSModeler2D" "$(AC_SDK_PATH)/Support/Modules/ObjectDatabase" "$(AC_SDK_PATH)/Support/Modules/UDLib" "$(AC_SDK_PATH)/Support/Modules/RS" "$(UE_Engine)/Binaries/Mac/DatasmithSDK"
+AC_INCLUDES_PATH = "$(AC_SDK_PATH)/Support/Inc" "$(AC_SDK_PATH)/Support/Modules/GSRoot/STL/imp" "$(AC_SDK_PATH)/Support/Modules/GSRoot" "$(AC_SDK_PATH)/Support/Modules/InputOutput"/** "$(AC_SDK_PATH)/Support/Modules/GSModelDevLib"/** "$(AC_SDK_PATH)/Support/Modules/GSUtils"/** "$(AC_SDK_PATH)/Support/Modules/GSModeler"/** "$(AC_SDK_PATH)/Support/Modules/Model3D"/** "$(AC_SDK_PATH)/Support/Modules/MPIInterfaceDevLib"/** "$(AC_SDK_PATH)/Support/Modules/Geometry"/** "$(AC_SDK_PATH)/Support/Modules/DGLib"/** "$(AC_SDK_PATH)/Support/Modules/GDL"/** "$(AC_SDK_PATH)/Support/Modules/GXImage"/** "$(AC_SDK_PATH)/Support/Modules/GX"/** "$(AC_SDK_PATH)/Support/Modules/DGGraphix"/** "$(AC_SDK_PATH)/Support/Modules/Graphix"/** "$(AC_SDK_PATH)/Support/Modules/TextEngine"/** "$(AC_SDK_PATH)/Support/Modules/GXImageBase"/** "$(AC_SDK_PATH)/Support/Modules/PointCloudManager" "$(AC_SDK_PATH)/Support/Modules/PointCloud" "$(AC_SDK_PATH)/Support/Modules/Pattern" "$(AC_SDK_PATH)/Support/Modules/CADInfrastructureBase" "$(AC_SDK_PATH)/Support/Modules/Brep" "$(AC_SDK_PATH)/Support/Modules/VBAttributes" "$(AC_SDK_PATH)/Support/Modules/GSModeler2D" "$(AC_SDK_PATH)/Support/Modules/ObjectDatabase" "$(AC_SDK_PATH)/Support/Modules/UDLib" "$(AC_SDK_PATH)/Support/Modules/RS" "$(UE_Engine)/Binaries/Mac/DatasmithSDK"
 
 UE_Runtime = $(UE_Engine)/Source/Runtime
 UE_Core = $(UE_Runtime)/Core/Public
@@ -27,7 +27,7 @@ CurrentYear = 2022
 
 FRAMEWORK_SEARCH_PATHS = $(FRAMEWORK_SEARCH_PATHS) $(AC_SDK_PATH)/Support/Frameworks
 
-LIBRARY_SEARCH_PATHS = "Dylibs" "$(AC_SDK_PATH)/Support/Lib/Mactel"
+LIBRARY_SEARCH_PATHS = "Dylibs" "$(AC_SDK_PATH)/Support/Lib/Mactel" "$(AC_SDK_PATH)/Support/Lib"
 
 REZ_SEARCH_PATHS = RO "$(AC_SDK_PATH)/Support/Inc" "$(AC_SDK_PATH)/Support/Modules/DGLib"
 
diff --git a/Engine/Source/Programs/Enterprise/Datasmith/DatasmithARCHICADExporter/Build/Config/app_to_bundle.sh b/Engine/Source/Programs/Enterprise/Datasmith/DatasmithARCHICADExporter/Build/Config/app_to_bundle.sh
--- a/Engine/Source/Programs/Enterprise/Datasmith/DatasmithARCHICADExporter/Build/Config/app_to_bundle.sh
+++ b/Engine/Source/Programs/Enterprise/Datasmith/DatasmithARCHICADExporter/Build/Config/app_to_bundle.sh
@@ -12,4 +12,5 @@ mv DatasmithARCHICAD23Exporter.app DatasmithARCHICAD23Exporter.bundle
 mv DatasmithARCHICAD24Exporter.app DatasmithARCHICAD24Exporter.bundle
 mv DatasmithARCHICAD25Exporter.app DatasmithARCHICAD25Exporter.bundle
 mv DatasmithARCHICAD26Exporter.app DatasmithARCHICAD26Exporter.bundle
+mv DatasmithARCHICAD27Exporter.app DatasmithARCHICAD27Exporter.bundle
 popd
diff --git a/Engine/Source/Programs/Enterprise/Datasmith/DatasmithARCHICADExporter/Build/Config/bundle_to_app.sh b/Engine/Source/Programs/Enterprise/Datasmith/DatasmithARCHICADExporter/Build/Config/bundle_to_app.sh
--- a/Engine/Source/Programs/Enterprise/Datasmith/DatasmithARCHICADExporter/Build/Config/bundle_to_app.sh
+++ b/Engine/Source/Programs/Enterprise/Datasmith/DatasmithARCHICADExporter/Build/Config/bundle_to_app.sh
@@ -12,4 +12,5 @@ mv DatasmithARCHICAD23Exporter.bundle DatasmithARCHICAD23Exporter.app
 mv DatasmithARCHICAD24Exporter.bundle DatasmithARCHICAD24Exporter.app
 mv DatasmithARCHICAD25Exporter.bundle DatasmithARCHICAD25Exporter.app
 mv DatasmithARCHICAD26Exporter.bundle DatasmithARCHICAD26Exporter.app
+mv DatasmithARCHICAD27Exporter.bundle DatasmithARCHICAD27Exporter.app
 popd
diff --git a/Engine/Source/Programs/Enterprise/Datasmith/DatasmithARCHICADExporter/Build/Config/setupdylibs.sh b/Engine/Source/Programs/Enterprise/Datasmith/DatasmithARCHICADExporter/Build/Config/setupdylibs.sh
--- a/Engine/Source/Programs/Enterprise/Datasmith/DatasmithARCHICADExporter/Build/Config/setupdylibs.sh
+++ b/Engine/Source/Programs/Enterprise/Datasmith/DatasmithARCHICADExporter/Build/Config/setupdylibs.sh
@@ -24,8 +24,8 @@ fi
 
 # Remove ArchiCAD resource tool from quarantine
 pushd "$UE_SDKS_ROOT/HostMac/Mac/Archicad"
-	chmod 777 23/Support/Tools/OSX/ResConv
-	xattr -r -d com.apple.quarantine 23/Support/Tools/OSX/ResConv
+	chmod 777 23.1/Support/Tools/OSX/ResConv
+	xattr -r -d com.apple.quarantine 23.1/Support/Tools/OSX/ResConv
 
 	chmod 777 24/Support/Tools/OSX/ResConv
 	xattr -r -d com.apple.quarantine 24/Support/Tools/OSX/ResConv
@@ -35,6 +35,9 @@ pushd "$UE_SDKS_ROOT/HostMac/Mac/Archicad"
 
 	chmod 777 26/Support/Tools/OSX/ResConv
 	xattr -r -d com.apple.quarantine 26/Support/Tools/OSX/ResConv
+
+	chmod 777 27/Support/Tools/OSX/ResConv
+	xattr -r -d com.apple.quarantine 27/Support/Tools/OSX/ResConv
 popd
 
 OurDylibFolder=$projectPath/Dylibs
diff --git a/Engine/Source/Programs/Enterprise/Datasmith/DatasmithARCHICADExporter/Installer/CustomActionManaged/CustomAction.cs b/Engine/Source/Programs/Enterprise/Datasmith/DatasmithARCHICADExporter/Installer/CustomActionManaged/CustomAction.cs
--- a/Engine/Source/Programs/Enterprise/Datasmith/DatasmithARCHICADExporter/Installer/CustomActionManaged/CustomAction.cs
+++ b/Engine/Source/Programs/Enterprise/Datasmith/DatasmithARCHICADExporter/Installer/CustomActionManaged/CustomAction.cs
@@ -122,7 +122,7 @@ namespace CustomActionManaged
 		{
 			try
 			{
-				string[] MajorVersions = new string[] { "26", "25", "24", "23", };
+				string[] MajorVersions = new string[] { "27", "26", "25", "24", "23", };
 
 				string[] MinorVersions = new string[] { "0", };
 
diff --git a/Engine/Source/Programs/Enterprise/Datasmith/DatasmithARCHICADExporter/Private/ConnectionWindow.cpp b/Engine/Source/Programs/Enterprise/Datasmith/DatasmithARCHICADExporter/Private/ConnectionWindow.cpp
--- a/Engine/Source/Programs/Enterprise/Datasmith/DatasmithARCHICADExporter/Private/ConnectionWindow.cpp
+++ b/Engine/Source/Programs/Enterprise/Datasmith/DatasmithARCHICADExporter/Private/ConnectionWindow.cpp
@@ -422,13 +422,13 @@ FConnectionDialog::FConnectionDialog()
         {
             return 0;
         }
-        GS::Ref<GS::Object> Object(ConnectionsListBox.GetItemObjectData(EventArgs.listItem));
-        if (Object == GS::null)
+        const GS::Object* Object = ConnectionsListBox.GetItemObjectData(EventArgs.listItem);
+        if (Object == nullptr)
         {
             UE_AC_DebugF("ConnectionsListBox.onToolTipRequested - Object is null\n");
             return 0;
         }
-        EventArgs.toolTipText += static_cast<FToolTipText&>(*Object).ToolTipText;
+        EventArgs.toolTipText += static_cast<const FToolTipText&>(*Object).ToolTipText;
         return 1;
     };
 #endif
diff --git a/Engine/Source/Programs/Enterprise/Datasmith/DatasmithARCHICADExporter/Private/MaterialsDatabase.cpp b/Engine/Source/Programs/Enterprise/Datasmith/DatasmithARCHICADExporter/Private/MaterialsDatabase.cpp
--- a/Engine/Source/Programs/Enterprise/Datasmith/DatasmithARCHICADExporter/Private/MaterialsDatabase.cpp
+++ b/Engine/Source/Programs/Enterprise/Datasmith/DatasmithARCHICADExporter/Private/MaterialsDatabase.cpp
@@ -13,6 +13,9 @@
 #include "DatasmithSceneFactory.h"
 
 BEGIN_NAMESPACE_UE_AC
+#if AC_VERSION > 26
+const API_AttributeIndex kInvalidMaterialIndex = ACAPI_CreateAttributeIndex(0);
+#endif
 
 class FMaterialInfo
 {
@@ -392,7 +395,11 @@ const FMaterialsDatabase::FMaterialSyncData& FMaterialsDatabase::GetMaterial(con
 	if (inACTextureIndex < kInvalidMaterialIndex)
 	{
 		UE_AC_DebugF("FMaterialsDatabase::GetMaterial - Invalid texture index (%d)\n", inACTextureIndex);
+#if AC_VERSION > 26
+		inACTextureIndex = kInvalidMaterialIndex.ToInt32_Deprecated();
+#else
 		inACTextureIndex = kInvalidMaterialIndex;
+#endif
 	}
 	FMaterialKey MaterialKey(inACMaterialIndex, inACTextureIndex, InSided);
 
diff --git a/Engine/Source/Programs/Enterprise/Datasmith/DatasmithARCHICADExporter/Private/MaterialsDatabase.h b/Engine/Source/Programs/Enterprise/Datasmith/DatasmithARCHICADExporter/Private/MaterialsDatabase.h
--- a/Engine/Source/Programs/Enterprise/Datasmith/DatasmithARCHICADExporter/Private/MaterialsDatabase.h
+++ b/Engine/Source/Programs/Enterprise/Datasmith/DatasmithARCHICADExporter/Private/MaterialsDatabase.h
@@ -18,10 +18,26 @@ enum ESided : bool
 	kDoubleSide = true
 };
 
+#if AC_VERSION > 26
+extern const API_AttributeIndex kInvalidMaterialIndex;
+inline bool operator != (GS::Int32 V1, const API_AttributeIndex& V2)
+{
+	return V1 != V2.ToInt32_Deprecated();
+}
+inline bool operator <= (GS::Int32 V1, const API_AttributeIndex& V2)
+{
+	return V1 <= V2.ToInt32_Deprecated();
+}
+inline bool operator < (GS::Int32 V1, const API_AttributeIndex& V2)
+{
+	return V1 < V2.ToInt32_Deprecated();
+}
+#else
 enum : int32
 {
 	kInvalidMaterialIndex = 0
 };
+#endif
 
 class FMaterialKey
 {
diff --git a/Engine/Source/Programs/Enterprise/Datasmith/DatasmithARCHICADExporter/Private/ProjectEvent.cpp b/Engine/Source/Programs/Enterprise/Datasmith/DatasmithARCHICADExporter/Private/ProjectEvent.cpp
--- a/Engine/Source/Programs/Enterprise/Datasmith/DatasmithARCHICADExporter/Private/ProjectEvent.cpp
+++ b/Engine/Source/Programs/Enterprise/Datasmith/DatasmithARCHICADExporter/Private/ProjectEvent.cpp
@@ -18,11 +18,12 @@ FAssValueName::SAssValueName TAssEnumName< API_NotifyEventID >::AssEnumName[] =
 	ValueName(APINotify_Close),
 	ValueName(APINotify_Quit),
 	ValueName(APINotify_TempSave),
+#if AC_VERSION < 27
 	ValueName(APINotify_ConvertUnId),
 	ValueName(APINotify_ConvertGuid),
 	ValueName(APINotify_ConvertLinkId),
 	ValueName(APINotify_ConvertDrwGuid),
-
+#endif
 	ValueName(APINotify_SendChanges),
 	ValueName(APINotify_ReceiveChanges),
 
diff --git a/Engine/Source/Programs/Enterprise/Datasmith/DatasmithARCHICADExporter/Private/SyncData.cpp b/Engine/Source/Programs/Enterprise/Datasmith/DatasmithARCHICADExporter/Private/SyncData.cpp
--- a/Engine/Source/Programs/Enterprise/Datasmith/DatasmithARCHICADExporter/Private/SyncData.cpp
+++ b/Engine/Source/Programs/Enterprise/Datasmith/DatasmithARCHICADExporter/Private/SyncData.cpp
@@ -505,7 +505,11 @@ void FSyncData::FLayer::Process(FProcessInfo* /* IOProcessInfo */)
 		API_Attribute attribute;
 		Zap(&attribute);
 		attribute.header.typeID = API_LayerID;
+#if AC_VERSION > 26
+		attribute.header.index = ACAPI_CreateAttributeIndex(LayerIndex);
+#else
 		attribute.header.index = short(LayerIndex);
+#endif
 		attribute.header.uniStringNamePtr = &LayerName;
 		GSErrCode error = ACAPI_Attribute_Get(&attribute);
 		if (error != NoError)
@@ -807,8 +811,12 @@ void FSyncData::FElement::Process(FProcessInfo* IOProcessInfo)
 			ActorElement->SetRotation(FGeometryUtil::GetRotationQuat(LocalToWorld.matrix));
 
 			// Set actor layer
-			ActorElement->SetLayer(
-				*IOProcessInfo->SyncContext.GetSyncDatabase().GetLayerName(IOProcessInfo->ElementID.GetHeader().layer));
+#if AC_VERSION > 26
+			const short Index = short(IOProcessInfo->ElementID.GetHeader().layer.ToInt32_Deprecated());
+			ActorElement->SetLayer(*IOProcessInfo->SyncContext.GetSyncDatabase().GetLayerName(Index));
+#else
+			ActorElement->SetLayer(*IOProcessInfo->SyncContext.GetSyncDatabase().GetLayerName(IOProcessInfo->ElementID.GetHeader().layer));
+#endif
 
 			bMetadataProcessed = false;
 			if (IOProcessInfo->bProcessMetaData)
@@ -818,8 +826,11 @@ void FSyncData::FElement::Process(FProcessInfo* IOProcessInfo)
 
 			FMeshClass* MeshClass = IOProcessInfo->ElementID.GetMeshClass();
 			UE_AC_Assert(MeshClass != nullptr);
+			constexpr short IsRelative = short(TR_DET_1 | TR_TRANSL_ONLY);
 			if (MeshClass->AddInstance(this, &IOProcessInfo->SyncContext.GetSyncDatabase()) == FMeshClass::kBuild)
 			{
+				MeshClass->Translation = ActorElement->GetTranslation();
+				MeshClass->Rotation = ActorElement->GetRotation();
 				FConvertGeometry2MeshElement* ConvertGeometry2MeshElement =
 					new FConvertGeometry2MeshElement(IOProcessInfo->SyncContext, this, MeshClass);
 				ConvertGeometry2MeshElement->AddElementGeometry(&IOProcessInfo->ElementID, WorldToLocal);
@@ -836,6 +847,43 @@ void FSyncData::FElement::Process(FProcessInfo* IOProcessInfo)
 				// delete ConvertGeometry2MeshElement;
 				FTaskMgr::GetMgr()->AddTask(ConvertGeometry2MeshElement, FTaskMgr::kSchedule);
 			}
+			else if(TypeID == API_MorphID && ((LocalToWorld.status & IsRelative) == IsRelative))
+			{
+				FTransform Parent(MeshClass->Rotation, MeshClass->Translation);
+				FTransform Child =  Parent * ActorElement->GetRelativeTransform();
+				ActorElement->SetTranslation(Child.GetTranslation());
+				ActorElement->SetRotation(Child.GetRotation());
+
+				//API_Element APIElement;
+				//Zap(&APIElement);
+				//APIElement.header.guid = GSGuid2APIGuid(ElementId);
+				//GSErrCode GSErr = ACAPI_Element_Get(&APIElement, 0);
+				//UE_AC_Assert(GSErr == NoError);
+				//UE_AC_Assert(APIElement.header.typeID == TypeID);
+
+				{
+					//API_Elem_Head	elemHead;
+					//BNZeroMemory(&elemHead, sizeof(API_Elem_Head));
+					//elemHead.guid = GSGuid2APIGuid(ElementId);
+
+					API_ElemInfo3D info3D;
+					BNZeroMemory(&info3D, sizeof(API_ElemInfo3D));
+
+					GSErrCode err = ACAPI_Element_Get3DInfo(IOProcessInfo->ElementID.GetHeader(), &info3D);
+					if (err == NoError)
+					{
+						//API_Component3D component;
+						//BNZeroMemory(&component, sizeof(component));
+						//component.header.typeID = API_BodyID;
+						//component.header.index = info3D.fbody;
+						//err = ACAPI_3D_GetComponent(&component);
+						//if (err == NoError) {
+						//	API_Tranmat			tm;
+						//	tm = component.body.tranmat;
+						//}
+					}
+				}
+			}
 		}
 	}
 }
@@ -1447,8 +1495,12 @@ void FSyncData::FHotLinkNode::Process(FProcessInfo* IOProcessInfo)
 	{
 		SetActorElement(FDatasmithSceneFactory::CreateActor(GSStringToUE(ElementId.ToUniString())));
 
+#if AC_VERSION < 26
 		API_HotlinkNode hotlinkNode;
 		Zap(&hotlinkNode);
+#else
+		API_HotlinkNode hotlinkNode = {0};
+#endif
 		hotlinkNode.guid = GSGuid2APIGuid(ElementId);
 		GSErrCode err = ACAPI_Database(APIDb_GetHotlinkNodeID, &hotlinkNode);
 		if (err == NoError)
diff --git a/Engine/Source/Programs/Enterprise/Datasmith/DatasmithARCHICADExporter/Private/SyncDatabase.cpp b/Engine/Source/Programs/Enterprise/Datasmith/DatasmithARCHICADExporter/Private/SyncDatabase.cpp
--- a/Engine/Source/Programs/Enterprise/Datasmith/DatasmithARCHICADExporter/Private/SyncDatabase.cpp
+++ b/Engine/Source/Programs/Enterprise/Datasmith/DatasmithARCHICADExporter/Private/SyncDatabase.cpp
@@ -405,7 +405,11 @@ const FString& FSyncDatabase::GetLayerName(short InLayerIndex)
 	FString* Found = LayerIndex2Name.Find(InLayerIndex);
 	if (Found == nullptr)
 	{
+#if AC_VERSION > 26
+		LayerIndex2Name.Add(InLayerIndex, GSStringToUE(UE_AC::GetLayerName(ACAPI_CreateAttributeIndex(InLayerIndex))));
+#else
 		LayerIndex2Name.Add(InLayerIndex, GSStringToUE(UE_AC::GetLayerName(InLayerIndex)));
+#endif
 		Found = LayerIndex2Name.Find(InLayerIndex);
 		UE_AC_TestPtr(Found);
 	}
@@ -789,7 +793,7 @@ void FSyncDatabase::ScanLights(FElementID& InElementID)
 					SyncData->SetParent(InElementID.GetSyncData());
 					SyncData->MarkAsModified();
 				}
-				FSyncData::FLight& LightSyncData = static_cast< FSyncData::FLight& >(*SyncData);
+				FSyncData::FLight& LightSyncData = static_cast<FSyncData::FLight&>(*SyncData);
 				LightSyncData.MarkAsExisting();
 
 				LightSyncData.SetLightData(LightData);
@@ -802,7 +806,7 @@ void FSyncDatabase::ScanLights(FElementID& InElementID)
 				{
 					SignaledTypes.insert(LightData.LightType);
 					UE_AC_DebugF("FSyncDatabase::ScanLights - Unhandled LightType=%d, %s", LightData.LightType,
-								 InElementID.GetElementName());
+						InElementID.GetElementName());
 				}
 			}
 		}
diff --git a/Engine/Source/Programs/Enterprise/Datasmith/DatasmithARCHICADExporter/Private/SyncDatabase.h b/Engine/Source/Programs/Enterprise/Datasmith/DatasmithARCHICADExporter/Private/SyncDatabase.h
--- a/Engine/Source/Programs/Enterprise/Datasmith/DatasmithARCHICADExporter/Private/SyncDatabase.h
+++ b/Engine/Source/Programs/Enterprise/Datasmith/DatasmithARCHICADExporter/Private/SyncDatabase.h
@@ -92,6 +92,9 @@ class FMeshClass
 	// Set Mesh Element to the sync data
 	void SetWaitingInstanceMesh(FSyncDatabase* IOSyncDatabase);
 
+	FVector Translation;
+	FQuat Rotation;
+
   private:
 	typedef TList< FSyncData* > FSyncDataList;
 	FSyncDataList*				HeadInstances;
@@ -187,6 +190,12 @@ class FSyncDatabase
 
 	// Return layer sync data (Create it if not present)
 	FSyncData& GetLayerSyncData(short InLayer);
+#if AC_VERSION > 26
+	FSyncData& GetLayerSyncData(const API_AttributeIndex& InLayer)
+	{
+		return GetLayerSyncData(short(InLayer.ToInt32_Deprecated()));
+	}
+#endif
 
 	// Delete obsolete syncdata (and it's Datasmith Element)
 	void DeleteSyncData(const GS::Guid& InGuid);
diff --git a/Engine/Source/Programs/Enterprise/Datasmith/DatasmithARCHICADExporter/Private/Synchronizer.cpp b/Engine/Source/Programs/Enterprise/Datasmith/DatasmithARCHICADExporter/Private/Synchronizer.cpp
--- a/Engine/Source/Programs/Enterprise/Datasmith/DatasmithARCHICADExporter/Private/Synchronizer.cpp
+++ b/Engine/Source/Programs/Enterprise/Datasmith/DatasmithARCHICADExporter/Private/Synchronizer.cpp
@@ -594,7 +594,9 @@ bool FSynchronizer::NeedAutoSyncUpdate() const
 void FSynchronizer::GetProjectPathAndName(GS::UniString* OutPath, GS::UniString* OutName)
 {
 	API_ProjectInfo ProjectInfo;
+#if AC_VERSION < 26
 	Zap(&ProjectInfo);
+#endif
 	GSErrCode GSErr = ACAPI_Environment(APIEnv_ProjectID, &ProjectInfo);
 	if (GSErr == NoError)
 	{
@@ -655,7 +657,7 @@ void FSynchronizer::DumpScene(const TSharedRef< IDatasmithScene >& InScene)
 	// If we change scene, we delete and recreate the folder
 	static int	   NbDumps = 0;
 	static FString PreviousFolderPath;
-	if (FolderPath != PreviousFolderPath)
+	if (!FolderPath.Equals(PreviousFolderPath))
 	{
 		NbDumps = 0;
 		PreviousFolderPath = FolderPath;
diff --git a/Engine/Source/Programs/Enterprise/Datasmith/DatasmithARCHICADExporter/Private/Utils/APIEnvir.h b/Engine/Source/Programs/Enterprise/Datasmith/DatasmithARCHICADExporter/Private/Utils/APIEnvir.h
--- a/Engine/Source/Programs/Enterprise/Datasmith/DatasmithARCHICADExporter/Private/Utils/APIEnvir.h
+++ b/Engine/Source/Programs/Enterprise/Datasmith/DatasmithARCHICADExporter/Private/Utils/APIEnvir.h
@@ -71,6 +71,43 @@ _Pragma("clang diagnostic ignored \"-Wdefaulted-function-deleted\"")
 
 #include "DGFileDlg.hpp"
 
+#if AC_VERSION > 26
+// Declarations of methods used by the AddOn and defined in ACAPI_MigrationHeader.h
+GSErrCode ACAPI_3D_CreateSight (void** newSightPtr);
+GSErrCode ACAPI_3D_GetComponent (API_Component3D* component);
+GSErrCode ACAPI_3D_GetCurrentWindowSight (void** sightPtr);
+GSErrCode ACAPI_3D_SelectSight (void* sight, void** oldSightPtr);
+GSErrCode ACAPI_Command_CallFromEventLoop (const API_ModulID* mdid, GSType cmdID, Int32 cmdVersion, GSHandle paramsHandle, bool silentMode, APICommandCallBackProc* callbackProc);
+GSErrCode ACAPI_Database (API_DatabaseID code, void* par1 = nullptr, void* par2 = nullptr, void* par3 = nullptr);
+GSErrCode ACAPI_Database_GetLast3DDefLevels (const GS::Array<API_ElemType>& elemTypes, GS::Array<double>& levels);
+GSErrCode ACAPI_ElemComponent_GetPropertyDefinitions (const API_ElemComponentID& elemComponent, API_PropertyDefinitionFilter filter, GS::Array<API_PropertyDefinition>& propertyDefinitions);
+GSErrCode ACAPI_ElemComponent_GetPropertyValues (const API_ElemComponentID& elemComponent, const GS::Array<API_PropertyDefinition>& propertyDefinitions, GS::Array<API_Property>& properties);
+GSErrCode ACAPI_ElementGroup_GetGroup (const API_Guid& elemGuid, API_Guid* groupGuid);
+GSErrCode ACAPI_Element_Get3DInfo (const API_Elem_Head& elemHead, API_ElemInfo3D* info3D);
+GSErrCode ACAPI_Element_GetCategoryValue (const API_Guid& elemGuid, const API_ElemCategory& elemCategory, API_ElemCategoryValue* elemCategoryValue);
+GSErrCode ACAPI_Element_GetConnectedElements (const API_Guid& guid, const API_ElemType& connectedElemType, GS::Array<API_Guid>* connectedElements, API_ElemFilterFlags filterBits = APIFilt_None, const API_Guid& renovationFilterGuid = APINULLGuid);
+GSErrCode ACAPI_Environment (API_EnvironmentID code, void* par1 = nullptr, void* par2 = nullptr, void* par3 = nullptr);
+GSErrCode ACAPI_Goodies (API_GoodiesID code, void* par1 = nullptr, void* par2 = nullptr, void* par3 = nullptr, void* par4 = nullptr);
+GSErrCode ACAPI_Install_FileTypeHandler3D (GSType cmdID, APIIO3DCommandProc* handlerProc);
+GSErrCode ACAPI_Install_MenuHandler (short menuStrResID, APIMenuCommandProc* handlerProc);
+GSErrCode ACAPI_Install_ModulCommandHandler (GSType cmdID, Int32 cmdVersion, APIModulCommandProc* handlerProc);
+GSErrCode ACAPI_Interface (API_InterfaceID code, void* par1 = nullptr, void* par2 = nullptr, void* par3 = nullptr, void* par4 = nullptr, void* par5 = nullptr);
+GSErrCode ACAPI_LibPart_Get (API_LibPart* libPart);
+GSErrCode ACAPI_LibPart_GetNum (Int32* count);
+GSErrCode ACAPI_LibPart_GetSection (Int32 libInd, API_LibPartSection* section, GSHandle* sectionHdl, GS::UniString* sectionStr, GS::UniString* password = nullptr);
+GSErrCode ACAPI_LibPart_GetSectionList (Int32 libInd, Int32* nSection, API_LibPartSection*** sections);
+GSErrCode ACAPI_LibPart_Search (API_LibPart* ancestor, bool createIfMissing, bool onlyPlaceable = false);
+GSErrCode ACAPI_ListData_Get (API_ListData* listdata);
+GSErrCode ACAPI_ListData_GetLocal (Int32 libIndex, const API_Elem_Head* elemHead, API_ListData* listdata);
+GSErrCode ACAPI_Notify_CatchNewElement (const API_ToolBoxItem* elemType, APIElementEventHandlerProc* handlerProc);
+GSErrCode ACAPI_Notify_CatchProjectEvent (GSFlags eventTypes, APIProjectEventHandlerProc* handlerProc);
+GSErrCode ACAPI_Notify_CatchViewEvent (GSFlags eventTypes, API_NavigatorMapID mapId, APIViewEventHandlerProc* handlerProc);
+GSErrCode ACAPI_Notify_InstallElementObserver (APIElementEventHandlerProc* handlerProc);
+GSErrCode ACAPI_Register_FileType (Int32 refCon, GSType ftype, GSType fcreator, const char* extname, short iconResID, short descStrResID, short descStrResItemID, API_IOMethod methodFlags);
+GSErrCode ACAPI_Register_Menu (short menuStrResID, short promptStrResID, APIMenuCodeID menuPosCode, GSFlags menuFlags);
+GSErrCode ACAPI_Register_SupportedService (GSType cmdID, Int32 cmdVersion);
+#endif
+
 // clang-format off
 #ifdef __clang__
 #pragma clang diagnostic pop
diff --git a/Engine/Source/Programs/Enterprise/Datasmith/DatasmithARCHICADExporter/Private/Utils/AddonTools.cpp b/Engine/Source/Programs/Enterprise/Datasmith/DatasmithARCHICADExporter/Private/Utils/AddonTools.cpp
--- a/Engine/Source/Programs/Enterprise/Datasmith/DatasmithARCHICADExporter/Private/Utils/AddonTools.cpp
+++ b/Engine/Source/Programs/Enterprise/Datasmith/DatasmithARCHICADExporter/Private/Utils/AddonTools.cpp
@@ -10,6 +10,18 @@
 #include <cstdarg>
 #include <time.h>
 
+#if AC_VERSION > 26
+class API_PenType;
+GSErrCode ACAPI_ProjectOperation_Quit() { return ACAPI_ProjectOperation_Quit(0); }
+GSErrCode ACAPI_GraphicalOverride_GetOverrideRule (API_OverrideRule& rule) { return GS::NoError; }
+GSErrCode ACAPI_GraphicalOverride_CreateOverrideRule (API_OverrideRule& rule) { const API_Guid _Guid{};  return ACAPI_GraphicalOverride_CreateOverrideRule(rule, _Guid); }
+GSErrCode ACAPI_Preferences_SetOldVersion (Int32 version, GSSize nByte, const void* data, unsigned short platformSign, API_FTypeID oldPlanFileID) { return GS::NoError; }
+
+// Definition of most of the methods from previous versions of the ArchiCAD SDK
+// Include once to avoid linkage errors for multiple implementations.
+#include "ACAPI_MigrationHeader.hpp"
+#endif
+
 #if PLATFORM_WINDOWS && (__cplusplus < 201103L)
 	#define va_copy(destination, source) ((destination) = (source))
 #endif
@@ -184,7 +196,7 @@ GS::UniString GetLayerName(API_AttributeIndex InLayer)
 	API_Attribute attribute;
 	Zap(&attribute);
 	attribute.header.typeID = API_LayerID;
-	attribute.header.index = short(InLayer);
+	attribute.header.index = static_cast<decltype(attribute.header.index)>(InLayer);
 	attribute.header.uniStringNamePtr = &LayerName;
 	GSErrCode error = ACAPI_Attribute_Get(&attribute);
 	if (error != NoError)
diff --git a/Engine/Source/Programs/Enterprise/Datasmith/DatasmithARCHICADExporter/Private/Utils/Element2String.cpp b/Engine/Source/Programs/Enterprise/Datasmith/DatasmithARCHICADExporter/Private/Utils/Element2String.cpp
--- a/Engine/Source/Programs/Enterprise/Datasmith/DatasmithARCHICADExporter/Private/Utils/Element2String.cpp
+++ b/Engine/Source/Programs/Enterprise/Datasmith/DatasmithARCHICADExporter/Private/Utils/Element2String.cpp
@@ -268,7 +268,11 @@ utf8_string FElement2String::GetElementAsString(const API_Element& InElement)
 	{
 		if (ObjectType->useObjMaterials)
 		{
+#if AC_VERSION < 27
 			ObjectData += Utf8StringFormat(", mat=%d", ObjectType->mat);
+#else
+			ObjectData += Utf8StringFormat(", mat=%d", ObjectType->mat.ToInt32_Deprecated());
+#endif
 		}
 		if (ObjectType->reflected)
 		{
diff --git a/Engine/Source/Programs/Enterprise/Datasmith/DatasmithARCHICADExporter/Private/Utils/Pasteboard.cpp b/Engine/Source/Programs/Enterprise/Datasmith/DatasmithARCHICADExporter/Private/Utils/Pasteboard.cpp
--- a/Engine/Source/Programs/Enterprise/Datasmith/DatasmithARCHICADExporter/Private/Utils/Pasteboard.cpp
+++ b/Engine/Source/Programs/Enterprise/Datasmith/DatasmithARCHICADExporter/Private/Utils/Pasteboard.cpp
@@ -62,92 +62,4 @@ void SetPasteboardWithString(const utf8_t* InUtf8String)
 #endif
 }
 
-// Set the string from the pasteboard contents
-utf8_string GetStringFromPasteboard()
-{
-	utf8_string TheString;
-
-#ifdef __GNUC__
-	PasteboardRef Pasteboard;
-	OSStatus	  OSErr = ::PasteboardCreate(kPasteboardClipboard, &Pasteboard);
-	if (OSErr != noErr)
-	{
-		fprintf(stderr, "PasteboardCreate error=%d\n", int(OSErr));
-	}
-	else
-	{
-		::PasteboardSyncFlags SyncFlags = ::PasteboardSynchronize(Pasteboard);
-		(void)SyncFlags;
-
-		ItemCount ItemCount;
-		OSErr = ::PasteboardGetItemCount(Pasteboard, &ItemCount);
-		if (OSErr != noErr)
-		{
-			fprintf(stderr, "PasteboardGetItemCount error=%d\n", int(OSErr));
-		}
-
-		for (CFIndex ItemIndex = 1; ItemIndex <= ItemCount && TheString.size() == 0; ItemIndex++)
-		{
-			PasteboardItemID ItemID;
-			OSErr = ::PasteboardGetItemIdentifier(Pasteboard, ItemIndex, &ItemID);
-			if (OSErr != noErr)
-			{
-				fprintf(stderr, "PasteboardGetItemIdentifier error=%d\n", int(OSErr));
-			}
-			else
-			{
-				CFArrayRef FlavorTypeArray;
-				OSErr = ::PasteboardCopyItemFlavors(Pasteboard, ItemID, &FlavorTypeArray);
-				if (OSErr != noErr)
-				{
-					fprintf(stderr, "PasteboardCopyItemFlavors error=%d\n", int(OSErr));
-				}
-				else
-				{
-					CFIndex FlavorCount = ::CFArrayGetCount(FlavorTypeArray);
-					for (CFIndex FlavorIndex = 0; FlavorIndex < FlavorCount && TheString.size() == 0; FlavorIndex++)
-					{
-						CFStringRef FlavorType = (CFStringRef)::CFArrayGetValueAtIndex(FlavorTypeArray, FlavorIndex);
-						if (::UTTypeConformsTo(FlavorType, CFSTR("public.utf8-plain-text")))
-						{
-							CFDataRef FlavorData;
-
-							OSErr = ::PasteboardCopyItemFlavorData(Pasteboard, ItemID, FlavorType, &FlavorData);
-							if (OSErr != noErr)
-							{
-								fprintf(stderr, "PasteboardCopyItemFlavorData error=%d\n", int(OSErr));
-							}
-							else
-							{
-								CFIndex FlavorDataSize = ::CFDataGetLength(FlavorData);
-								if (FlavorDataSize > 0)
-								{
-									TheString.assign(reinterpret_cast< const char* >(::CFDataGetBytePtr(FlavorData)),
-													 FlavorDataSize);
-								}
-
-								CFRelease(FlavorData);
-							}
-						}
-					}
-				}
-				CFRelease(FlavorTypeArray);
-			}
-		}
-
-		CFRelease(Pasteboard);
-	}
-#else
-	if (OpenClipboard(NULL))
-	{
-		HANDLE HData = GetClipboardData(CF_TEXT);
-		TheString = reinterpret_cast< char* >(GlobalLock(HData));
-		GlobalUnlock(HData);
-		CloseClipboard();
-	}
-#endif
-
-	return TheString;
-}
-
 END_NAMESPACE_UE_AC
diff --git a/Engine/Source/Programs/Enterprise/Datasmith/DatasmithARCHICADExporter/Private/Utils/Pasteboard.h b/Engine/Source/Programs/Enterprise/Datasmith/DatasmithARCHICADExporter/Private/Utils/Pasteboard.h
--- a/Engine/Source/Programs/Enterprise/Datasmith/DatasmithARCHICADExporter/Private/Utils/Pasteboard.h
+++ b/Engine/Source/Programs/Enterprise/Datasmith/DatasmithARCHICADExporter/Private/Utils/Pasteboard.h
@@ -11,7 +11,4 @@ BEGIN_NAMESPACE_UE_AC
 // Copy the string to the pasteboard
 extern void SetPasteboardWithString(const utf8_t* InUtf8String);
 
-// Set the string from the pasteboard contents
-extern utf8_string GetStringFromPasteboard();
-
 END_NAMESPACE_UE_AC
diff --git a/Engine/Source/Programs/Enterprise/Datasmith/DatasmithARCHICADExporter/Private/Utils/ViewState.cpp b/Engine/Source/Programs/Enterprise/Datasmith/DatasmithARCHICADExporter/Private/Utils/ViewState.cpp
--- a/Engine/Source/Programs/Enterprise/Datasmith/DatasmithARCHICADExporter/Private/Utils/ViewState.cpp
+++ b/Engine/Source/Programs/Enterprise/Datasmith/DatasmithARCHICADExporter/Private/Utils/ViewState.cpp
@@ -114,18 +114,28 @@ bool FViewState::operator==(const FViewState& InOther) const
 void FViewState::CollectVisibleLayers()
 {
 	API_Attribute LayerAttribute;
-	Zap(&LayerAttribute);
+	BNZeroMemory(&LayerAttribute, sizeof(LayerAttribute));
 	LayerAttribute.header.typeID = API_LayerID;
 
-	API_AttributeIndex LayerCount = 0;
-	GSErrCode		   GSErr = ACAPI_Attribute_GetNum(API_LayerID, &LayerCount);
+	GS::Int32 LayerCount = 0;
+#if AC_VERSION > 26
+	GS::UInt32 LocalCount = 0;
+	GSErrCode GSErr = ACAPI_Attribute_GetNum(API_LayerID, LocalCount);
+	LayerCount = GS::Int32(LocalCount);
+#else
+	GSErrCode GSErr = ACAPI_Attribute_GetNum(API_LayerID, &LayerCount);
+#endif
 	if (GSErr == NoError)
 	{
 		VisibleLayers.Reserve(LayerCount);
 
-		for (API_AttributeIndex Index = 1; Index <= LayerCount && GSErr == NoError; Index++)
+		for (GS::Int32 Index = 1; Index <= LayerCount && GSErr == NoError; Index++)
 		{
+#if AC_VERSION > 26
+			LayerAttribute.header.index = ACAPI_CreateAttributeIndex(Index);
+#else
 			LayerAttribute.header.index = Index;
+#endif
 			GSErr = ACAPI_Attribute_Get(&LayerAttribute);
 			if (GSErr == NoError)
 			{
@@ -134,7 +144,7 @@ void FViewState::CollectVisibleLayers()
 					VisibleLayers.Add(Index);
 				}
 			}
-			else if (GSErr == APIERR_DELETED)
+			else if (GSErr == APIERR_DELETED || GSErr == APIERR_BADINDEX)
 			{
 				GSErr = NoError;
 			}
diff --git a/Engine/Source/Programs/Enterprise/Datasmith/DatasmithARCHICADExporter/Private/Utils/ViewState.h b/Engine/Source/Programs/Enterprise/Datasmith/DatasmithARCHICADExporter/Private/Utils/ViewState.h
--- a/Engine/Source/Programs/Enterprise/Datasmith/DatasmithARCHICADExporter/Private/Utils/ViewState.h
+++ b/Engine/Source/Programs/Enterprise/Datasmith/DatasmithARCHICADExporter/Private/Utils/ViewState.h
@@ -44,7 +44,7 @@ class FViewState
 	bool CompareCamera(const FViewState& InOther) const;
 
 	FCutInfo					 CutInfo;
-	TArray< API_AttributeIndex > VisibleLayers;
+	TArray< GS::Int32 >			VisibleLayers;
 	API_3DProjectionInfo		 ProjSets;
 	GSErrCode					 ProjSetsError;
 };
