[3dsMax Exporter] Compile warning fix and unused code cleanup

diff --git a/Engine/Source/Programs/Enterprise/Datasmith/DatasmithMaxExporter/DatasmithMax2016.Build.cs b/Engine/Source/Programs/Enterprise/Datasmith/DatasmithMaxExporter/DatasmithMax2016.Build.cs
deleted file mode 100644
diff --git a/Engine/Source/Programs/Enterprise/Datasmith/DatasmithMaxExporter/DatasmithMax2016.Target.cs b/Engine/Source/Programs/Enterprise/Datasmith/DatasmithMaxExporter/DatasmithMax2016.Target.cs
deleted file mode 100644
diff --git a/Engine/Source/Programs/Enterprise/Datasmith/DatasmithMaxExporter/Private/DatasmithMaxArnold.cpp b/Engine/Source/Programs/Enterprise/Datasmith/DatasmithMaxExporter/Private/DatasmithMaxArnold.cpp
--- a/Engine/Source/Programs/Enterprise/Datasmith/DatasmithMaxExporter/Private/DatasmithMaxArnold.cpp
+++ b/Engine/Source/Programs/Enterprise/Datasmith/DatasmithMaxExporter/Private/DatasmithMaxArnold.cpp
@@ -3,7 +3,7 @@
 #include "DatasmithMaxWriter.h"
 
 #include "DatasmithSceneFactory.h"
-#include "DatasmithMaxSceneParser.h"
+#include "DatasmithMaxSceneHelper.h"
 
 #include "Misc/Paths.h"
 
diff --git a/Engine/Source/Programs/Enterprise/Datasmith/DatasmithMaxExporter/Private/DatasmithMaxCameraExporter.cpp b/Engine/Source/Programs/Enterprise/Datasmith/DatasmithMaxExporter/Private/DatasmithMaxCameraExporter.cpp
--- a/Engine/Source/Programs/Enterprise/Datasmith/DatasmithMaxExporter/Private/DatasmithMaxCameraExporter.cpp
+++ b/Engine/Source/Programs/Enterprise/Datasmith/DatasmithMaxExporter/Private/DatasmithMaxCameraExporter.cpp
@@ -4,7 +4,7 @@
 #include "DatasmithMaxDirectLink.h"
 #include "DatasmithMaxExporterDefines.h"
 #include "DatasmithMaxHelper.h"
-#include "DatasmithMaxSceneParser.h"
+#include "DatasmithMaxSceneHelper.h"
 #include "DatasmithMaxWriter.h"
 #include "DatasmithMaxLogger.h"
 #include "DatasmithSceneFactory.h"
diff --git a/Engine/Source/Programs/Enterprise/Datasmith/DatasmithMaxExporter/Private/DatasmithMaxConverters.cpp b/Engine/Source/Programs/Enterprise/Datasmith/DatasmithMaxExporter/Private/DatasmithMaxConverters.cpp
--- a/Engine/Source/Programs/Enterprise/Datasmith/DatasmithMaxExporter/Private/DatasmithMaxConverters.cpp
+++ b/Engine/Source/Programs/Enterprise/Datasmith/DatasmithMaxExporter/Private/DatasmithMaxConverters.cpp
@@ -10,7 +10,7 @@
 
 #include "DatasmithMaxSceneExporter.h"
 #include "DatasmithMaxCameraExporter.h"
-#include "DatasmithMaxSceneParser.h"
+#include "DatasmithMaxSceneHelper.h"
 #include "DatasmithMaxLogger.h"
 #include "DatasmithMaxHelper.h"
 
@@ -83,7 +83,7 @@ void FLightNodeConverter::ConvertToDatasmith(ISceneTracker& SceneTracker, FNodeT
 
 		if (!LightElement)
 		{
-			if (FDatasmithMaxSceneParser::GetLightClass(NodeTracker.Node) == EMaxLightClass::SkyEquivalent)
+			if (FDatasmithMaxSceneHelper::GetLightClass(NodeTracker.Node) == EMaxLightClass::SkyEquivalent)
 			{
 				SceneTracker.GetDatasmithSceneRef()->SetUsePhysicalSky(true);
 			}
diff --git a/Engine/Source/Programs/Enterprise/Datasmith/DatasmithMaxExporter/Private/DatasmithMaxCorona.cpp b/Engine/Source/Programs/Enterprise/Datasmith/DatasmithMaxExporter/Private/DatasmithMaxCorona.cpp
--- a/Engine/Source/Programs/Enterprise/Datasmith/DatasmithMaxExporter/Private/DatasmithMaxCorona.cpp
+++ b/Engine/Source/Programs/Enterprise/Datasmith/DatasmithMaxExporter/Private/DatasmithMaxCorona.cpp
@@ -4,7 +4,7 @@
 
 #include "DatasmithSceneFactory.h"
 #include "DatasmithMaxTexmapParser.h"
-#include "DatasmithMaxSceneParser.h"
+#include "DatasmithMaxSceneHelper.h"
 #include "DatasmithMaxSceneExporter.h"
 
 #include "DatasmithMaterialsUtils.h"
diff --git a/Engine/Source/Programs/Enterprise/Datasmith/DatasmithMaxExporter/Private/DatasmithMaxDirectLinkMaterial.cpp b/Engine/Source/Programs/Enterprise/Datasmith/DatasmithMaxExporter/Private/DatasmithMaxDirectLinkMaterial.cpp
--- a/Engine/Source/Programs/Enterprise/Datasmith/DatasmithMaxExporter/Private/DatasmithMaxDirectLinkMaterial.cpp
+++ b/Engine/Source/Programs/Enterprise/Datasmith/DatasmithMaxExporter/Private/DatasmithMaxDirectLinkMaterial.cpp
@@ -315,9 +315,6 @@ void FMaterialsCollectionTracker::SetMaterialsForMeshActor(const TSharedPtr<IDat
 }
 
 
-// Copied from
-// FDatasmithMaxSceneParser::MaterialEnum
-// FDatasmithMaxSceneParser::TexEnum
 // Collects actual materials that are used by the top-level material(assigned to node)
 class FMaterialEnum
 {
diff --git a/Engine/Source/Programs/Enterprise/Datasmith/DatasmithMaxExporter/Private/DatasmithMaxDirectLinkMesh.cpp b/Engine/Source/Programs/Enterprise/Datasmith/DatasmithMaxExporter/Private/DatasmithMaxDirectLinkMesh.cpp
--- a/Engine/Source/Programs/Enterprise/Datasmith/DatasmithMaxExporter/Private/DatasmithMaxDirectLinkMesh.cpp
+++ b/Engine/Source/Programs/Enterprise/Datasmith/DatasmithMaxExporter/Private/DatasmithMaxDirectLinkMesh.cpp
@@ -10,7 +10,7 @@
 #include "DatasmithMaxMeshExporter.h"
 #include "DatasmithMaxSceneExporter.h"
 #include "DatasmithMaxAttributes.h"
-#include "DatasmithMaxSceneParser.h"
+#include "DatasmithMaxSceneHelper.h"
 #include "DatasmithMaxLogger.h"
 
 #include "Logging/LogMacros.h"
@@ -167,7 +167,7 @@ INode* GetCollisionNode(ISceneTracker& SceneTracker, INode* Node, const FDatasmi
 	
 
 	FString OriginalName = Node->GetName();
-	for ( const FString& CollisionNodePrefix : FDatasmithMaxSceneParser::CollisionNodesPrefixes )
+	for ( const FString& CollisionNodePrefix : FDatasmithMaxSceneHelper::CollisionNodesPrefixes )
 	{
 		if (FNodeTracker* CollisionNode = SceneTracker.GetNodeTrackerByNodeName(*( CollisionNodePrefix + TEXT("_") + OriginalName )))
 		{
diff --git a/Engine/Source/Programs/Enterprise/Datasmith/DatasmithMaxExporter/Private/DatasmithMaxExporter.cpp b/Engine/Source/Programs/Enterprise/Datasmith/DatasmithMaxExporter/Private/DatasmithMaxExporter.cpp
deleted file mode 100644
diff --git a/Engine/Source/Programs/Enterprise/Datasmith/DatasmithMaxExporter/Private/DatasmithMaxExporterWithDirectLink.cpp b/Engine/Source/Programs/Enterprise/Datasmith/DatasmithMaxExporter/Private/DatasmithMaxExporterWithDirectLink.cpp
--- a/Engine/Source/Programs/Enterprise/Datasmith/DatasmithMaxExporter/Private/DatasmithMaxExporterWithDirectLink.cpp
+++ b/Engine/Source/Programs/Enterprise/Datasmith/DatasmithMaxExporter/Private/DatasmithMaxExporterWithDirectLink.cpp
@@ -4,15 +4,13 @@
 
 #include "DatasmithMaxDirectLink.h"
 
-#include "DatasmithMaxExporter.h"
-
 #include "DatasmithMaxSceneExporter.h"
 #include "DatasmithMaxHelper.h"
 #include "DatasmithMaxWriter.h"
 #include "DatasmithMaxClassIDs.h"
 
 #include "DatasmithMaxLogger.h"
-#include "DatasmithMaxSceneParser.h"
+#include "DatasmithMaxSceneHelper.h"
 #include "DatasmithMaxCameraExporter.h"
 #include "DatasmithMaxAttributes.h"
 #include "DatasmithMaxProgressManager.h"
@@ -2247,7 +2245,7 @@ public:
 
 		// Check if node changed its being assigned as collision
 		{
-			if (FDatasmithMaxSceneParser::HasCollisionName(NodeTracker.Node))
+			if (FDatasmithMaxSceneHelper::HasCollisionName(NodeTracker.Node))
 			{
 				CollisionNodes.Add(&NodeTracker); // Always view node with 'collision' name as a collision node(i.e. no render)
 
@@ -2370,7 +2368,7 @@ public:
 		{
 			SCENE_UPDATE_STAT_INC(UpdateNode, LightsEncontered);
 
-			if (EMaxLightClass::Unknown == FDatasmithMaxSceneParser::GetLightClass(NodeTracker.Node))
+			if (EMaxLightClass::Unknown == FDatasmithMaxSceneHelper::GetLightClass(NodeTracker.Node))
 			{
 				SCENE_UPDATE_STAT_INC(UpdateNode, LightsSkippedAsUnknown);
 				break;
@@ -2400,7 +2398,7 @@ public:
 			}
 			else
 			{
-				if (FDatasmithMaxSceneParser::HasCollisionName(NodeTracker.Node))
+				if (FDatasmithMaxSceneHelper::HasCollisionName(NodeTracker.Node))
 				{
 					ConvertNamedCollisionNode(NodeTracker);
 				}
@@ -2533,6 +2531,7 @@ public:
 
 
 		// Update geometry using selected Node
+		if (SelectedNodeTracker)
 		{
 			FNodeTracker& NodeTracker = *SelectedNodeTracker;
 
@@ -3654,8 +3653,14 @@ public:
 		LastInputInfo.dwTime = 0;
 		if (GetLastInputInfo(&LastInputInfo))
 		{
+
+// Disable "Consider using 'GetTickCount64' instead of 'GetTickCount'" warning - we don't GetLastInputInfo returns results of GetTickCount
+// and we won't miss much a single update when 32-bit timer wraparound would happen
+#pragma warning( push )
+#pragma warning( disable: 28159 )
 			DWORD CurrentTime = GetTickCount();
 			int32 IdlePeriod = GetTickCount() - LastInputInfo.dwTime;
+#pragma warning( pop )
 			LogDebug(FString::Printf(TEXT("CurrentTime: %ld, Idle time: %ld, IdlePeriod: %ld"), CurrentTime, LastInputInfo.dwTime, IdlePeriod));
 
 			if (IdlePeriod > FMath::RoundToInt(AutoSyncIdleDelaySeconds*1000))
diff --git a/Engine/Source/Programs/Enterprise/Datasmith/DatasmithMaxExporter/Private/DatasmithMaxMental.cpp b/Engine/Source/Programs/Enterprise/Datasmith/DatasmithMaxExporter/Private/DatasmithMaxMental.cpp
--- a/Engine/Source/Programs/Enterprise/Datasmith/DatasmithMaxExporter/Private/DatasmithMaxMental.cpp
+++ b/Engine/Source/Programs/Enterprise/Datasmith/DatasmithMaxExporter/Private/DatasmithMaxMental.cpp
@@ -1,7 +1,7 @@
 // Copyright Epic Games, Inc. All Rights Reserved.
 
 #include "DatasmithMaxWriter.h"
-#include "DatasmithMaxSceneParser.h"
+#include "DatasmithMaxSceneHelper.h"
 
 void FDatasmithMaxMatWriter::ExportArchDesignMaterial(TSharedRef< IDatasmithScene > DatasmithScene, TSharedPtr< IDatasmithMaterialElement >& MaterialElement, Mtl* Material)
 {
diff --git a/Engine/Source/Programs/Enterprise/Datasmith/DatasmithMaxExporter/Private/DatasmithMaxMeshExporter.cpp b/Engine/Source/Programs/Enterprise/Datasmith/DatasmithMaxExporter/Private/DatasmithMaxMeshExporter.cpp
--- a/Engine/Source/Programs/Enterprise/Datasmith/DatasmithMaxExporter/Private/DatasmithMaxMeshExporter.cpp
+++ b/Engine/Source/Programs/Enterprise/Datasmith/DatasmithMaxExporter/Private/DatasmithMaxMeshExporter.cpp
@@ -6,7 +6,7 @@
 #include "DatasmithExportOptions.h"
 #include "DatasmithUtils.h"
 #include "DatasmithMaxExporterDefines.h"
-#include "DatasmithMaxSceneParser.h"
+#include "DatasmithMaxSceneHelper.h"
 #include "DatasmithMesh.h"
 #include "DatasmithMeshExporter.h"
 #include "DatasmithMaxLogger.h"
@@ -208,7 +208,7 @@ INode* FDatasmithMaxMeshExporter::GetCollisionNode(INode* OriginalNode, const FD
 	}
 
 	FString OriginalName = OriginalNode->GetName();
-	for ( const FString& CollisionNodePrefix : FDatasmithMaxSceneParser::CollisionNodesPrefixes )
+	for ( const FString& CollisionNodePrefix : FDatasmithMaxSceneHelper::CollisionNodesPrefixes )
 	{
 		INode* CollisionNode = GetCOREInterface()->GetINodeByName( *( CollisionNodePrefix + TEXT("_") + OriginalName ) );
 		if (CollisionNode)
diff --git a/Engine/Source/Programs/Enterprise/Datasmith/DatasmithMaxExporter/Private/DatasmithMaxScanline.cpp b/Engine/Source/Programs/Enterprise/Datasmith/DatasmithMaxExporter/Private/DatasmithMaxScanline.cpp
--- a/Engine/Source/Programs/Enterprise/Datasmith/DatasmithMaxExporter/Private/DatasmithMaxScanline.cpp
+++ b/Engine/Source/Programs/Enterprise/Datasmith/DatasmithMaxExporter/Private/DatasmithMaxScanline.cpp
@@ -5,7 +5,7 @@
 
 #include "DatasmithSceneFactory.h"
 #include "DatasmithSceneExporter.h"
-#include "DatasmithMaxSceneParser.h"
+#include "DatasmithMaxSceneHelper.h"
 #include "DatasmithMaxLogger.h"
 #include "DatasmithExportOptions.h"
 #include "DatasmithMaxSceneExporter.h"
diff --git a/Engine/Source/Programs/Enterprise/Datasmith/DatasmithMaxExporter/Private/DatasmithMaxSceneExporter.cpp b/Engine/Source/Programs/Enterprise/Datasmith/DatasmithMaxExporter/Private/DatasmithMaxSceneExporter.cpp
--- a/Engine/Source/Programs/Enterprise/Datasmith/DatasmithMaxExporter/Private/DatasmithMaxSceneExporter.cpp
+++ b/Engine/Source/Programs/Enterprise/Datasmith/DatasmithMaxExporter/Private/DatasmithMaxSceneExporter.cpp
@@ -9,7 +9,7 @@
 #include "DatasmithMaxExporterDefines.h"
 #include "DatasmithMaxHelper.h"
 #include "DatasmithMaxLogger.h"
-#include "DatasmithMaxSceneParser.h"
+#include "DatasmithMaxSceneHelper.h"
 #include "DatasmithMaxWriter.h"
 #include "DatasmithSceneExporter.h"
 #include "DatasmithSceneFactory.h"
@@ -38,7 +38,7 @@ FMaxLightCoordinateConversionParams::FMaxLightCoordinateConversionParams(INode*
 	if (ObjState.obj != nullptr && ObjState.obj->SuperClassID() == LIGHT_CLASS_ID)
 	{
 		bIsLight = true;
-		bIsCoronaLight = FDatasmithMaxSceneParser::GetLightClass(LightNode) == EMaxLightClass::CoronaLight;
+		bIsCoronaLight = FDatasmithMaxSceneHelper::GetLightClass(LightNode) == EMaxLightClass::CoronaLight;
 		LightShape = Shape;
 	}
 }
@@ -279,7 +279,7 @@ void FDatasmithMaxSceneExporter::ExportAnimation( TSharedRef< IDatasmithLevelSeq
 
 TSharedPtr< IDatasmithLightActorElement > FDatasmithMaxSceneExporter::CreateLightElementForNode(INode* Node, const TCHAR* Name)
 {
-	EMaxLightClass LightClass = FDatasmithMaxSceneParser::GetLightClass(Node);
+	EMaxLightClass LightClass = FDatasmithMaxSceneHelper::GetLightClass(Node);
 	if (LightClass == EMaxLightClass::Unknown)
 	{
 		return TSharedPtr< IDatasmithLightActorElement >();
@@ -666,7 +666,7 @@ bool FDatasmithMaxSceneExporter::ParseTransformAnimation(INode* ParentNode, INod
 
 bool FDatasmithMaxSceneExporter::ParseLight(DatasmithMaxDirectLink::FLightNodeConverter& Converter, INode* Node, TSharedRef< IDatasmithLightActorElement > LightElement, TSharedRef< IDatasmithScene > DatasmithScene)
 {
-	EMaxLightClass LightClass = FDatasmithMaxSceneParser::GetLightClass(Node);
+	EMaxLightClass LightClass = FDatasmithMaxSceneHelper::GetLightClass(Node);
 	if (LightClass == EMaxLightClass::Unknown)
 	{
 		return false;
diff --git a/Engine/Source/Programs/Enterprise/Datasmith/DatasmithMaxExporter/Private/DatasmithMaxSceneHelper.cpp b/Engine/Source/Programs/Enterprise/Datasmith/DatasmithMaxExporter/Private/DatasmithMaxSceneHelper.cpp
new file mode 100644
--- /dev/null
+++ b/Engine/Source/Programs/Enterprise/Datasmith/DatasmithMaxExporter/Private/DatasmithMaxSceneHelper.cpp
@@ -0,0 +1,167 @@
+// Copyright Epic Games, Inc. All Rights Reserved.
+
+#include "DatasmithMaxSceneHelper.h"
+
+#include "DatasmithMaxLogger.h"
+
+#include "Windows/AllowWindowsPlatformTypes.h"
+
+MAX_INCLUDES_START
+#include "lslights.h"
+#include "xref/iXrefObj.h"
+MAX_INCLUDES_END
+
+const TArray< FString, TInlineAllocator< 4 > > FDatasmithMaxSceneHelper::CollisionNodesPrefixes = { TEXT("UBX"), TEXT("UCX"), TEXT("UCP"), TEXT("USP") };
+
+EMaxLightClass FDatasmithMaxSceneHelper::GetLightClass(INode* Node)
+{
+	EMaxLightClass LightClass = EMaxLightClass::Unknown;
+	ObjectState ObjState = Node->EvalWorldState(0);
+
+	if (ObjState.obj->SuperClassID() == LIGHT_CLASS_ID)
+	{
+		LightObject* Light = (LightObject*)ObjState.obj;
+		if (Light != nullptr)
+		{
+			Class_ID ClassID = Light->ClassID();
+			if (ClassID == THEAOMNICLASS)
+			{
+				LightClass = EMaxLightClass::TheaLightOmni;
+			}
+			else if (ClassID == THEASPOTCLASS)
+			{
+				LightClass = EMaxLightClass::TheaLightSpot;
+			}
+			else if (ClassID == THEAIESCLASS)
+			{
+				LightClass = EMaxLightClass::TheaLightIES;
+			}
+			else if (ClassID == THEAPLANECLASS)
+			{
+				LightClass = EMaxLightClass::TheaLightPlane;
+			}
+			else if (ClassID == REGULARSPOTACLASS || ClassID == REGULARSPOTBCLASS)
+			{
+				LightClass = EMaxLightClass::SpotLight;
+			}
+			else if (ClassID == REGULARDIRECTACLASS || ClassID == REGULARDIRECTBCLASS)
+			{
+				LightClass = EMaxLightClass::DirectLight;
+
+				INode *ParentNode = Node->GetParentNode();
+				if (ParentNode)
+				{
+					ObjectState ParentObjectState = ParentNode->EvalWorldState(0);
+					if (ParentObjectState.obj && ParentObjectState.obj->ClassID() == DAYLIGHTASSEMBLYCLASS)
+					{
+						LightClass = EMaxLightClass::SunEquivalent;
+					}
+				}
+			}
+			else if (ClassID == REGULAROMNICLASS)
+			{
+				LightClass = EMaxLightClass::OmniLight;
+			}
+			else if (Light->IsSubClassOf(LIGHTSCAPE_LIGHT_CLASS))
+			{
+				LightClass = EMaxLightClass::PhotometricLight;
+			}
+			else if (ClassID == PHOTOPLANEACLASS || ClassID == PHOTOPLANEBCLASS)
+			{
+				LightClass = EMaxLightClass::PhotoplaneLight;
+			}
+			else if (ClassID == VRAYLIGHTIESCLASS)
+			{
+				LightClass = EMaxLightClass::VRayLightIES;
+			}
+			else if (ClassID == VRAYLIGHTCLASS)
+			{
+				LightClass = EMaxLightClass::VRayLight;
+			}
+			else if (ClassID == CORONALIGHTCLASS)
+			{
+				LightClass = EMaxLightClass::CoronaLight;
+			}
+			else if (ClassID == ARNOLDLIGHTCLASS)
+			{
+				LightClass = EMaxLightClass::ArnoldLight;
+			}
+			else if (ClassID == SUNLIGHTACLASS ||
+				ClassID == SUNLIGHTBCLASS ||
+				ClassID == VRAYSUNCLASS ||
+				ClassID == SUNPOSITIONERCLASS ||
+				ClassID == CORONASUNCLASS ||
+				ClassID == CORONASUNCLASSB)
+			{
+				LightClass = EMaxLightClass::SunEquivalent;
+			}
+			else if (ClassID == SKYLIGHTCLASS ||
+				ClassID == IESSKYLIGHTCLASS ||
+				ClassID == MRSKYLIGHTCLASS)
+			{
+				LightClass = EMaxLightClass::SkyEquivalent;
+			}
+			else if (ClassID == MRPORTALLIGHTCLASS)
+			{
+				LightClass = EMaxLightClass::SkyPortal;
+			}
+		}
+	}
+
+	if (LightClass == EMaxLightClass::Unknown)
+	{
+		DatasmithMaxLogger::Get().AddUnsupportedLight(Node);
+	}
+
+	return LightClass;
+}
+
+bool FDatasmithMaxSceneHelper::CanBeTriMesh(Object* Obj)
+{
+	if (Obj == NULL)
+	{
+		return false;
+	}
+
+	int CanConvertToTri = Obj->CanConvertToType(Class_ID(TRIOBJ_CLASS_ID, 0));
+	if (CanConvertToTri == 0)
+	{
+		return false;
+	}
+
+	TriObject* Tri = (TriObject*)Obj->ConvertToType(GetCOREInterface()->GetTime(), Class_ID(TRIOBJ_CLASS_ID, 0));
+	if (Tri == NULL)
+	{
+		return false;
+	}
+
+	bool bResult = true;
+	if (Tri->GetMesh().getNumFaces() == 0)
+	{
+		bResult = false;
+	}
+
+	if (Obj != Tri)
+	{
+		Tri->DeleteMe();
+	}
+
+	return bResult;
+}
+
+bool FDatasmithMaxSceneHelper::HasCollisionName(INode* Node)
+{
+	FString NodeName = Node->GetName();
+	FString LeftString, RightString;
+
+	if ( NodeName.Split( TEXT("_"), &LeftString, &RightString ) )
+	{
+		if ( CollisionNodesPrefixes.Find( LeftString ) != INDEX_NONE )
+		{
+			return true;
+		}
+	}
+	return false;
+}
+
+#include "Windows/HideWindowsPlatformTypes.h"
diff --git a/Engine/Source/Programs/Enterprise/Datasmith/DatasmithMaxExporter/Private/DatasmithMaxSceneHelper.h b/Engine/Source/Programs/Enterprise/Datasmith/DatasmithMaxExporter/Private/DatasmithMaxSceneHelper.h
new file mode 100644
--- /dev/null
+++ b/Engine/Source/Programs/Enterprise/Datasmith/DatasmithMaxExporter/Private/DatasmithMaxSceneHelper.h
@@ -0,0 +1,78 @@
+// Copyright Epic Games, Inc. All Rights Reserved.
+#pragma once
+
+#include "CoreMinimal.h"
+
+#include "DatasmithMaxClassIDs.h"
+#include "DatasmithMaxExporterDefines.h"
+#include "DatasmithSceneFactory.h"
+#include "DatasmithUtils.h"
+
+#include "Templates/SharedPointer.h"
+
+#include "Windows/AllowWindowsPlatformTypes.h"
+MAX_INCLUDES_START
+	#include "max.h"
+MAX_INCLUDES_END
+#include "Windows/HideWindowsPlatformTypes.h"
+
+class FMaxRendereableNode;
+class FDatasmithMaxStaticMeshAttributes;
+
+// Dummy function for calls to ProgressStart in the Max API
+DWORD WINAPI DummyFunction(LPVOID Arg);
+
+// This tell the exporter how to manage the instancing of a Geometry
+enum class EMaxExporterInstanceMode
+{
+	NotInstanced,
+	InstanceMaster,
+	InstanceCopy,
+	UnrealHISM
+};
+
+enum class EMaxTriMode
+{
+	TriNotChecked,
+	TriEnabled,
+	TriDisabled
+};
+
+enum class EMaxEntityType
+{
+	Geometry,
+	Light,
+	Camera,
+	Dummy
+};
+
+enum class EMaxLightClass
+{
+	Unknown,
+	TheaLightOmni,
+	TheaLightSpot,
+	TheaLightIES,
+	TheaLightPlane,
+	SpotLight,
+	DirectLight,
+	OmniLight,
+	PhotometricLight,
+	PhotoplaneLight,
+	VRayLightIES,
+	VRayLight,
+	CoronaLight,
+	ArnoldLight,
+	SunEquivalent,
+	SkyEquivalent,
+	SkyPortal
+};
+
+class FDatasmithMaxSceneHelper
+{
+public:
+	static EMaxLightClass GetLightClass(INode* InNode);
+	static bool CanBeTriMesh(Object* Obj);
+	static const TArray< FString, TInlineAllocator< 4 > > CollisionNodesPrefixes; // List of supported mesh prefixes for Unreal collision
+
+	static bool HasCollisionName(INode* Node);
+};
diff --git a/Engine/Source/Programs/Enterprise/Datasmith/DatasmithMaxExporter/Private/DatasmithMaxSceneParser.cpp b/Engine/Source/Programs/Enterprise/Datasmith/DatasmithMaxExporter/Private/DatasmithMaxSceneParser.cpp
deleted file mode 100644
diff --git a/Engine/Source/Programs/Enterprise/Datasmith/DatasmithMaxExporter/Private/DatasmithMaxThea.cpp b/Engine/Source/Programs/Enterprise/Datasmith/DatasmithMaxExporter/Private/DatasmithMaxThea.cpp
--- a/Engine/Source/Programs/Enterprise/Datasmith/DatasmithMaxExporter/Private/DatasmithMaxThea.cpp
+++ b/Engine/Source/Programs/Enterprise/Datasmith/DatasmithMaxExporter/Private/DatasmithMaxThea.cpp
@@ -3,7 +3,7 @@
 #include "DatasmithMaxWriter.h"
 
 #include "DatasmithExportOptions.h"
-#include "DatasmithMaxSceneParser.h"
+#include "DatasmithMaxSceneHelper.h"
 
 FString FDatasmithMaxMatWriter::DumpBitmapThea(TSharedPtr<IDatasmithCompositeTexture>& CompTex, BitmapTex* InBitmapTex, const TCHAR* Prefix, bool bForceInvert, bool bIsGrayscale)
 {
diff --git a/Engine/Source/Programs/Enterprise/Datasmith/DatasmithMaxExporter/Private/DatasmithMaxVRay.cpp b/Engine/Source/Programs/Enterprise/Datasmith/DatasmithMaxExporter/Private/DatasmithMaxVRay.cpp
--- a/Engine/Source/Programs/Enterprise/Datasmith/DatasmithMaxExporter/Private/DatasmithMaxVRay.cpp
+++ b/Engine/Source/Programs/Enterprise/Datasmith/DatasmithMaxExporter/Private/DatasmithMaxVRay.cpp
@@ -1,7 +1,7 @@
 // Copyright Epic Games, Inc. All Rights Reserved.
 
 #include "DatasmithMaxWriter.h"
-#include "DatasmithMaxSceneParser.h"
+#include "DatasmithMaxSceneHelper.h"
 #include "DatasmithMaxSceneExporter.h"
 
 #include "Misc/Paths.h"
diff --git a/Engine/Source/Programs/Enterprise/Datasmith/DatasmithMaxExporter/Private/DatasmithMaxWriter.cpp b/Engine/Source/Programs/Enterprise/Datasmith/DatasmithMaxExporter/Private/DatasmithMaxWriter.cpp
--- a/Engine/Source/Programs/Enterprise/Datasmith/DatasmithMaxExporter/Private/DatasmithMaxWriter.cpp
+++ b/Engine/Source/Programs/Enterprise/Datasmith/DatasmithMaxExporter/Private/DatasmithMaxWriter.cpp
@@ -6,7 +6,7 @@
 #include "MaxMaterialsToUEPbr/DatasmithMaxMaterialsToUEPbr.h"
 #include "DatasmithExportOptions.h"
 #include "DatasmithSceneFactory.h"
-#include "DatasmithMaxSceneParser.h"
+#include "DatasmithMaxSceneHelper.h"
 #include "DatasmithMaxLogger.h"
 #include "Misc/Paths.h"
 
diff --git a/Engine/Source/Programs/Enterprise/Datasmith/DatasmithMaxExporter/Public/DatasmithMaxExporter.h b/Engine/Source/Programs/Enterprise/Datasmith/DatasmithMaxExporter/Public/DatasmithMaxExporter.h
deleted file mode 100644
diff --git a/Engine/Source/Programs/Enterprise/Datasmith/DatasmithMaxExporter/Public/DatasmithMaxSceneParser.h b/Engine/Source/Programs/Enterprise/Datasmith/DatasmithMaxExporter/Public/DatasmithMaxSceneParser.h
deleted file mode 100644
