[solidworks] Add support for datasmith exporter of Solidworks 2024

diff --git a/Engine/Source/Programs/Enterprise/Datasmith/DatasmithSolidworksExporter/RegistryKeyExistsCustomAction/CustomAction.cs b/Engine/Source/Programs/Enterprise/Datasmith/DatasmithSolidworksExporter/RegistryKeyExistsCustomAction/CustomAction.cs
--- a/Engine/Source/Programs/Enterprise/Datasmith/DatasmithSolidworksExporter/RegistryKeyExistsCustomAction/CustomAction.cs
+++ b/Engine/Source/Programs/Enterprise/Datasmith/DatasmithSolidworksExporter/RegistryKeyExistsCustomAction/CustomAction.cs
@@ -61,6 +61,7 @@ namespace RegistryKeyExistsCustomAction
 			var tools2021 = CheckSolidworksInstalledSub(session, "SOLIDWORKS 2021", "SOLIDWORKS2021INSTALLED");
 			var tools2022 = CheckSolidworksInstalledSub(session, "SOLIDWORKS 2022", "SOLIDWORKS2022INSTALLED");
 			var tools2023 = CheckSolidworksInstalledSub(session, "SOLIDWORKS 2023", "SOLIDWORKS2023INSTALLED");
+			var tools2024 = CheckSolidworksInstalledSub(session, "SOLIDWORKS 2024", "SOLIDWORKS2024INSTALLED");
 			// SOLIDWORKSTOOLSPATH is unused(not installing solidworkstools.dll anymore)
 			// if (!string.IsNullOrEmpty(tools2022))
 			// 	session["SOLIDWORKSTOOLSPATH"] = tools2022;
