[3ds Max Exporter] 3ds Max 2025 support

diff --git a/Engine/Source/Programs/Enterprise/Datasmith/DatasmithMaxExporter/DatasmithMax2017.Build.cs b/Engine/Source/Programs/Enterprise/Datasmith/DatasmithMaxExporter/DatasmithMax2017.Build.cs
--- a/Engine/Source/Programs/Enterprise/Datasmith/DatasmithMaxExporter/DatasmithMax2017.Build.cs
+++ b/Engine/Source/Programs/Enterprise/Datasmith/DatasmithMaxExporter/DatasmithMax2017.Build.cs
@@ -70,6 +70,10 @@ namespace UnrealBuildTool.Rules
 					PublicAdditionalLibraries.Add(Path.Combine(LibraryPaths, "Maxscrpt.lib"));
 					PublicAdditionalLibraries.Add(Path.Combine(LibraryPaths, "mesh.lib"));
 				}
+				else
+				{
+					throw new BuildException("3ds Max Exporter - SDK doesn't exist {0}.", MaxSDKLocation);
+				}
 			}
 
 			// Itoo ForestPack/RailClone API
diff --git a/Engine/Source/Programs/Enterprise/Datasmith/DatasmithMaxExporter/DatasmithMax2025.Build.cs b/Engine/Source/Programs/Enterprise/Datasmith/DatasmithMaxExporter/DatasmithMax2025.Build.cs
new file mode 100644
--- /dev/null
+++ b/Engine/Source/Programs/Enterprise/Datasmith/DatasmithMaxExporter/DatasmithMax2025.Build.cs
@@ -0,0 +1,16 @@
+// Copyright Epic Games, Inc. All Rights Reserved.
+
+using System.Diagnostics;
+
+namespace UnrealBuildTool.Rules
+{
+	[SupportedPlatforms("Win64")]
+	public class DatasmithMax2025 : DatasmithMaxBase
+	{
+		public DatasmithMax2025(ReadOnlyTargetRules Target)
+			: base(Target)
+		{
+		}
+		public override string GetMaxVersion() { return "2025"; }
+	}
+}
\ No newline at end of file
diff --git a/Engine/Source/Programs/Enterprise/Datasmith/DatasmithMaxExporter/DatasmithMax2025.Target.cs b/Engine/Source/Programs/Enterprise/Datasmith/DatasmithMaxExporter/DatasmithMax2025.Target.cs
new file mode 100644
--- /dev/null
+++ b/Engine/Source/Programs/Enterprise/Datasmith/DatasmithMaxExporter/DatasmithMax2025.Target.cs
@@ -0,0 +1,15 @@
+// Copyright Epic Games, Inc. All Rights Reserved.
+
+using UnrealBuildTool;
+using System.Collections.Generic;
+
+public class DatasmithMax2025Target : DatasmithMaxBaseTarget
+{
+	public DatasmithMax2025Target(TargetInfo Target) : base(Target)
+	{
+		LaunchModuleName = "DatasmithMax2025";
+		ExeBinariesSubFolder = @"3DSMax\2025";
+
+		AddCopyPostBuildStep(Target);
+	}
+}
diff --git a/Engine/Source/Programs/Enterprise/Datasmith/DatasmithMaxExporter/Plugin/PackageContents.xml b/Engine/Source/Programs/Enterprise/Datasmith/DatasmithMaxExporter/Plugin/PackageContents.xml
--- a/Engine/Source/Programs/Enterprise/Datasmith/DatasmithMaxExporter/Plugin/PackageContents.xml
+++ b/Engine/Source/Programs/Enterprise/Datasmith/DatasmithMaxExporter/Plugin/PackageContents.xml
@@ -16,7 +16,12 @@
 	<!--  todo: Icon="./DatasmithIcon32.png" -->
     <CompanyDetails Name="EpicGames, Inc." Url="www.epicgames.com"/>
 
-    <RuntimeRequirements OS="Win64" Platform="3ds Max" SeriesMin="2017" SeriesMax="2024" />
+    <RuntimeRequirements OS="Win64" Platform="3ds Max" SeriesMin="2017" SeriesMax="2025" />
+
+    <Components Description="plugins parts">
+        <RuntimeRequirements OS="Win64" Platform="3ds Max" SeriesMin="2025" SeriesMax="2025" />
+        <ComponentEntry ModuleName="./2025/DatasmithMax2025.gup" />
+    </Components>
 
     <Components Description="plugins parts">
         <RuntimeRequirements OS="Win64" Platform="3ds Max" SeriesMin="2024" SeriesMax="2024" />
@@ -59,18 +64,18 @@
     </Components>
 
     <Components Description="pre-start-up scripts parts">
-        <RuntimeRequirements OS="Win64" Platform="3ds Max" SeriesMin="2017" SeriesMax="2024" />
+        <RuntimeRequirements OS="Win64" Platform="3ds Max" SeriesMin="2017" SeriesMax="2025" />
         <ComponentEntry ModuleName="./UI/DatasmithPreStartUp.ms" />
         <ComponentEntry ModuleName="./UI/UnrealStudio_Datasmith_Attributes_Holder_Modifier.ms" />
     </Components>
 
     <Components Description="macroscripts parts">
-        <RuntimeRequirements OS="Win64" Platform="3ds Max" SeriesMin="2017" SeriesMax="2024" />
+        <RuntimeRequirements OS="Win64" Platform="3ds Max" SeriesMin="2017" SeriesMax="2025" />
         <ComponentEntry ModuleName="./UI/DatasmithMacros.mcr" />
     </Components>
 
     <Components Description="post-start-up scripts parts">
-        <RuntimeRequirements OS="Win64" Platform="3ds Max" SeriesMin="2017" SeriesMax="2024" />
+        <RuntimeRequirements OS="Win64" Platform="3ds Max" SeriesMin="2017" SeriesMax="2025" />
         <ComponentEntry ModuleName="./UI/DatasmithPostStartUp.ms" />
     </Components>
 
diff --git a/Engine/Source/Programs/Enterprise/Datasmith/DatasmithMaxExporter/Private/DatasmithMaxDirectLinkMaterial.cpp b/Engine/Source/Programs/Enterprise/Datasmith/DatasmithMaxExporter/Private/DatasmithMaxDirectLinkMaterial.cpp
--- a/Engine/Source/Programs/Enterprise/Datasmith/DatasmithMaxExporter/Private/DatasmithMaxDirectLinkMaterial.cpp
+++ b/Engine/Source/Programs/Enterprise/Datasmith/DatasmithMaxExporter/Private/DatasmithMaxDirectLinkMaterial.cpp
@@ -424,11 +424,16 @@ void FMaterialsCollectionTracker::ConvertMaterial(Mtl* Material, TSharedRef<IDat
 	TSharedPtr<IDatasmithBaseMaterialElement> DatasmithMaterial;
 	if (FDatasmithMaxMaterialsToUEPbr* MaterialConverter = FDatasmithMaxMaterialsToUEPbrManager::GetMaterialConverter(Material))
 	{
-		 
 		MaterialConverter->Convert(DatasmithScene, DatasmithMaterial, Material, AssetsPath);
-
 		AddDatasmithMaterialForUsedMaterial(DatasmithScene, Material, DatasmithMaterial);
 	}
+	else
+	{
+		MSTR Classname;
+		Material->GetClassName(Classname);
+		FString MaterialDesc = FString::Printf(TEXT("\"%s\" of type <%s> (0x%08x-0x%08x)"), Material->GetName().data(), Classname.data(), Material->ClassID().PartA(), Material->ClassID().PartB());
+		LogWarning(FString(TEXT("Material not supported: ")) + MaterialDesc);
+	}
 
 	// Tie texture used by an actual material to tracked material
 	for (Texmap* Tex : TexmapsUsedByMaterial)
diff --git a/Engine/Source/Programs/Enterprise/Datasmith/DatasmithMaxExporter/Private/DatasmithMaxExporterWithDirectLink.cpp b/Engine/Source/Programs/Enterprise/Datasmith/DatasmithMaxExporter/Private/DatasmithMaxExporterWithDirectLink.cpp
--- a/Engine/Source/Programs/Enterprise/Datasmith/DatasmithMaxExporter/Private/DatasmithMaxExporterWithDirectLink.cpp
+++ b/Engine/Source/Programs/Enterprise/Datasmith/DatasmithMaxExporter/Private/DatasmithMaxExporterWithDirectLink.cpp
@@ -623,7 +623,11 @@ public:
 
 	FString GetConfigPath()
 	{
+#if MAX_PRODUCT_YEAR_NUMBER >= 2025
+		FString PlugCfgPath = (TCHAR*)GetCOREInterface()->GetDir(APP_PLUGCFG_DIR).data();
+#else
 		FString PlugCfgPath = GetCOREInterface()->GetDir(APP_PLUGCFG_DIR);
+#endif
 		return FPaths::Combine(PlugCfgPath, TEXT("UnrealDatasmithMax.ini"));
 	}
 
diff --git a/Engine/Source/Programs/Enterprise/Datasmith/DatasmithMaxExporter/Private/DatasmithMaxScanline.cpp b/Engine/Source/Programs/Enterprise/Datasmith/DatasmithMaxExporter/Private/DatasmithMaxScanline.cpp
--- a/Engine/Source/Programs/Enterprise/Datasmith/DatasmithMaxExporter/Private/DatasmithMaxScanline.cpp
+++ b/Engine/Source/Programs/Enterprise/Datasmith/DatasmithMaxExporter/Private/DatasmithMaxScanline.cpp
@@ -141,7 +141,7 @@ FString FDatasmithMaxMatWriter::GetActualBitmapPath(BitmapInfo* InBitmapInfo)
 		return FString();
 	}
 
-	FString ActualBitmapPath = FDatasmithMaxSceneExporter::GetActualPath(InBitmapInfo->GetPathEx().GetString());
+	FString ActualBitmapPath = FDatasmithMaxSceneExporter::GetActualPath(InBitmapInfo->GetPathEx().GetCStr());
 	
 	if (ActualBitmapPath.IsEmpty())
 	{
