[Rhino Exporter] Rhino8 support

diff --git a/Engine/Source/Programs/Enterprise/Datasmith/DatasmithRhinoExporter/DatasmithRhino6.Target.cs b/Engine/Source/Programs/Enterprise/Datasmith/DatasmithRhinoExporter/DatasmithRhino6.Target.cs
--- a/Engine/Source/Programs/Enterprise/Datasmith/DatasmithRhinoExporter/DatasmithRhino6.Target.cs
+++ b/Engine/Source/Programs/Enterprise/Datasmith/DatasmithRhinoExporter/DatasmithRhino6.Target.cs
@@ -138,7 +138,7 @@ public abstract class DatasmithRhinoBaseTarget : TargetRules
 		PostBuildSteps.Add(string.Format(@"sed -i '' 's/""<AssemblyFileVersion>""/""{1}""/' ""{0}""", CustomAssemblyInfoPath, VersionString));
 	}
 
-	public string GetRhinoThirdPartyFolder()
+	public virtual string GetRhinoThirdPartyFolder()
 	{
 		return Path.Combine("$(EngineDir)", "Restricted", "NotForLicensees", "Source", "ThirdParty", "Enterprise", "RhinoCommonSDK_" + GetVersion());
 	}
diff --git a/Engine/Source/Programs/Enterprise/Datasmith/DatasmithRhinoExporter/DatasmithRhino8.Build.cs b/Engine/Source/Programs/Enterprise/Datasmith/DatasmithRhinoExporter/DatasmithRhino8.Build.cs
new file mode 100644
--- /dev/null
+++ b/Engine/Source/Programs/Enterprise/Datasmith/DatasmithRhinoExporter/DatasmithRhino8.Build.cs
@@ -0,0 +1,20 @@
+// Copyright Epic Games, Inc. All Rights Reserved.
+
+using System.IO;
+
+namespace UnrealBuildTool.Rules
+{
+	[SupportedPlatforms("Win64", "Mac")]
+	public class DatasmithRhino8 : DatasmithRhinoBase
+	{
+		public DatasmithRhino8(ReadOnlyTargetRules Target)
+			: base(Target)
+		{
+		}
+
+		public override string GetRhinoVersion()
+		{
+			return "8";
+		}
+	}
+}
diff --git a/Engine/Source/Programs/Enterprise/Datasmith/DatasmithRhinoExporter/DatasmithRhino8.Target.cs b/Engine/Source/Programs/Enterprise/Datasmith/DatasmithRhinoExporter/DatasmithRhino8.Target.cs
new file mode 100644
--- /dev/null
+++ b/Engine/Source/Programs/Enterprise/Datasmith/DatasmithRhinoExporter/DatasmithRhino8.Target.cs
@@ -0,0 +1,39 @@
+// Copyright Epic Games, Inc. All Rights Reserved.
+
+using UnrealBuildTool;
+using System;
+using System.IO;
+
+public class DatasmithRhino8Target : DatasmithRhinoBaseTarget
+{
+	public DatasmithRhino8Target(TargetInfo Target)
+		: base(Target)
+	{
+	}
+
+	public override string GetVersion() { return "8"; }
+
+	public override string GetRhinoInstallFolderWindows()
+	{
+		try
+		{
+			return OperatingSystem.IsWindows() ? Microsoft.Win32.Registry.GetValue(@"HKEY_LOCAL_MACHINE\SOFTWARE\McNeel\Rhinoceros\8.0\Install", "Path", "") as string : null;
+		}
+		catch(Exception)
+		{
+			return "";
+		}
+	}
+
+	public override string GetRhinoInstallFolderMac()
+	{
+		return "/Applications/Rhino 8.app/Contents/Frameworks/RhCore.framework/Versions/Current/Resources/";
+	}
+
+	public override string GetRhinoThirdPartyFolder()
+	{
+		// Use SDK for Rhino 7 so that we don't have to build against .Net 7
+		return Path.Combine("$(EngineDir)", "Restricted", "NotForLicensees", "Source", "ThirdParty", "Enterprise", "RhinoCommonSDK_7");
+	}
+
+}
diff --git a/Engine/Source/Programs/Enterprise/Datasmith/DatasmithRhinoExporter/DatasmithRhino8/DatasmithRhino8.cs b/Engine/Source/Programs/Enterprise/Datasmith/DatasmithRhinoExporter/DatasmithRhino8/DatasmithRhino8.cs
new file mode 100644
--- /dev/null
+++ b/Engine/Source/Programs/Enterprise/Datasmith/DatasmithRhinoExporter/DatasmithRhino8/DatasmithRhino8.cs
@@ -0,0 +1,19 @@
+// Copyright Epic Games, Inc. All Rights Reserved.
+
+namespace DatasmithRhino
+{
+	///<summary>
+	/// <para>Every RhinoCommon .rhp assembly must have one and only one PlugIn-derived
+	/// class. DO NOT create instances of this class yourself. It is the
+	/// responsibility of Rhino to create an instance of this class.</para>
+	/// <para>To complete plug-in information, please also see all PlugInDescription
+	/// attributes in AssemblyInfo.cs (you might need to click "Project" ->
+	/// "Show All Files" to see it in the "Solution Explorer" window).</para>
+	///</summary>
+	public class DatasmithRhino8 : DatasmithRhinoPlugin
+	{
+		public DatasmithRhino8() : base()
+		{
+		}
+	}
+}
\ No newline at end of file
diff --git a/Engine/Source/Programs/Enterprise/Datasmith/DatasmithRhinoExporter/DatasmithRhino8/DatasmithRhino8.csproj b/Engine/Source/Programs/Enterprise/Datasmith/DatasmithRhinoExporter/DatasmithRhino8/DatasmithRhino8.csproj
new file mode 100644
--- /dev/null
+++ b/Engine/Source/Programs/Enterprise/Datasmith/DatasmithRhinoExporter/DatasmithRhino8/DatasmithRhino8.csproj
@@ -0,0 +1,247 @@
+﻿<?xml version="1.0" encoding="utf-8"?>
+<Project ToolsVersion="4.0" DefaultTargets="Build" xmlns="http://schemas.microsoft.com/developer/msbuild/2003">
+  <PropertyGroup>
+    <Configuration Condition=" '$(Configuration)' == '' ">Debug</Configuration>
+    <Platform Condition=" '$(Platform)' == '' ">x64</Platform>
+    <AssemblyInfo Condition=" '$(AssemblyInfo)' == ''">..\Properties\DefaultAssemblyInfo.cs</AssemblyInfo>
+    <Platform_OS Condition=" '$(Platform_OS)' == '' ">Win64</Platform_OS>
+    <!--If the evironment variable is not define, try to find the Rhino files on build machines-->
+    <ProductVersion>8.0.30703</ProductVersion>
+    <SchemaVersion>2.0</SchemaVersion>
+    <ProjectGuid>{D1FDC795-B334-4933-B680-088119CDC6BB}</ProjectGuid>
+    <OutputType>Library</OutputType>
+    <AppDesignerFolder>Properties</AppDesignerFolder>
+    <RootNamespace>DatasmithRhino</RootNamespace>
+    <AssemblyName Condition="'$(Configuration)' == 'Debug'">DatasmithRhino8-Debug</AssemblyName>
+    <AssemblyName Condition="'$(Configuration)' == 'Shipping' OR '$(Configuration)' == 'Release'">DatasmithRhino8-Shipping</AssemblyName>
+    <AssemblyName Condition="'$(Configuration)' != 'Debug' AND '$(Configuration)' != 'Shipping' AND '$(Configuration)' != 'Release'">DatasmithRhino8</AssemblyName>
+    <TargetFrameworkVersion>v4.5</TargetFrameworkVersion>
+    <FileAlignment>512</FileAlignment>
+    <IsWebBootstrapper>false</IsWebBootstrapper>
+    <NuGetPackageImportStamp>
+    </NuGetPackageImportStamp>
+  </PropertyGroup>
+
+  <PropertyGroup Condition="'$(Configuration)|$(Platform)' == 'Debug|x64'">
+    <PlatformTarget>AnyCPU</PlatformTarget>
+    <DebugSymbols>true</DebugSymbols>
+    <DebugType>full</DebugType>
+    <Optimize>false</Optimize>
+    <OutputPath>..\..\..\..\..\..\Binaries\$(Platform_OS)\Rhino\8</OutputPath>
+    <DefineConstants>DEBUG;TRACE;RHINO_VERSION_8</DefineConstants>
+    <ErrorReport>prompt</ErrorReport>
+    <CodeAnalysisIgnoreBuiltInRules>false</CodeAnalysisIgnoreBuiltInRules>
+  </PropertyGroup>
+  <PropertyGroup Condition="'$(Configuration)|$(Platform)' == 'Development|x64'">
+    <PlatformTarget>AnyCPU</PlatformTarget>
+    <DebugType>pdbonly</DebugType>
+    <Optimize>true</Optimize>
+    <OutputPath>..\..\..\..\..\..\Binaries\$(Platform_OS)\Rhino\8</OutputPath>
+    <DefineConstants>TRACE;RHINO_VERSION_8</DefineConstants>
+    <ErrorReport>prompt</ErrorReport>
+    <WarningLevel>4</WarningLevel>
+  </PropertyGroup>
+  <PropertyGroup Condition="'$(Configuration)|$(Platform)' == 'Release|x64'">
+    <PlatformTarget>AnyCPU</PlatformTarget>
+    <DebugType>pdbonly</DebugType>
+    <Optimize>true</Optimize>
+    <OutputPath>..\..\..\..\..\..\Binaries\$(Platform_OS)\Rhino\8</OutputPath>
+    <DefineConstants>TRACE;RHINO_VERSION_8</DefineConstants>
+    <ErrorReport>prompt</ErrorReport>
+    <WarningLevel>4</WarningLevel>
+  </PropertyGroup>
+  <PropertyGroup Condition="'$(Platform_OS)' == 'Win64'">
+    <DefineConstants>$(DefineConstants);WIN64_OS</DefineConstants>
+  </PropertyGroup>
+  <PropertyGroup Condition="'$(Platform_OS)' == 'Mac'">
+    <DefineConstants>$(DefineConstants);MAC_OS</DefineConstants>
+  </PropertyGroup>
+  <PropertyGroup Condition="'$(DATASMITHRHINO_EXPERIMENTAL)' == 'TRUE'">
+    <DefineConstants>$(DefineConstants);DATASMITHRHINO_EXPERIMENTAL</DefineConstants>
+  </PropertyGroup>
+  <ItemGroup>
+    <Reference Include="System" />
+    <Reference Include="System.Core" />
+    <Reference Include="System.Drawing" />
+    <Reference Include="Microsoft.CSharp" />
+    <Reference Include="RhinoCommon">
+      <SpecificVersion>False</SpecificVersion>
+      <HintPath>$(RHINO8_PATH)\rhinocommon.dll</HintPath>
+      <Private>False</Private>
+    </Reference>
+    <Reference Include="Eto">
+      <SpecificVersion>False</SpecificVersion>
+      <HintPath>$(RHINO8_PATH)\Eto.dll</HintPath>
+      <Private>False</Private>
+    </Reference>
+    <Reference Include="Rhino.UI">
+      <SpecificVersion>False</SpecificVersion>
+      <HintPath>$(RHINO8_PATH)\Rhino.UI.dll</HintPath>
+      <Private>False</Private>
+    </Reference>
+  </ItemGroup>
+  <ItemGroup Condition="'$(Platform_OS)' == 'Mac'">
+    <Reference Include="RhinoMac">
+      <SpecificVersion>False</SpecificVersion>
+      <HintPath>$(RHINO8_PATH)\RhinoMac.dll</HintPath>
+      <Private>False</Private>
+    </Reference>
+  </ItemGroup>
+  <ItemGroup>
+    <Compile Include="DatasmithRhino8.cs" />
+    <Compile Include="$(AssemblyInfo)">
+      <Link>Properties\AssemblyInfo.cs</Link>
+    </Compile>
+    <!-- DatasmithCSharpFacade API, possibly change that. -->
+    <Compile Include="..\..\..\..\..\..\Binaries\$(Platform_OS)\DatasmithFacadeCSharp\Public\*.cs">
+      <Link>DatasmithSwigFacade\%(Filename)%(Extension)</Link>
+    </Compile>
+  </ItemGroup>
+  <ItemGroup Condition="'$(Platform_OS)' == 'Win64'">
+    <Compile Include="..\Private\**\*.cs">
+      <Link>DatasmithRhinoExporter\%(RecursiveDir)%(Filename)%(Extension)</Link>
+    </Compile>
+  </ItemGroup>
+  <!-- TODO Remove this explicit definition once .net uses msbuild on Mac -->
+  <ItemGroup Condition="'$(Platform_OS)' == 'Mac'">
+    <Compile Include="..\Private\Commands\*.cs">
+      <Link>DatasmithRhinoExporter\Commands\%(Filename)%(Extension)</Link>
+    </Compile>
+    <Compile Include="..\Private\DirectLink\*.cs">
+      <Link>DatasmithRhinoExporter\DirectLink\%(Filename)%(Extension)</Link>
+    </Compile>
+    <Compile Include="..\Private\ElementExporters\*.cs">
+      <Link>DatasmithRhinoExporter\ElementExporters\%(Filename)%(Extension)</Link>
+    </Compile>
+    <Compile Include="..\Private\ExportContext\*.cs">
+      <Link>DatasmithRhinoExporter\ExportContext\%(Filename)%(Extension)</Link>
+    </Compile>
+    <Compile Include="..\Private\Utils\*.cs">
+      <Link>DatasmithRhinoExporter\Utils\%(Filename)%(Extension)</Link>
+    </Compile>
+    <Compile Include="..\Private\*.cs">
+      <Link>DatasmithRhinoExporter\%(Filename)%(Extension)</Link>
+    </Compile>
+  </ItemGroup>
+  <ItemGroup>
+    <Compile Include="..\SharedResources\Localization\Resources.Designer.cs">
+      <Link>Properties\Localization\Resources.Designer.cs</Link>
+      <AutoGen>True</AutoGen>
+      <DesignTime>True</DesignTime>
+      <DependentUpon>Resources.resx</DependentUpon>
+    </Compile>
+    <Compile Include="..\SharedResources\Localization\Resources.de.Designer.cs">
+      <Link>Properties\Localization\Resources.Designer.cs</Link>
+      <AutoGen>True</AutoGen>
+      <DesignTime>True</DesignTime>
+      <DependentUpon>Resources.de.resx</DependentUpon>
+    </Compile>
+    <Compile Include="..\SharedResources\Localization\Resources.es.Designer.cs">
+      <Link>Properties\Localization\Resources.Designer.cs</Link>
+      <AutoGen>True</AutoGen>
+      <DesignTime>True</DesignTime>
+      <DependentUpon>Resources.es.resx</DependentUpon>
+    </Compile>
+    <Compile Include="..\SharedResources\Localization\Resources.fr.Designer.cs">
+      <Link>Properties\Localization\Resources.Designer.cs</Link>
+      <AutoGen>True</AutoGen>
+      <DesignTime>True</DesignTime>
+      <DependentUpon>Resources.fr.resx</DependentUpon>
+    </Compile>
+    <Compile Include="..\SharedResources\Localization\Resources.ja.Designer.cs">
+      <Link>Properties\Localization\Resources.Designer.cs</Link>
+      <AutoGen>True</AutoGen>
+      <DesignTime>True</DesignTime>
+      <DependentUpon>Resources.ja.resx</DependentUpon>
+    </Compile>
+    <Compile Include="..\SharedResources\Localization\Resources.ko.Designer.cs">
+      <Link>Properties\Localization\Resources.Designer.cs</Link>
+      <AutoGen>True</AutoGen>
+      <DesignTime>True</DesignTime>
+      <DependentUpon>Resources.ko.resx</DependentUpon>
+    </Compile>
+    <Compile Include="..\SharedResources\Localization\Resources.pt.Designer.cs">
+      <Link>Properties\Localization\Resources.Designer.cs</Link>
+      <AutoGen>True</AutoGen>
+      <DesignTime>True</DesignTime>
+      <DependentUpon>Resources.pt.resx</DependentUpon>
+    </Compile>
+    <Compile Include="..\SharedResources\Localization\Resources.zh.Designer.cs">
+      <Link>Properties\Localization\Resources.Designer.cs</Link>
+      <AutoGen>True</AutoGen>
+      <DesignTime>True</DesignTime>
+      <DependentUpon>Resources.zh.resx</DependentUpon>
+    </Compile>
+  </ItemGroup>
+  <ItemGroup>
+    <EmbeddedResource Include="..\SharedResources\UnrealEngine.ico" />
+    <EmbeddedResource Include="..\SharedResources\Localization\Resources.resx">
+      <Link>Properties\Localization\Resources.resx</Link>
+      <Generator>ResXFileCodeGenerator</Generator>
+      <LastGenOutput>Resources.Designer.cs</LastGenOutput>
+    </EmbeddedResource>
+    <EmbeddedResource Include="..\SharedResources\Localization\Resources.de.resx">
+      <Link>Properties\Localization\Resources.de-DE.resx</Link>
+      <Generator>ResXFileCodeGenerator</Generator>
+      <LastGenOutput>Resources.de.Designer.cs</LastGenOutput>
+    </EmbeddedResource>
+    <EmbeddedResource Include="..\SharedResources\Localization\Resources.es.resx">
+      <Link>Properties\Localization\Resources.es-ES.resx</Link>
+      <Generator>ResXFileCodeGenerator</Generator>
+      <LastGenOutput>Resources.es.Designer.cs</LastGenOutput>
+    </EmbeddedResource>
+    <EmbeddedResource Include="..\SharedResources\Localization\Resources.fr.resx">
+      <Link>Properties\Localization\Resources.fr-FR.resx</Link>
+      <Generator>ResXFileCodeGenerator</Generator>
+      <LastGenOutput>Resources.fr.Designer.cs</LastGenOutput>
+    </EmbeddedResource>
+    <EmbeddedResource Include="..\SharedResources\Localization\Resources.ja.resx">
+      <Link>Properties\Localization\Resources.ja-JP.resx</Link>
+      <Generator>ResXFileCodeGenerator</Generator>
+      <LastGenOutput>Resources.ja.Designer.cs</LastGenOutput>
+    </EmbeddedResource>
+    <EmbeddedResource Include="..\SharedResources\Localization\Resources.ko.resx">
+      <Link>Properties\Localization\Resources.ko-KR.resx</Link>
+      <Generator>ResXFileCodeGenerator</Generator>
+      <LastGenOutput>Resources.ko.Designer.cs</LastGenOutput>
+    </EmbeddedResource>
+    <EmbeddedResource Include="..\SharedResources\Localization\Resources.pt.resx">
+      <Link>Properties\Localization\Resources.pt-BR.resx</Link>
+      <Generator>ResXFileCodeGenerator</Generator>
+      <LastGenOutput>Resources.pt.Designer.cs</LastGenOutput>
+    </EmbeddedResource>
+    <EmbeddedResource Include="..\SharedResources\Localization\Resources.zh.resx">
+      <Link>Properties\Localization\Resources.zh-CN.resx</Link>
+      <Generator>ResXFileCodeGenerator</Generator>
+      <LastGenOutput>Resources.zh.Designer.cs</LastGenOutput>
+    </EmbeddedResource>
+  </ItemGroup>
+  <Import Project="$(MSBuildToolsPath)\Microsoft.CSharp.targets" />
+  <!-- To modify your build process, add your task inside one of the targets below and uncomment it. 
+       Other similar extension points exist, see Microsoft.Common.targets.
+  <Target Name="BeforeBuild">
+  </Target>
+  <Target Name="AfterBuild">
+  </Target>
+  -->
+  <PropertyGroup Condition="'$(Platform_OS)' == 'Win64'">
+    <PostBuildEvent>
+		Copy "$(TargetPath)" "$(TargetDir)$(TargetName).rhp"
+		Copy "$(MSBuildProjectDirectory)\..\SharedResources\DatasmithRhino.rui" "$(TargetDir)$(TargetName).rui"
+    </PostBuildEvent>
+  </PropertyGroup>
+  <PropertyGroup Condition="'$(Platform_OS)' == 'Mac'">
+    <PostBuildEvent>
+    sh "$(ProjectDir)../Installer/CreateMacPackages.sh" $(TargetName) 8
+    </PostBuildEvent>
+  </PropertyGroup>
+  <PropertyGroup>
+    <FallbackCulture>en-US</FallbackCulture>
+  </PropertyGroup>
+  <PropertyGroup Condition="'$(Configuration)' == 'Debug' And '$(Platform_OS)' == 'Win64'">
+    <StartProgram>$(RHINO8_PATH)\Rhino.exe</StartProgram>
+    <StartArguments>
+    </StartArguments>
+    <StartAction>Program</StartAction>
+  </PropertyGroup>
+</Project>
\ No newline at end of file
diff --git a/Engine/Source/Programs/Enterprise/Datasmith/DatasmithRhinoExporter/Installer/CreateMacPackages.sh b/Engine/Source/Programs/Enterprise/Datasmith/DatasmithRhinoExporter/Installer/CreateMacPackages.sh
--- a/Engine/Source/Programs/Enterprise/Datasmith/DatasmithRhinoExporter/Installer/CreateMacPackages.sh
+++ b/Engine/Source/Programs/Enterprise/Datasmith/DatasmithRhinoExporter/Installer/CreateMacPackages.sh
@@ -1,5 +1,8 @@
 #! /bin/sh
 
+set -e
+set -x
+
 PLUGIN_NAME=$1
 PLUGIN_VERSION=$2
 SCRIPT_DIR=$(builtin cd $(dirname $0); pwd)
@@ -21,6 +24,8 @@ cp "$DATASMITH_FACADE_BINARY_DIR/"*".dylib" "$PLUGIN_RHP_PATH"
 cp "$TBB_BINARY_DIR/"*".dylib" "$PLUGIN_RHP_PATH"
 cp "$PLUGIN_BINARY_DIR/$PLUGIN_NAME.dll" "$PLUGIN_RHP_PATH/$PLUGIN_NAME.rhp"
 cp "$PLUGIN_BASE_DIR/Config/DatasmithRhino.rhp.config" "$PLUGIN_RHP_PATH/$PLUGIN_NAME.rhp.config"
+# Starting Rhino 8, Rhino UI systems on Mac supports same .rui file for UI library as on Windows
+cp "$PLUGIN_BASE_DIR/SharedResources/DatasmithRhino.rui" "$PLUGIN_RHP_PATH/$PLUGIN_NAME.rui"
 
 mkdir "$PLUGIN_RHP_PATH/Resources"
 cp "$PLUGIN_BASE_DIR/SharedResources/DatasmithRhino.plist" "$PLUGIN_RHP_PATH/Resources/DatasmithRhino.plist"
diff --git a/Engine/Source/Programs/Enterprise/Datasmith/DatasmithRhinoExporter/Private/ExportContext/DatasmithRhinoExportContext.cs b/Engine/Source/Programs/Enterprise/Datasmith/DatasmithRhinoExporter/Private/ExportContext/DatasmithRhinoExportContext.cs
--- a/Engine/Source/Programs/Enterprise/Datasmith/DatasmithRhinoExporter/Private/ExportContext/DatasmithRhinoExportContext.cs
+++ b/Engine/Source/Programs/Enterprise/Datasmith/DatasmithRhinoExporter/Private/ExportContext/DatasmithRhinoExportContext.cs
@@ -23,7 +23,7 @@ namespace DatasmithRhino.ExportContext
 		public bool bIsInWorksession {
 			get {
 				//Only check for worksession on Windows, the feature is not implemented on Mac and calling the API throws exception.
-				return Environment.OSVersion.Platform == PlatformID.Win32NT 
+				return System.Environment.OSVersion.Platform == PlatformID.Win32NT 
 					&& RhinoDocument.Worksession != null
 					&& RhinoDocument.Worksession.ModelCount > 1;
 			}
@@ -829,13 +829,13 @@ namespace DatasmithRhino.ExportContext
 		private bool IsUnsupportedObject(RhinoObject InObject)
 		{
 //Disabling obsolete warning as GetRenderPrimitiveList() is deprecated since Rhino5 but as of Rhino7 no alternative exists.
-#pragma warning disable CS0612
+#pragma warning disable CS0612, CS0618
 			// Geometry objects without meshes are currently not supported, unless they are points.
 			return InObject.ComponentType == ModelComponentType.ModelGeometry 
 				&& !InObject.IsMeshable(MeshType.Render)
 				&& InObject.ObjectType != ObjectType.Point
-				&& (InObject.ObjectType != ObjectType.Curve || null == InObject.GetRenderPrimitiveList(ActiveViewportInfo, true));
-#pragma warning restore CS0612
+				&& (InObject.ObjectType != ObjectType.Curve || null == InObject.GetRenderPrimitiveList(ActiveViewportInfo, false));
+#pragma warning restore CS0612, CS0618
 		}
 
 		private void InstanciateDefinition(DatasmithActorInfo ParentNode, DatasmithActorInfo DefinitionNode)
@@ -1534,9 +1534,9 @@ namespace DatasmithRhino.ExportContext
 			if (ActiveViewportInfo != null)
 			{
 				//Disabling obsolete warning as GetRenderPrimitiveList() is deprecated since Rhino5 but as of Rhino7 no alternative exists.
-#pragma warning disable CS0612
+#pragma warning disable CS0612, CS0618
 				RenderMeshes = InRhinoObject.GetRenderPrimitiveList(ActiveViewportInfo, false)?.ToMeshArray();
-#pragma warning restore CS0612
+#pragma warning restore CS0612, CS0618
 			}
 			if (RenderMeshes == null)
 			{
