Fixed irrelevant logspam.

Undo single Engine/Source/Programs/Enterprise/Datasmith/DatasmithMaxExporter/DatasmithMax2017.Build.cs revision 8

diff --git a/Engine/Source/Programs/Enterprise/Datasmith/DatasmithMaxExporter/DatasmithMax2017.Build.cs b/Engine/Source/Programs/Enterprise/Datasmith/DatasmithMaxExporter/DatasmithMax2017.Build.cs
--- a/Engine/Source/Programs/Enterprise/Datasmith/DatasmithMaxExporter/DatasmithMax2017.Build.cs
+++ b/Engine/Source/Programs/Enterprise/Datasmith/DatasmithMaxExporter/DatasmithMax2017.Build.cs
@@ -70,10 +70,6 @@ namespace UnrealBuildTool.Rules
 					PublicAdditionalLibraries.Add(Path.Combine(LibraryPaths, "Maxscrpt.lib"));
 					PublicAdditionalLibraries.Add(Path.Combine(LibraryPaths, "mesh.lib"));
 				}
-				else
-				{
-					throw new BuildException("3ds Max Exporter - SDK doesn't exist {0}.", MaxSDKLocation);
-				}
 			}
 
 			// Itoo ForestPack/RailClone API
