[Sketchup] Add support for Sketchup 2024

diff --git a/Engine/Source/Programs/Enterprise/Datasmith/DatasmithSketchUpRubyExporter/Build/build_mac.sh b/Engine/Source/Programs/Enterprise/Datasmith/DatasmithSketchUpRubyExporter/Build/build_mac.sh
--- a/Engine/Source/Programs/Enterprise/Datasmith/DatasmithSketchUpRubyExporter/Build/build_mac.sh
+++ b/Engine/Source/Programs/Enterprise/Datasmith/DatasmithSketchUpRubyExporter/Build/build_mac.sh
@@ -93,6 +93,7 @@ BuildSketchUpPlugin 2020 SDK_Mac_2020-2-171 "x86_64" "$Intermediatex86_64Path"
 BuildSketchUpPlugin 2021 SDK_Mac_2021-0-338 "x86_64" "$Intermediatex86_64Path"
 BuildSketchUpPlugin 2022 SDK_Mac_2022-0-353 "x86_64 arm64" "$IntermediateUniversalPath"
 BuildSketchUpPlugin 2023 SDK_Mac_2023-0-366 "x86_64 arm64" "$IntermediateUniversalPath"
+BuildSketchUpPlugin 2024 SDK_Mac_2024-0-483 "x86_64 arm64" "$IntermediateUniversalPath"
 
 
 # install_name_tool -change @rpath/DatasmithSDK.dylib @loader_path/Dylibs/DatasmithSDK.dylib DatasmithSketchUp.bundle 
diff --git a/Engine/Source/Programs/Enterprise/Datasmith/DatasmithSketchUpRubyExporter/DatasmithSketchUpRuby2024.Build.cs b/Engine/Source/Programs/Enterprise/Datasmith/DatasmithSketchUpRubyExporter/DatasmithSketchUpRuby2024.Build.cs
new file mode 100644
--- /dev/null
+++ b/Engine/Source/Programs/Enterprise/Datasmith/DatasmithSketchUpRubyExporter/DatasmithSketchUpRuby2024.Build.cs
@@ -0,0 +1,35 @@
+// Copyright Epic Games, Inc. All Rights Reserved.
+
+using System.IO;
+
+namespace UnrealBuildTool.Rules
+{
+
+	[SupportedPlatforms("Win64", "Mac")]
+	public class DatasmithSketchUpRuby2024 : DatasmithSketchUpRubyBase
+	{
+		public DatasmithSketchUpRuby2024(ReadOnlyTargetRules Target)
+			: base(Target)
+		{
+OptimizeCode = CodeOptimization.Never;
+bUseUnity = false;
+PCHUsage = PCHUsageMode.NoPCHs;
+
+			PrivateDefinitions.Add("SKP_SDK_2024");
+		}
+
+		public override string GetSketchUpSDKFolder()
+		{
+			return "SDK_WIN_x64_2024-0-484";
+		}
+
+		public override string GetSketchUpEnvVar()
+		{
+			return "SKP_SDK_2024";
+		}
+		public override string GetRubyLibName()
+		{
+			return "x64-ucrt-ruby320.lib";
+		}
+	}
+}
diff --git a/Engine/Source/Programs/Enterprise/Datasmith/DatasmithSketchUpRubyExporter/DatasmithSketchUpRuby2024.Target.cs b/Engine/Source/Programs/Enterprise/Datasmith/DatasmithSketchUpRubyExporter/DatasmithSketchUpRuby2024.Target.cs
new file mode 100644
--- /dev/null
+++ b/Engine/Source/Programs/Enterprise/Datasmith/DatasmithSketchUpRubyExporter/DatasmithSketchUpRuby2024.Target.cs
@@ -0,0 +1,16 @@
+// Copyright Epic Games, Inc. All Rights Reserved.
+
+using UnrealBuildTool;
+
+[SupportedPlatforms("Win64")]
+public class DatasmithSketchUpRuby2024Target : DatasmithSketchUpRubyBaseTarget
+{
+	public DatasmithSketchUpRuby2024Target(TargetInfo Target)
+		: base(Target)
+	{
+		LaunchModuleName = "DatasmithSketchUpRuby2024";
+		ExeBinariesSubFolder = @"SketchUpRuby/2024";
+
+		AddCopyPostBuildStep(Target);
+	}
+}
diff --git a/Engine/Source/Programs/Enterprise/Datasmith/DatasmithSketchUpRubyExporter/Private/DatasmithSketchUpExporter.cpp b/Engine/Source/Programs/Enterprise/Datasmith/DatasmithSketchUpRubyExporter/Private/DatasmithSketchUpExporter.cpp
--- a/Engine/Source/Programs/Enterprise/Datasmith/DatasmithSketchUpRubyExporter/Private/DatasmithSketchUpExporter.cpp
+++ b/Engine/Source/Programs/Enterprise/Datasmith/DatasmithSketchUpRubyExporter/Private/DatasmithSketchUpExporter.cpp
@@ -42,6 +42,14 @@
 #pragma warning(disable: 4191)
 // disable(SU2019 & SU2020): 'register' is no longer a supported storage class	
 #pragma warning(disable: 5033)
+// disable(SU2024): '_Header_cstdbool': warning STL4004: <ccomplex>, <cstdalign>, <cstdbool>, and <ctgmath> are deprecated in C++17. You can define _SILENCE_CXX17_C_HEADER_DEPRECATION_WARNING or _SILENCE_ALL_CXX17_DEPRECATION_WARNINGS to suppress this warning.
+#pragma warning(disable: 4996)
+// disable(SU2024): Implicit conversion from 'VALUE' to bool. Possible information loss
+#pragma warning(disable: 4800)
+// disable(SU2024): Dereferencing NULL pointer
+#pragma warning(disable: 6011)
+
+
 #undef DEPRECATED
 #include <ruby.h>
 #pragma warning(pop)
