Fixed compilation error of DatasmithSDKSample

#rnx

diff --git a/Engine/Source/Programs/Enterprise/Datasmith/DatasmithSDK/DatasmithSDK.Target.cs b/Engine/Source/Programs/Enterprise/Datasmith/DatasmithSDK/DatasmithSDK.Target.cs
--- a/Engine/Source/Programs/Enterprise/Datasmith/DatasmithSDK/DatasmithSDK.Target.cs
+++ b/Engine/Source/Programs/Enterprise/Datasmith/DatasmithSDK/DatasmithSDK.Target.cs
@@ -151,9 +151,11 @@ public class DatasmithSDKTarget : TargetRules
 			(@"$(EngineDir)\Source\Runtime\TraceLog\Public\", Path.Combine(DstDatasmithSDK,"Private\\"), "*.h"),
 			(@"$(EngineDir)\Source\Runtime\TraceLog\Public\", Path.Combine(DstDatasmithSDK,"Private\\"), "*.inl"),
 			(@"$(EngineDir)\Source\Runtime\Messaging\Public\", Path.Combine(DstDatasmithSDK,"Private\\"), "*.h"),
+			(@"$(EngineDir)\Source\Runtime\Messaging\Public\", Path.Combine(DstDatasmithSDK,"Private\\"), "*.inl"),
 			(@"$(EngineDir)\Source\Runtime\Core\Public\", Path.Combine(DstDatasmithSDK,"Private\\"), "*.h"),
 			(@"$(EngineDir)\Source\Runtime\Core\Public\", Path.Combine(DstDatasmithSDK,"Private\\"), "*.inl"),
 			(@"$(EngineDir)\Source\Runtime\CoreUObject\Public\", Path.Combine(DstDatasmithSDK,"Private\\"), "*.h"),
+			(@"$(EngineDir)\Source\Runtime\CoreUObject\Public\", Path.Combine(DstDatasmithSDK,"Private\\"), "*.inl"),
 		};
 
 		foreach ((string, string, string) ToCopy in FilesToCopy)
