[3ds Max Exporter] Compilation Warning fix (FLinearColor initialized)

diff --git a/Engine/Source/Programs/Enterprise/Datasmith/DatasmithMaxExporter/Private/DatasmithMaxSceneExporter.cpp b/Engine/Source/Programs/Enterprise/Datasmith/DatasmithMaxExporter/Private/DatasmithMaxSceneExporter.cpp
--- a/Engine/Source/Programs/Enterprise/Datasmith/DatasmithMaxExporter/Private/DatasmithMaxSceneExporter.cpp
+++ b/Engine/Source/Programs/Enterprise/Datasmith/DatasmithMaxExporter/Private/DatasmithMaxSceneExporter.cpp
@@ -1410,7 +1410,7 @@ void FDatasmithMaxSceneExporter::ParseSun(INode* Node, TSharedRef<IDatasmithLigh
 	{
 		int32 CoronaSunColorMode = 0;
 		float CoronaSunTemperature = 0;
-		FLinearColor CoronaSunColor;
+		FLinearColor CoronaSunColor(ForceInitToZero);
 
 		const int NumParamBlocks = Light.NumParamBlocks();
 
