[3ds Max Exporter] Fixed compilation warning after certain deprecations in slate headers

diff --git a/Engine/Source/Programs/Enterprise/Datasmith/DatasmithMaxExporter/Private/UI/Datasmith3dsMaxUI.cpp b/Engine/Source/Programs/Enterprise/Datasmith/DatasmithMaxExporter/Private/UI/Datasmith3dsMaxUI.cpp
--- a/Engine/Source/Programs/Enterprise/Datasmith/DatasmithMaxExporter/Private/UI/Datasmith3dsMaxUI.cpp
+++ b/Engine/Source/Programs/Enterprise/Datasmith/DatasmithMaxExporter/Private/UI/Datasmith3dsMaxUI.cpp
@@ -158,7 +158,6 @@ public:
 			.VAlign( VAlign_Fill )
 			[
 				SAssignNew( MessagesView, SListView<TSharedPtr<FMessageData>> )
-				.ItemHeight(32)
 				.ListItemsSource( &Messages )
 				.OnGenerateRow( this, &SMessagesWidget::OnGenerateRow )
 				.SelectionMode(ESelectionMode::Multi)
