DatasmithSDK: Mac build fix on Horde

- SDK build copies(or rather tries to) multiple sets of files and some of the sets may not be present(e.g. Win32 folder on non-windows build) meaning that on Mac/Linux builds such copy step produces an error in the command file. When this operation is last in the list of steps sh script returns error code resulting in build failure.

- e.g. https://horde.devtools.epicgames.com/job/664c4c9bce5df5adab1912f2?step=6b04

diff --git a/Engine/Source/Programs/Enterprise/Datasmith/DatasmithSDK/DatasmithSDK.Target.cs b/Engine/Source/Programs/Enterprise/Datasmith/DatasmithSDK/DatasmithSDK.Target.cs
--- a/Engine/Source/Programs/Enterprise/Datasmith/DatasmithSDK/DatasmithSDK.Target.cs
+++ b/Engine/Source/Programs/Enterprise/Datasmith/DatasmithSDK/DatasmithSDK.Target.cs
@@ -89,7 +89,8 @@ public class DatasmithSDKTarget : TargetRules
 		SrcPath = SrcPath.Replace("\\", "/");
 		DestPath = DestPath.Replace("\\", "/");
 		PostBuildSteps.Add(string.Format("echo Copying \"{0}\" to {1}\n", SrcPath, DestPath));
-		PostBuildSteps.Add(string.Format("cp -R -f {0}/{2} {1}\n", SrcPath, DestPath, Files));
+		// || true to clear error code when source not found
+		PostBuildSteps.Add(string.Format("cp -R -f {0}/{2} {1} || true\n", SrcPath, DestPath, Files));
 	}
 
 	public void PostBuildCopyUnix(string SrcPath, string DestPath)
@@ -97,7 +98,8 @@ public class DatasmithSDKTarget : TargetRules
 		SrcPath = SrcPath.Replace("\\", "/");
 		DestPath = DestPath.Replace("\\", "/");
 		PostBuildSteps.Add(string.Format("echo Copying \"{0}\" to {1}\n", SrcPath, DestPath));
-		PostBuildSteps.Add(string.Format("cp -R -f {0} {1}\n", SrcPath, DestPath));
+		// || true to clear error code when source not found
+		PostBuildSteps.Add(string.Format("cp -R -f {0} {1} || true\n", SrcPath, DestPath));
 	}
 
 	public void AddPostBuildSteps()
