[Archicad exporter] Morph crash fix, morph instances fix, mirrored transform fix

- misplaced morph elements and mirrored objects issues as represented in UE-120861

- fixed UE-141483 which crashed in 5.4 (and was misplaced before)

diff --git a/Engine/Source/Programs/Enterprise/Datasmith/DatasmithARCHICADExporter/Private/Element2StaticMesh.cpp b/Engine/Source/Programs/Enterprise/Datasmith/DatasmithARCHICADExporter/Private/Element2StaticMesh.cpp
--- a/Engine/Source/Programs/Enterprise/Datasmith/DatasmithARCHICADExporter/Private/Element2StaticMesh.cpp
+++ b/Engine/Source/Programs/Enterprise/Datasmith/DatasmithARCHICADExporter/Private/Element2StaticMesh.cpp
@@ -55,6 +55,7 @@ FElement2StaticMesh::FElement2StaticMesh(const FSyncContext& InSyncContext)
 	: SyncContext(InSyncContext)
 	, bSomeHasTextures(false)
 	, BugsCount(0)
+	, GeometryShift(0, 0, 0)
 {
 }
 
@@ -190,14 +191,11 @@ void FElement2StaticMesh::AddVertex(GS::Int32 InBodyVertex, const Geometry::Vect
 	{
 		// Not already used, get value from body
 		CurrentBody.GetVertex(InBodyVertex, &vertex.LocalVertex, ModelerAPI::CoordinateSystem::ElemLocal);
-		if (!bIsIdentity)
-		{
-			Geometry::Point3D WorldPt = Vertex2Point3D(vertex.LocalVertex);
-			Geometry::Point3D LocalPt = World2Local.Apply(WorldPt);
-			vertex.LocalVertex.x = LocalPt.x;
-			vertex.LocalVertex.y = LocalPt.y;
-			vertex.LocalVertex.z = LocalPt.z;
-		}
+
+		vertex.LocalVertex.x = vertex.LocalVertex.x + GeometryShift.x;
+		vertex.LocalVertex.y = vertex.LocalVertex.y + GeometryShift.y;
+		vertex.LocalVertex.z = vertex.LocalVertex.z + GeometryShift.z;
+
 		vertex.LocalVertex.x *= SyncContext.ScaleLength;
 		vertex.LocalVertex.y *= SyncContext.ScaleLength;
 		vertex.LocalVertex.z *= SyncContext.ScaleLength;
@@ -252,8 +250,7 @@ void FElement2StaticMesh::AddVertex(GS::Int32 InBodyVertex, const Geometry::Vect
 		}
 	}
 
-	Geometry::Vector3D VertexWorldNormal = bIsIdentity ? VertexNormal : Matrix * VertexNormal;
-	FVector3f CurrentNormal(float(VertexWorldNormal.x), -float(VertexWorldNormal.y), float(VertexWorldNormal.z));
+	FVector3f CurrentNormal(float(VertexNormal.x), -float(VertexNormal.y), float(VertexNormal.z));
 
 	// Create triangles
 	if (VertexCount == 0)
@@ -339,7 +336,7 @@ void FElement2StaticMesh::InitPolygonMaterial()
 }
 
 void FElement2StaticMesh::AddElementGeometry(const ModelerAPI::Element&		   InModelElement,
-											 const Geometry::Transformation3D& InWorld2Local)
+											 const Geometry::Vector3D&          InGeometryShift)
 {
 #if 0
 	UE_AC_TraceF("Element\n%s\n", F3DElement2String::Element2String(InModelElement).c_str());
@@ -351,9 +348,7 @@ void FElement2StaticMesh::AddElementGeometry(const ModelerAPI::Element&		   InMo
 		UE_AC_DebugF("FElement2StaticMesh::AddElementGeometry - Break element found\n");
 	}
 #endif
-	World2Local = InWorld2Local;
-	Matrix = InWorld2Local.GetMatrix();
-	bIsIdentity = InWorld2Local.IsIdentity();
+	GeometryShift = InGeometryShift;
 
 	// Collect geometry from element's bodies
 	GS::Int32 NbBodies = InModelElement.GetTessellatedBodyCount();
diff --git a/Engine/Source/Programs/Enterprise/Datasmith/DatasmithARCHICADExporter/Private/Element2StaticMesh.h b/Engine/Source/Programs/Enterprise/Datasmith/DatasmithARCHICADExporter/Private/Element2StaticMesh.h
--- a/Engine/Source/Programs/Enterprise/Datasmith/DatasmithARCHICADExporter/Private/Element2StaticMesh.h
+++ b/Engine/Source/Programs/Enterprise/Datasmith/DatasmithARCHICADExporter/Private/Element2StaticMesh.h
@@ -61,7 +61,7 @@ class FElement2StaticMesh
 	TSharedPtr< IDatasmithMeshElement > CreateMesh();
 
 	// Collect geometry of the element
-	void AddElementGeometry(const ModelerAPI::Element& InModelElement, const Geometry::Transformation3D& InWorld2Local);
+	void AddElementGeometry(const ModelerAPI::Element& InModelElement, const Geometry::Vector3D& InGeometryShift);
 
 	// Return the numbers of bugs detected during conversion
 	unsigned int GetBugsCount() const { return BugsCount; }
@@ -126,9 +126,8 @@ class FElement2StaticMesh
 	// Current context
 	const FSyncContext& SyncContext;
 
-	Geometry::Transformation3D World2Local;
-	Geometry::Matrix33		   Matrix;
-	bool					   bIsIdentity = true;
+	// Trasform to 'shift' original mesh vertices in local space
+	Geometry::Vector3D GeometryShift;
 
 	// Working variables
 	bool				 bSomeHasTextures; // True if at least one triangles need uv
diff --git a/Engine/Source/Programs/Enterprise/Datasmith/DatasmithARCHICADExporter/Private/ElementID.cpp b/Engine/Source/Programs/Enterprise/Datasmith/DatasmithARCHICADExporter/Private/ElementID.cpp
--- a/Engine/Source/Programs/Enterprise/Datasmith/DatasmithARCHICADExporter/Private/ElementID.cpp
+++ b/Engine/Source/Programs/Enterprise/Datasmith/DatasmithARCHICADExporter/Private/ElementID.cpp
@@ -100,7 +100,11 @@ void FElementID::InitElement(FSyncData* IOSyncdata)
 	Index3D = IOSyncdata->GetIndex3D();
 	if (Index3D > 0)
 	{
-		SyncContext.GetModel().GetElement(Index3D, &Element3D);
+		GS::Int32 ElementCount = SyncContext.GetModel().GetElementCount();
+		if (ensureMsgf(ElementCount >= Index3D, TEXT("Element index became outside of elements range of the model")))
+		{
+			SyncContext.GetModel().GetElement(Index3D, &Element3D);
+		}
 	}
 	MeshClass = nullptr;
 	LibPartInfo = nullptr;
@@ -159,8 +163,9 @@ FMeshClass* FElementID::GetMeshClass()
 			MeshClass->ElementType = Element3D.GetType();
 			MeshClass->TransformCount = bTransformed ? 0 : 1;
 			SyncContext.GetSyncDatabase().AddInstance(HashValue, std::move(NewInstance));
+			
 			UE_AC_VerboseF("FElementID::GetMeshClass - First instance %u {%s}\n", MeshClass->Hash,
-						   APIGuidToString(ElementGuid).ToUtf8());
+						   APIGuidToString(APIElementHeader.guid).ToUtf8());
 		}
 		else
 		{
@@ -182,7 +187,7 @@ FMeshClass* FElementID::GetMeshClass()
 			++MeshClass->InstancesCount;
 			MeshClass->TransformCount += bTransformed ? 0 : 1;
 			UE_AC_VerboseF("FElementID::GetMeshClass - Reuse MeshClass %u {%s}\n", MeshClass->Hash,
-						   APIGuidToString(ElementGuid).ToUtf8());
+						   APIGuidToString(APIElementHeader.guid).ToUtf8());
 		}
 	}
 	return MeshClass;
@@ -246,7 +251,7 @@ void FElementID::CollectDependantElementsType(API_ElemTypeID TypeID) const
 		}
 		ChildSyncData->SetParent(SyncData);
 		UE_AC_VerboseF("FElementID::ConnectedElements %u %s -> %s\n", i, APIGuidToString(ConnectedElements[i]).ToUtf8(),
-					   SyncData->ElementId.ToUniString().ToUtf8());
+					   SyncData->GetId().ToUniString().ToUtf8());
 	}
 }
 
diff --git a/Engine/Source/Programs/Enterprise/Datasmith/DatasmithARCHICADExporter/Private/MaterialsDatabase.cpp b/Engine/Source/Programs/Enterprise/Datasmith/DatasmithARCHICADExporter/Private/MaterialsDatabase.cpp
--- a/Engine/Source/Programs/Enterprise/Datasmith/DatasmithARCHICADExporter/Private/MaterialsDatabase.cpp
+++ b/Engine/Source/Programs/Enterprise/Datasmith/DatasmithARCHICADExporter/Private/MaterialsDatabase.cpp
@@ -85,7 +85,7 @@ void FMaterialInfo::Init(const FSyncContext& SyncContext, const API_Component3D&
 									  SyncContext.GetTexturesCache().GetTexture(SyncContext, textureIndex).Fingerprint);
 			}
 			UE_AC_VerboseF("Simulate Guid for material %d, %s Guid=%s\n", MaterialKey.ACMaterialIndex,
-						   DisplayName.ToUtf8(), APIGuidToString(MatGuid).ToUtf8());
+						   TCHAR_TO_UTF8(*DatasmithLabel), APIGuidToString(MatGuid).ToUtf8());
 		}
 		SyncData.MaterialId = APIGuid2GSGuid(MatGuid);
 		SyncData.MaterialIndex = CUmat.umat.mater.head.index;
diff --git a/Engine/Source/Programs/Enterprise/Datasmith/DatasmithARCHICADExporter/Private/SyncData.cpp b/Engine/Source/Programs/Enterprise/Datasmith/DatasmithARCHICADExporter/Private/SyncData.cpp
--- a/Engine/Source/Programs/Enterprise/Datasmith/DatasmithARCHICADExporter/Private/SyncData.cpp
+++ b/Engine/Source/Programs/Enterprise/Datasmith/DatasmithARCHICADExporter/Private/SyncData.cpp
@@ -567,7 +567,7 @@ class FConvertGeometry2MeshElement : public FTaskMgr::FTask
 	FConvertGeometry2MeshElement(const FSyncContext& InSyncContext, FSyncData::FElement* InElementSyncData,
 								 FMeshClass* InMeshClass);
 
-	void AddElementGeometry(FElementID* IOElementID, const Geometry::Transformation3D& InLocalToWorld);
+	void AddElementGeometry(FElementID* IOElementID, const Geometry::Vector3D& InGeometryShift);
 
 	bool HasGeometry() const { return Element2StaticMesh.HasGeometry(); }
 
@@ -623,11 +623,12 @@ FConvertGeometry2MeshElement::FConvertGeometry2MeshElement(const FSyncContext&	I
 }
 
 void FConvertGeometry2MeshElement::AddElementGeometry(FElementID*						IOElementID,
-													  const Geometry::Transformation3D& InWorldToLocal)
+													  const Geometry::Vector3D&          InGeometryShift
+	)
 {
 	UE_AC_TestPtr(IOElementID);
 
-	Element2StaticMesh.AddElementGeometry(IOElementID->GetElement3D(), InWorldToLocal);
+	Element2StaticMesh.AddElementGeometry(IOElementID->GetElement3D(), InGeometryShift);
 }
 
 FSyncData::FElement::FElement(const GS::Guid& InGuid, const FSyncContext& /* InSyncContext */)
@@ -775,19 +776,29 @@ void FSyncData::FElement::Process(FProcessInfo* IOProcessInfo)
 
 			ModelerAPI::Transformation LocalToWorld =
 				IOProcessInfo->ElementID.GetElement3D().GetElemLocalToWorldTransformation();
-			Geometry::Transformation3D WorldToLocal; // Set 2 identity for Instances (i.e. Object with transform)
-			if ((LocalToWorld.status & TR_IDENT) != 0)
+			// Shift geometry to pivot it at the bounds center
+			Geometry::Vector3D GeometryShift;
 			{
-				Box3D Bounds = IOProcessInfo->ElementID.GetElement3D().GetBounds();
-				LocalToWorld.matrix[0][3] = (Bounds.xMin + Bounds.xMax) * 0.5;
-				LocalToWorld.matrix[1][3] = (Bounds.yMin + Bounds.yMax) * 0.5;
-				LocalToWorld.matrix[2][3] = Bounds.zMin;
-				LocalToWorld.status = (LocalToWorld.matrix[0][3] == 0.0 && LocalToWorld.matrix[1][3] == 0.0 &&
-									   LocalToWorld.matrix[2][3] == 0.0)
-										  ? TR_IDENT
-										  : TR_TRANSL_ONLY;
-				WorldToLocal.SetOffset(Geometry::Vector3D(-LocalToWorld.matrix[0][3], -LocalToWorld.matrix[1][3],
-														  -LocalToWorld.matrix[2][3]));
+				Box3D LocalBounds = IOProcessInfo->ElementID.GetElement3D().GetBounds(
+					ModelerAPI::CoordinateSystem::ElemLocal);
+				Geometry::Point3D LocalBoundsCenter{
+					(LocalBounds.xMin + LocalBounds.xMax) * 0.5,
+					(LocalBounds.yMin + LocalBounds.yMax) * 0.5,
+					LocalBounds.zMin};
+
+				// Transform center to world
+				TRANMAT LocalToWorldTranmatOrig;
+				LocalToWorld.ToTRANMAT(&LocalToWorldTranmatOrig);
+				Geometry::Point3D BoundsCenterWorld = Geometry::TransformPoint(LocalToWorldTranmatOrig, LocalBoundsCenter);
+
+				// "Re-pivot" object to the center of bounding box by
+				// ...shifting geometry to have its local zero coordinates at the geometry bounds center
+				GeometryShift = -LocalBoundsCenter;
+				// ...and changing transform to translate geometry zero point to bounds center in world space we computed
+				LocalToWorld.matrix[0][3] = BoundsCenterWorld[0];
+				LocalToWorld.matrix[1][3] = BoundsCenterWorld[1];
+				LocalToWorld.matrix[2][3] = BoundsCenterWorld[2];
+				LocalToWorld.status = BoundsCenterWorld.IsNullVector(EPS) ? TR_IDENT : TR_TRANSL_ONLY;
 			}
 
 			if (!ActorElement.IsValid())
@@ -833,7 +844,7 @@ void FSyncData::FElement::Process(FProcessInfo* IOProcessInfo)
 				MeshClass->Rotation = ActorElement->GetRotation();
 				FConvertGeometry2MeshElement* ConvertGeometry2MeshElement =
 					new FConvertGeometry2MeshElement(IOProcessInfo->SyncContext, this, MeshClass);
-				ConvertGeometry2MeshElement->AddElementGeometry(&IOProcessInfo->ElementID, WorldToLocal);
+				ConvertGeometry2MeshElement->AddElementGeometry(&IOProcessInfo->ElementID, GeometryShift);
 				if (ConvertGeometry2MeshElement->HasGeometry())
 				{
 					UE_AC_STAT(IOProcessInfo->SyncContext.Stats.TotalMeshClassesCreated++);
@@ -847,43 +858,6 @@ void FSyncData::FElement::Process(FProcessInfo* IOProcessInfo)
 				// delete ConvertGeometry2MeshElement;
 				FTaskMgr::GetMgr()->AddTask(ConvertGeometry2MeshElement, FTaskMgr::kSchedule);
 			}
-			else if(TypeID == API_MorphID && ((LocalToWorld.status & IsRelative) == IsRelative))
-			{
-				FTransform Parent(MeshClass->Rotation, MeshClass->Translation);
-				FTransform Child =  Parent * ActorElement->GetRelativeTransform();
-				ActorElement->SetTranslation(Child.GetTranslation());
-				ActorElement->SetRotation(Child.GetRotation());
-
-				//API_Element APIElement;
-				//Zap(&APIElement);
-				//APIElement.header.guid = GSGuid2APIGuid(ElementId);
-				//GSErrCode GSErr = ACAPI_Element_Get(&APIElement, 0);
-				//UE_AC_Assert(GSErr == NoError);
-				//UE_AC_Assert(APIElement.header.typeID == TypeID);
-
-				{
-					//API_Elem_Head	elemHead;
-					//BNZeroMemory(&elemHead, sizeof(API_Elem_Head));
-					//elemHead.guid = GSGuid2APIGuid(ElementId);
-
-					API_ElemInfo3D info3D;
-					BNZeroMemory(&info3D, sizeof(API_ElemInfo3D));
-
-					GSErrCode err = ACAPI_Element_Get3DInfo(IOProcessInfo->ElementID.GetHeader(), &info3D);
-					if (err == NoError)
-					{
-						//API_Component3D component;
-						//BNZeroMemory(&component, sizeof(component));
-						//component.header.typeID = API_BodyID;
-						//component.header.index = info3D.fbody;
-						//err = ACAPI_3D_GetComponent(&component);
-						//if (err == NoError) {
-						//	API_Tranmat			tm;
-						//	tm = component.body.tranmat;
-						//}
-					}
-				}
-			}
 		}
 	}
 }
diff --git a/Engine/Source/Programs/Enterprise/Datasmith/DatasmithARCHICADExporter/Private/Utils/APIEnvir.h b/Engine/Source/Programs/Enterprise/Datasmith/DatasmithARCHICADExporter/Private/Utils/APIEnvir.h
--- a/Engine/Source/Programs/Enterprise/Datasmith/DatasmithARCHICADExporter/Private/Utils/APIEnvir.h
+++ b/Engine/Source/Programs/Enterprise/Datasmith/DatasmithARCHICADExporter/Private/Utils/APIEnvir.h
@@ -54,6 +54,7 @@ _Pragma("clang diagnostic ignored \"-Wdefaulted-function-deleted\"")
 #include "ConvexPolygon.hpp"
 
 #include "Transformation.hpp"
+#include "TranMatToTransformationConverter.hpp"
 #include "Parameter.hpp"
 #include "Light.hpp"
 
