[SketchUp] Fixed lightting issue when geometry is transformed within group

diff --git a/Engine/Source/Programs/Enterprise/Datasmith/DatasmithSketchUpRubyExporter/DatasmithSketchUpRuby2024.Build.cs b/Engine/Source/Programs/Enterprise/Datasmith/DatasmithSketchUpRubyExporter/DatasmithSketchUpRuby2024.Build.cs
--- a/Engine/Source/Programs/Enterprise/Datasmith/DatasmithSketchUpRubyExporter/DatasmithSketchUpRuby2024.Build.cs
+++ b/Engine/Source/Programs/Enterprise/Datasmith/DatasmithSketchUpRubyExporter/DatasmithSketchUpRuby2024.Build.cs
@@ -11,10 +11,6 @@ namespace UnrealBuildTool.Rules
 		public DatasmithSketchUpRuby2024(ReadOnlyTargetRules Target)
 			: base(Target)
 		{
-OptimizeCode = CodeOptimization.Never;
-bUseUnity = false;
-PCHUsage = PCHUsageMode.NoPCHs;
-
 			PrivateDefinitions.Add("SKP_SDK_2024");
 		}
 
diff --git a/Engine/Source/Programs/Enterprise/Datasmith/DatasmithSketchUpRubyExporter/Private/DatasmithSketchUpMesh.cpp b/Engine/Source/Programs/Enterprise/Datasmith/DatasmithSketchUpRubyExporter/Private/DatasmithSketchUpMesh.cpp
--- a/Engine/Source/Programs/Enterprise/Datasmith/DatasmithSketchUpRubyExporter/Private/DatasmithSketchUpMesh.cpp
+++ b/Engine/Source/Programs/Enterprise/Datasmith/DatasmithSketchUpRubyExporter/Private/DatasmithSketchUpMesh.cpp
@@ -377,7 +377,7 @@ namespace DatasmithSketchUp
 		{
 			for (int32 Y = 0; Y < 4; ++Y)
 			{
-				TransformForNormals.values[X + Y*4] = TransformNoTranslation.values[Y + X*4];
+				TransformForNormals.values[X + Y*4] = TransformInverse.values[Y + X*4];
 			}
 		}
 
