[Archicad] Exporter accepts elements with External type without errors

- previously, export was raising an exception when this element was encountered. Breaking export or constantly re-displaying error dialog on auto-sync
- additionally, removed erro dialog popping up when exception was thrown from auto-sync - just logging the error

diff --git a/Engine/Source/Programs/Enterprise/Datasmith/DatasmithARCHICADExporter/Private/Palette.cpp b/Engine/Source/Programs/Enterprise/Datasmith/DatasmithARCHICADExporter/Private/Palette.cpp
--- a/Engine/Source/Programs/Enterprise/Datasmith/DatasmithARCHICADExporter/Private/Palette.cpp
+++ b/Engine/Source/Programs/Enterprise/Datasmith/DatasmithARCHICADExporter/Private/Palette.cpp
@@ -116,7 +116,7 @@ class FPaletteWindow : public DG::Palette,
 			if (--Count == 0)
 			{
 				Count = Delay;
-				TryFunctionCatchAndAlert("AutoSync - Check View State", []() -> GSErrCode {
+				TryFunctionCatchAndLog("AutoSync - Check View State", []() -> GSErrCode {
 					FSynchronizer* Synchronizer = FSynchronizer::GetCurrent();
 					if (Synchronizer != nullptr)
 					{
diff --git a/Engine/Source/Programs/Enterprise/Datasmith/DatasmithARCHICADExporter/Private/Utils/ElementTools.cpp b/Engine/Source/Programs/Enterprise/Datasmith/DatasmithARCHICADExporter/Private/Utils/ElementTools.cpp
--- a/Engine/Source/Programs/Enterprise/Datasmith/DatasmithARCHICADExporter/Private/Utils/ElementTools.cpp
+++ b/Engine/Source/Programs/Enterprise/Datasmith/DatasmithARCHICADExporter/Private/Utils/ElementTools.cpp
@@ -162,6 +162,14 @@ bool FElementTools::GetInfoString(const API_Guid& InGUID, GS::UniString* OutStri
 // Tool: Return the localize name for element type id
 const utf8_t* FElementTools::TypeName(API_ElemTypeID InElementType)
 {
+#if AC_VERSION >= 26 
+	if (API_ExternalElemID == InElementType)
+	{
+		// todo: might be possible to extract more information on the exact external type used
+		return "External";	
+	}
+#endif
+
 	UE_AC_Assert(API_FirstElemType <= InElementType && InElementType <= API_LastElemType);
 
 	static const utf8_t* TypeNames[API_LastElemType + 1] = {};
