[Archicad] Fixel area lights orientation

- Archicad window lights had their height/width inverted. Repro: create a Window light, set it's Width/Height to different values. Export/Sync - notice how exported Width/Height are oriented.
- AreaLight rotation was broken. Repro: create an AreaLight(it should lay in horizontal plane), rotate it in horizontal plane, Export/Sync - notice how rotation is exported

diff --git a/Engine/Source/Programs/Enterprise/Datasmith/DatasmithARCHICADExporter/Private/GeometryUtil.cpp b/Engine/Source/Programs/Enterprise/Datasmith/DatasmithARCHICADExporter/Private/GeometryUtil.cpp
--- a/Engine/Source/Programs/Enterprise/Datasmith/DatasmithARCHICADExporter/Private/GeometryUtil.cpp
+++ b/Engine/Source/Programs/Enterprise/Datasmith/DatasmithARCHICADExporter/Private/GeometryUtil.cpp
@@ -3,6 +3,7 @@
 #include "GeometryUtil.h"
 
 #include "Vector3D.hpp"
+#include "Math/RotationMatrix.h"
 
 BEGIN_NAMESPACE_UE_AC
 
@@ -90,22 +91,29 @@ FQuat FGeometryUtil::GetRotationQuat(const double Matrix[3][4])
 	return FQuat(FVector(float(RotAxis.x), float(-RotAxis.y), float(RotAxis.z)), float(RotAngle)).Inverse();
 }
 
-// Return the Quat equivalent to the direction vector
-FQuat FGeometryUtil::GetRotationQuat(const ModelerAPI::Vector& Direction)
+// Convert Archicad direction vector to Unreal one
+FVector GetDirectionVector(const ModelerAPI::Vector& Vec)
 {
-	const Geometry::Vector3< double > DefaultDirVec(1.0, 0.0, 0.0);
-	Geometry::Vector3< double >		  DirVec(Direction.x, -Direction.y, Direction.z);
-	DirVec.NormalizeVector();
-
-	const double distToDirSqr = (DirVec - DefaultDirVec).GetLengthSqr();
-	const double RotAngle = acos((2.0 - distToDirSqr) * 0.5); // Rotation angle in radian
-
-	Geometry::Vector3< double > RotAxis = DefaultDirVec ^ DirVec;
-	RotAxis.NormalizeVector();
+	return FVector(Vec.x, -Vec.y, Vec.z);
+}
 
-	return FQuat(FVector(float(RotAxis.x), float(RotAxis.y), float(RotAxis.z)), float(RotAngle));
+// Return the Quat equivalent of rotation defined by Direction and Up vectors
+FQuat FGeometryUtil::GetRotationQuat(const ModelerAPI::Vector& Direction, const ModelerAPI::Vector& Up)
+{
+	FVector DirVec = GetDirectionVector(Direction);
+	if (!DirVec.Normalize())
+	{
+		return FQuat::Identity;
+	}
+	FVector UpVec = GetDirectionVector(Up);
+	if (!UpVec.Normalize())
+	{
+		return FQuat::Identity;
+	}
+	return FQuat(FRotationMatrix::MakeFromXZ(DirVec, UpVec));
 }
 
+
 // Convert Archicad camera rotation to an Unreal Quat
 FQuat FGeometryUtil::GetRotationQuat(const double PitchInDegrees, const double YawInDegrees, const double RollInDegrees)
 {
diff --git a/Engine/Source/Programs/Enterprise/Datasmith/DatasmithARCHICADExporter/Private/GeometryUtil.h b/Engine/Source/Programs/Enterprise/Datasmith/DatasmithARCHICADExporter/Private/GeometryUtil.h
--- a/Engine/Source/Programs/Enterprise/Datasmith/DatasmithARCHICADExporter/Private/GeometryUtil.h
+++ b/Engine/Source/Programs/Enterprise/Datasmith/DatasmithARCHICADExporter/Private/GeometryUtil.h
@@ -20,8 +20,8 @@ class FGeometryUtil
 	// Extract the rotation from the matrix and return as a Quat
 	static FQuat GetRotationQuat(const double Matrix[3][4]);
 
-	// Return the Quat equivalent to the direction vector
-	static FQuat GetRotationQuat(const ModelerAPI::Vector& Direction);
+	// Return the Quat equivalent of rotation defined by Direction and Up vectors
+	static FQuat GetRotationQuat(const ModelerAPI::Vector& Direction, const ModelerAPI::Vector& Up);
 
 	// Convert Archicad camera rotation to an Unreal Quat
 	static FQuat GetRotationQuat(const double PitchInDegrees, const double YawInDegrees, const double RollInDegrees);
diff --git a/Engine/Source/Programs/Enterprise/Datasmith/DatasmithARCHICADExporter/Private/SyncData.cpp b/Engine/Source/Programs/Enterprise/Datasmith/DatasmithARCHICADExporter/Private/SyncData.cpp
--- a/Engine/Source/Programs/Enterprise/Datasmith/DatasmithARCHICADExporter/Private/SyncData.cpp
+++ b/Engine/Source/Programs/Enterprise/Datasmith/DatasmithARCHICADExporter/Private/SyncData.cpp
@@ -1310,8 +1310,8 @@ void FSyncData::FLight::Process(FProcessInfo* IOProcessInfo)
 					break;
 			}
 			AreaLightElement.SetLightType(AreaLightType);
-			AreaLightElement.SetWidth(float(Parameters.AreaSize.y * IOProcessInfo->SyncContext.ScaleLength));
-			AreaLightElement.SetLength(float(Parameters.AreaSize.x * IOProcessInfo->SyncContext.ScaleLength));
+			AreaLightElement.SetWidth(float(Parameters.AreaSize.x * IOProcessInfo->SyncContext.ScaleLength));
+			AreaLightElement.SetLength(float(Parameters.AreaSize.y * IOProcessInfo->SyncContext.ScaleLength));
 		}
 		if (!Parameters.IESFileName.IsEmpty())
 		{
@@ -1442,7 +1442,7 @@ FSyncData::FLight::FLightData::FLightData(const ModelerAPI::Light& InLight)
 	Color = ACRGBColorToUELinearColor(InLight.GetColor());
 
 	Position = FGeometryUtil::GetTranslationVector(InLight.GetPosition());
-	Rotation = FGeometryUtil::GetRotationQuat(InLight.GetDirection());
+	Rotation = FGeometryUtil::GetRotationQuat(InLight.GetDirection(), InLight.GetUpVector());
 }
 
 bool FSyncData::FLight::FLightData::operator!=(const FLightData& InOther) const
