[Naviswork] Add support for Naviswork 2025

diff --git a/Engine/Source/Programs/Enterprise/Datasmith/DatasmithNavisworksExporter/DatasmithNavisworksPlugin/DatasmithNavisworksPlugin.csproj b/Engine/Source/Programs/Enterprise/Datasmith/DatasmithNavisworksExporter/DatasmithNavisworksPlugin/DatasmithNavisworksPlugin.csproj
--- a/Engine/Source/Programs/Enterprise/Datasmith/DatasmithNavisworksExporter/DatasmithNavisworksPlugin/DatasmithNavisworksPlugin.csproj
+++ b/Engine/Source/Programs/Enterprise/Datasmith/DatasmithNavisworksExporter/DatasmithNavisworksPlugin/DatasmithNavisworksPlugin.csproj
@@ -8,6 +8,8 @@
     <ProjectGuid>{D16F2120-CCE4-4FCC-92C8-86124AB1FA25}</ProjectGuid>
     <TargetFrameworkProfile />
   </PropertyGroup>
+  <Import Project="..\Navisworks2025.props" Condition="'$(Configuration)|$(Platform)'=='Debug2025|AnyCPU'" />
+  <Import Project="..\Navisworks2025.props" Condition="'$(Configuration)|$(Platform)'=='Release2025|AnyCPU'" />
   <Import Project="..\Navisworks2024.props" Condition="'$(Configuration)|$(Platform)'=='Debug2024|AnyCPU'" />
   <Import Project="..\Navisworks2024.props" Condition="'$(Configuration)|$(Platform)'=='Release2024|AnyCPU'" />
   <Import Project="..\Navisworks2023.props" Condition="'$(Configuration)|$(Platform)'=='Debug2023|AnyCPU'" />
@@ -30,6 +32,30 @@
     <TargetFrameworkVersion>v4.8</TargetFrameworkVersion>
     <FileAlignment>512</FileAlignment>
   </PropertyGroup>
+  <PropertyGroup Condition=" '$(Configuration)|$(Platform)' == 'Debug2025|AnyCPU' ">
+    <DebugSymbols>true</DebugSymbols>
+    <DebugType>full</DebugType>
+    <Optimize>false</Optimize>
+    <OutputPath>bin\Debug2025\</OutputPath>
+    <DefineConstants>DEBUG;TRACE;DATASMITH_NAVISWORKS_2025_OR_GREATER</DefineConstants>
+    <ErrorReport>prompt</ErrorReport>
+    <WarningLevel>4</WarningLevel>
+    <StartAction>Program</StartAction>
+    <StartProgram>$(ProgramW6432)\Autodesk\Navisworks Manage $(Navisworks_Version)\Roamer.exe</StartProgram>
+    <EnableUnmanagedDebugging>true</EnableUnmanagedDebugging>
+    <AllowUnsafeBlocks>true</AllowUnsafeBlocks>
+  </PropertyGroup>
+  <PropertyGroup Condition=" '$(Configuration)|$(Platform)' == 'Release2025|AnyCPU' ">
+    <DebugType>pdbonly</DebugType>
+    <Optimize>true</Optimize>
+    <OutputPath>bin\Release2025\</OutputPath>
+    <DefineConstants>TRACE;DATASMITH_NAVISWORKS_2025_OR_GREATER</DefineConstants>
+    <ErrorReport>prompt</ErrorReport>
+    <WarningLevel>4</WarningLevel>
+    <StartAction>Program</StartAction>
+    <StartProgram>$(ProgramW6432)\Autodesk\Navisworks Manage $(Navisworks_Version)\Roamer.exe</StartProgram>
+    <AllowUnsafeBlocks>true</AllowUnsafeBlocks>
+  </PropertyGroup>
   <PropertyGroup Condition=" '$(Configuration)|$(Platform)' == 'Debug2024|AnyCPU' ">
     <DebugSymbols>true</DebugSymbols>
     <DebugType>full</DebugType>
diff --git a/Engine/Source/Programs/Enterprise/Datasmith/DatasmithNavisworksExporter/DatasmithNavisworksPlugin/PackageContents.xml b/Engine/Source/Programs/Enterprise/Datasmith/DatasmithNavisworksExporter/DatasmithNavisworksPlugin/PackageContents.xml
--- a/Engine/Source/Programs/Enterprise/Datasmith/DatasmithNavisworksExporter/DatasmithNavisworksPlugin/PackageContents.xml
+++ b/Engine/Source/Programs/Enterprise/Datasmith/DatasmithNavisworksExporter/DatasmithNavisworksPlugin/PackageContents.xml
@@ -25,4 +25,8 @@
 	    <RuntimeRequirements OS="Win64" Platform="NAVMAN|NAVSIM" SeriesMin="Nw21" SeriesMax="Nw21" />
         <ComponentEntry AppName="Datasmith Exporter for Autodesk® Navisworks® 2024" AppType="ManagedPlugin" ModuleName="./DatasmithNavisworksPlugin2024.dll" />
     </Components>
+    <Components Description="2025">
+	    <RuntimeRequirements OS="Win64" Platform="NAVMAN|NAVSIM" SeriesMin="Nw22" SeriesMax="Nw22" />
+        <ComponentEntry AppName="Datasmith Exporter for Autodesk® Navisworks® 2025" AppType="ManagedPlugin" ModuleName="./DatasmithNavisworksPlugin2025.dll" />
+    </Components>
 </ApplicationPackage>
diff --git a/Engine/Source/Programs/Enterprise/Datasmith/DatasmithNavisworksExporter/Private/DatasmithNavisworksPlugin.cs b/Engine/Source/Programs/Enterprise/Datasmith/DatasmithNavisworksExporter/Private/DatasmithNavisworksPlugin.cs
--- a/Engine/Source/Programs/Enterprise/Datasmith/DatasmithNavisworksExporter/Private/DatasmithNavisworksPlugin.cs
+++ b/Engine/Source/Programs/Enterprise/Datasmith/DatasmithNavisworksExporter/Private/DatasmithNavisworksPlugin.cs
@@ -2776,12 +2776,19 @@ namespace DatasmithNavisworks
 		{
 			ItemFilterParams.bHasClipBox = false;
 
-			LcOaClipPlaneSet ClipPlaneSet = ActiveDocument.ActiveView.Viewer.GetClipPlaneSet();
-			LcOaClipPlaneSetMode ClipPlaneSetMode = ClipPlaneSet.GetMode();
-
-			if (ClipPlaneSetMode == LcOaClipPlaneSetMode.eMODE_BOX)
-			{
-				BoundingBox3D Box3D = ClipPlaneSet.GetBox();
+#if DATASMITH_NAVISWORKS_2025_OR_GREATER
+			ClipPlaneSet PlaneSet = ActiveDocument.ActiveView.Viewer.GetClipPlaneSet();
+			ClipPlaneSetMode PlaneSetMode = PlaneSet.Mode;
+			if (PlaneSetMode == ClipPlaneSetMode.Box)
+			{
+				BoundingBox3D Box3D = PlaneSet.Box;
+#else
+			LcOaClipPlaneSet PlaneSet = ActiveDocument.ActiveView.Viewer.GetClipPlaneSet();
+			LcOaClipPlaneSetMode PlaneSetMode = PlaneSet.GetMode();
+			if (PlaneSetMode == LcOaClipPlaneSetMode.eMODE_BOX)
+			{
+				BoundingBox3D Box3D = PlaneSet.GetBox();
+#endif
 
 				ItemFilterParams.bHasClipBox = true;
 				Box ClipBox = new Box
diff --git a/Engine/Source/Programs/Enterprise/Datasmith/DatasmithNavisworksExporter/UBTModules/DatasmithNavisworks2025.Build.cs b/Engine/Source/Programs/Enterprise/Datasmith/DatasmithNavisworksExporter/UBTModules/DatasmithNavisworks2025.Build.cs
new file mode 100644
--- /dev/null
+++ b/Engine/Source/Programs/Enterprise/Datasmith/DatasmithNavisworksExporter/UBTModules/DatasmithNavisworks2025.Build.cs
@@ -0,0 +1,20 @@
+// Copyright Epic Games, Inc. All Rights Reserved.
+
+using System.IO;
+
+namespace UnrealBuildTool.Rules
+{
+	[SupportedPlatforms("Win64")]
+	public class DatasmithNavisworks2025 : DatasmithNavisworksBase
+	{
+		public DatasmithNavisworks2025(ReadOnlyTargetRules Target)
+			: base(Target)
+		{
+		}
+
+		public override string GetNavisworksVersion()
+		{
+			return "2025";
+		}
+	}
+}
diff --git a/Engine/Source/Programs/Enterprise/Datasmith/DatasmithNavisworksExporter/UBTModules/DatasmithNavisworks2025.Target.cs b/Engine/Source/Programs/Enterprise/Datasmith/DatasmithNavisworksExporter/UBTModules/DatasmithNavisworks2025.Target.cs
new file mode 100644
--- /dev/null
+++ b/Engine/Source/Programs/Enterprise/Datasmith/DatasmithNavisworksExporter/UBTModules/DatasmithNavisworks2025.Target.cs
@@ -0,0 +1,14 @@
+// Copyright Epic Games, Inc. All Rights Reserved.
+
+using UnrealBuildTool;
+using System.IO;
+
+public class DatasmithNavisworks2025Target : DatasmithNavisworksBaseTarget
+{
+	public DatasmithNavisworks2025Target(TargetInfo Target)
+		: base(Target)
+	{
+	}
+
+	public override string GetVersion() { return "2025"; }
+}
