[Solidworks] Replaced FaceId with PersistentIds for face identifications

This resutls in more stable face-material connection in SW export, fixing the "One variant color does not convert to UE" issue
Also,
- removed Parallel access to SW API - this slowed things down and interfered with debugging
- Code speed-up(due to Parallel fix) - x2
- added more diagnostics logging to simplify debugging of meshes and materials

diff --git a/Engine/Source/Programs/Enterprise/Datasmith/DatasmithSolidworksExporter/Private/Body.cs b/Engine/Source/Programs/Enterprise/Datasmith/DatasmithSolidworksExporter/Private/Body.cs
--- a/Engine/Source/Programs/Enterprise/Datasmith/DatasmithSolidworksExporter/Private/Body.cs
+++ b/Engine/Source/Programs/Enterprise/Datasmith/DatasmithSolidworksExporter/Private/Body.cs
@@ -12,19 +12,16 @@ namespace DatasmithSolidworks
 {
 	public class FBody
 	{
-		public List<FBodyFace> Faces { get; set; } = new List<FBodyFace>();
 		public Body2 Body { get; private set; } = null;
 		public FBoundingBox Bounds { get; private set; }
 
 		public class FBodyFace
 		{
 			public Face2 Face { get; private set; } = null;
-			public FBody ParentBody { get; private set; } = null;
 
-			public FBodyFace(Face2 InFace, FBody InParentBody)
+			public FBodyFace(Face2 InFace)
 			{
 				Face = InFace;
-				ParentBody = InParentBody;
 			}
 
 			public FTriangleStrip ExtractGeometry()
@@ -51,56 +48,6 @@ namespace DatasmithSolidworks
 			Bounds.Add(new FVec3(BoundsArray[0], BoundsArray[1], BoundsArray[2]));
 			Bounds.Add(new FVec3(BoundsArray[3], BoundsArray[4], BoundsArray[5]));
 
-#if true
-			object[] ArrFaces = Body.GetFaces();
-
-			if (ArrFaces != null)
-			{
-				foreach (object ObjFace in ArrFaces)
-				{
-					Faces.Add(new FBodyFace(ObjFace as Face2, this));
-				}
-			}
-#else
-			// Get body faces
-			dynamic Face = Body.GetFirstFace() as Face2;
-			while (Face != null)
-			{
-				
-				Face = Face.GetNextFace() as Face2;
-			}
-#endif
-		}
-
-		static private List<FBody> FetchBodies(EnumBodies2 InEnumerator)
-		{
-			
-			List<FBody> Bodies = new List<FBody>();
-
-			if (InEnumerator == null)
-			{
-				return Bodies;
-			}
-
-			try
-			{
-				Body2 Body = null;
-				do
-				{
-					int Fetched = 0;
-					InEnumerator.Next(1, out Body, ref Fetched);
-					if (Body != null && Body.Visible && !Body.IsTemporaryBody())
-					{
-						Bodies.Add(new FBody(Body));
-					}
-				} while (Body != null);
-			}
-			catch
-			{
-				Debug.Assert(false);
-			}
-
-			return Bodies;
 		}
 
 		public static ConcurrentBag<FBody> FetchBodies(object[] ObjSolidBodies, object[] ObjSheetBodies)
diff --git a/Engine/Source/Programs/Enterprise/Datasmith/DatasmithSolidworksExporter/Private/ConfigurationExporter.cs b/Engine/Source/Programs/Enterprise/Datasmith/DatasmithSolidworksExporter/Private/ConfigurationExporter.cs
--- a/Engine/Source/Programs/Enterprise/Datasmith/DatasmithSolidworksExporter/Private/ConfigurationExporter.cs
+++ b/Engine/Source/Programs/Enterprise/Datasmith/DatasmithSolidworksExporter/Private/ConfigurationExporter.cs
@@ -66,26 +66,8 @@ namespace DatasmithSolidworks
 				{
 					return;
 				}
-
-				foreach (FVec3 Vertex in MeshData.Vertices)
-				{
-					Hash ^= Vertex.GetHashCode();
-				}
-
-				foreach (FVec3 Normal in MeshData.Normals)
-				{
-					Hash ^= Normal.GetHashCode();
-				}
-
-				foreach (FVec2 TexCoord in MeshData.TexCoords)
-				{
-					Hash ^= TexCoord.GetHashCode();
-				}
-
-				foreach (FTriangle Triangle in MeshData.Triangles)
-				{
-					Hash ^= Triangle.GetHashCode();
-				}
+				
+				Hash = MeshData.HashCode;
 			}
 
 			public override bool Equals(object Obj)
@@ -474,7 +456,7 @@ namespace DatasmithSolidworks
 
 					LogDedent();
 				}
-
+				
 				// Export materials
 				InDoc.SetExportStatus($"Component Materials");
 				HashSet<FComponentName> ComponentNamesToExportSet = new HashSet<FComponentName>();
@@ -495,7 +477,6 @@ namespace DatasmithSolidworks
 				
 				// Export meshes
 				InDoc.SetExportStatus($"Component Meshes");
-
 				InDoc.ProcessConfigurationMeshes(ExtractedMeshes, MeshesConfiguration);
 
 				// Combine separate scene trees into the single one with configuration-specific data
diff --git a/Engine/Source/Programs/Enterprise/Datasmith/DatasmithSolidworksExporter/Private/Document.cs b/Engine/Source/Programs/Enterprise/Datasmith/DatasmithSolidworksExporter/Private/Document.cs
--- a/Engine/Source/Programs/Enterprise/Datasmith/DatasmithSolidworksExporter/Private/Document.cs
+++ b/Engine/Source/Programs/Enterprise/Datasmith/DatasmithSolidworksExporter/Private/Document.cs
@@ -6,9 +6,13 @@ using System.Collections.Concurrent;
 using System.Collections.Generic;
 using System.Diagnostics;
 using System.IO;
+using System.Linq;
 using System.Threading;
 using DatasmithSolidworks.Names;
 using static DatasmithSolidworks.Addin;
+using SolidWorks.Interop.swconst;
+using static DatasmithSolidworks.FBody;
+using static DatasmithSolidworks.FDocumentTracker;
 
 namespace DatasmithSolidworks
 {
@@ -31,6 +35,28 @@ namespace DatasmithSolidworks
 
 		public uint FaceCounter = 1; // Face Id generator
 
+		// Encapsulates SW PersistentReference of a face
+		public readonly struct FFaceId
+		{
+			public readonly byte[] Value;
+			
+			public FFaceId(byte[] InValue=null)
+			{
+				Value = InValue;
+			}
+			
+			public bool IsValid()
+			{
+				return Value != null;
+			}
+			
+			// Get key to use in a dictionary
+			public string GetKey()
+			{
+				return Convert.ToBase64String(Value);
+			}
+		}
+		
 		protected FDocumentTracker(FDocument InDoc, FDatasmithExporter InExporter)
 		{
 			Doc = InDoc;
@@ -56,21 +82,22 @@ namespace DatasmithSolidworks
 			bDocumentIsDirty = bInDirty;
 		}
 
-		public uint GetFaceId(IFace2 InFace)
+		// Currently, FaceId uses PersistentReference which replaced internal FaceId(GetFaceId/SetFaceId). GetFaceId/SetFaceId was too unstable.
+		// PersistentReference is always defined. But GetOrAssignFaceId is left to identify places in code where Id is used with intent that it has to exist.
+		// Calls to GetFaceId don't expect Id to be defined and can ignore if it's not.
+		public FFaceId GetOrAssignFaceId(IFace2 InFace)
 		{
-			uint FaceId = unchecked((uint)InFace.GetFaceId());
-			if (!IsValidFaceId(FaceId))
-			{
-				uint Count = FaceCounter++;
-				FaceId = unchecked((uint)(0xAA << 24)) | Count;
-				InFace.SetFaceId((int)FaceId);
-			}
-			return FaceId;
+			return new FFaceId(Doc.SwDoc.Extension.GetPersistReference3(InFace));
 		}
-
-		public static bool IsValidFaceId(uint InFaceId)
+		
+		public FFaceId GetFaceId(IFace2 InFace)
+		{
+			return InFace != null ? new FFaceId(Doc.SwDoc.Extension.GetPersistReference3(InFace)) : new FFaceId();
+		}
+		
+		public static bool IsValidFaceId(FFaceId InFaceId)
 		{
-			return (InFaceId >> 24 == 0xAA);
+			return InFaceId.IsValid();
 		}
 
 		public bool IsUpdateInProgress()
@@ -238,18 +265,23 @@ namespace DatasmithSolidworks
 		}
 
 		public abstract FMeshData ExtractComponentMeshData(Component2 Comp);
-
-
+		
 		// Extracts meshes used for the assembly configuration
 		public void ProcessConfigurationMeshes(List<FDatasmithExporter.FMeshExportInfo> MeshExportInfos, FMeshes.FConfiguration MeshesConfiguration)
 		{
-			
+			LogDebug($"ProcessConfigurationMeshes:");
+			LogIndent();
+
 			// Extract meshes data and prepare for parallel datasmith export
 			// note: mesh data need to be extracted from the component when required configuration is active(i.e. can't move it outside of configuration enumeration loop)
-			foreach (Component2 Comp in MeshesConfiguration.EnumerateComponents())
+			// Enumerate components stable
+			foreach (var CP in MeshesConfiguration.EnumerateComponents().Select(Comp => new {Comp, Name=new FComponentName(Comp)} ) .OrderBy(CP => CP.Name.ToString()))
 			{
-				FMeshData MeshData = ExtractComponentMeshData(Comp);
-				FComponentName ComponentName = new FComponentName(Comp);
+				
+				FMeshData MeshData = ExtractComponentMeshData(CP.Comp);
+				FComponentName ComponentName = CP.Name;
+				LogDebug($"{ComponentName}:");
+				LogDebug($"{MeshData}");
 				
 				if (MeshData != null)
 				{
@@ -262,6 +294,7 @@ namespace DatasmithSolidworks
 					});
 				}
 			}
+			LogDedent();
 		}
 
 		public void AssignMaterialsToDatasmithMeshes(List<FDatasmithExporter.FMeshExportInfo> CreatedMeshes)
diff --git a/Engine/Source/Programs/Enterprise/Datasmith/DatasmithSolidworksExporter/Private/Material.cs b/Engine/Source/Programs/Enterprise/Datasmith/DatasmithSolidworksExporter/Private/Material.cs
--- a/Engine/Source/Programs/Enterprise/Datasmith/DatasmithSolidworksExporter/Private/Material.cs
+++ b/Engine/Source/Programs/Enterprise/Datasmith/DatasmithSolidworksExporter/Private/Material.cs
@@ -8,6 +8,7 @@ using System.Collections.Generic;
 using System.Globalization;
 using SolidWorks.Interop.sldworks;
 using System.Diagnostics;
+using static DatasmithSolidworks.Addin;
 
 namespace DatasmithSolidworks
 {
@@ -492,11 +493,13 @@ namespace DatasmithSolidworks
 		public FMaterial()
 		{
 		}
-
+		
 		public FMaterial(RenderMaterial InRenderMat, IModelDocExtension InExt)
 		{
 			if (InRenderMat != null)
 			{
+				LogDebug($"Construct FMaterial from {GetRenderMaterialDebugString(InRenderMat)}");
+				
 				Source = InRenderMat;
 
 				Type = Source.IlluminationShaderType;
@@ -801,6 +804,7 @@ namespace DatasmithSolidworks
 					}
 				}
 			}
+			LogDebug($"Constructed {GetMaterialDebugString()}");
 		}
 
 		public void SetAppearance(IAppearanceSetting InSetting)
@@ -808,6 +812,7 @@ namespace DatasmithSolidworks
 			if (InSetting != null)
 			{
 				Appearance = InSetting;
+				LogDebug($"FMaterial::SetAppearance Color={Appearance.Color:X8}");
 
 				double ADiffuse, AEmission, AReflectivity, ASpecular, ASpecularSpread, ATransparency;
 				int LSpecularColor, LPrimaryColor;
@@ -1093,5 +1098,15 @@ namespace DatasmithSolidworks
 
 			return true;
 		}
+		
+		public static string GetRenderMaterialDebugString(RenderMaterial InRenderMat)
+		{
+			return $@"RenderMaterial(FileName='{InRenderMat.FileName}', IlluminationShaderType={InRenderMat.IlluminationShaderType}, PrimaryColor={InRenderMat.PrimaryColor:X8}, TextureFilename='{InRenderMat.TextureFilename}')";
+		}
+		
+		public string GetMaterialDebugString()
+		{
+			return $"FMaterial(Name='{Name}', FileName='{FileName}', Type={Type}, PrimaryColor={PrimaryColor})";
+		}
 	}
 }
diff --git a/Engine/Source/Programs/Enterprise/Datasmith/DatasmithSolidworksExporter/Private/MeshData.cs b/Engine/Source/Programs/Enterprise/Datasmith/DatasmithSolidworksExporter/Private/MeshData.cs
--- a/Engine/Source/Programs/Enterprise/Datasmith/DatasmithSolidworksExporter/Private/MeshData.cs
+++ b/Engine/Source/Programs/Enterprise/Datasmith/DatasmithSolidworksExporter/Private/MeshData.cs
@@ -1,18 +1,29 @@
 // Copyright Epic Games, Inc. All Rights Reserved.
 
 using System.Collections.Generic;
-using System.Runtime.InteropServices;
+using System.Linq;
+using System.Security.Policy;
+using System.Text;
 
 namespace DatasmithSolidworks
 {
     public class FMeshData
     {
-        public FVec3[] Vertices { get; set; } = null;
-        public FVec3[] Normals { get; set; } = null;
-        public FVec2[] TexCoords { get; set; } = null;
-        public FTriangle[] Triangles { get; set; } = null;
-
-		public static FMeshData Create(List<FGeometryChunk> InChunks)
+	    public readonly FVec3[] Vertices;
+	    public readonly FVec3[] Normals;
+	    public readonly FVec2[] TexCoords;
+        public readonly FTriangle[] Triangles;
+        public int HashCode => ComputeHashCode();
+        
+        public FMeshData(FVec3[] InVertices, FVec3[] InNormals, FVec2[] InTexCoords, FTriangle[] InTriangles)
+        {
+	        Vertices = InVertices;
+	        Normals = InNormals;
+	        TexCoords = InTexCoords;
+	        Triangles = InTriangles;
+        }
+        
+        public static FMeshData Create(List<FGeometryChunk> InChunks)
 		{
 			List<FTriangle> AllTriangles = new List<FTriangle>();
 
@@ -38,17 +49,15 @@ namespace DatasmithSolidworks
 				}
 				VertexOffset += Chunk.Vertices.Length;
 			}
-
-			FMeshData MeshData = new FMeshData();
-
-			MeshData.Vertices = new FVec3[Verts.Count];
-			MeshData.TexCoords = new FVec2[Verts.Count];
-			MeshData.Normals = new FVec3[AllTriangles.Count * 3];
+			
+			FVec3[] Vertices = new FVec3[Verts.Count];
+			FVec2[] TexCoords = new FVec2[Verts.Count];
+			FVec3[] Normals = new FVec3[AllTriangles.Count * 3];
 
 			for (int Idx = 0; Idx < Verts.Count; Idx++)
 			{
-				MeshData.Vertices[Idx] = Verts[Idx].P;
-				MeshData.TexCoords[Idx] = Verts[Idx].UV;
+				Vertices[Idx] = Verts[Idx].P;
+				TexCoords[Idx] = Verts[Idx].UV;
 			}
 
 			for (int I = 0; I < AllTriangles.Count; I++)
@@ -56,14 +65,109 @@ namespace DatasmithSolidworks
 				FTriangle Triangle = AllTriangles[I];
 				int Idx = I * 3;
 
-				MeshData.Normals[Idx + 0] = Verts[Triangle.Index1].N;
-				MeshData.Normals[Idx + 1] = Verts[Triangle.Index2].N;
-				MeshData.Normals[Idx + 2] = Verts[Triangle.Index3].N;
+				Normals[Idx + 0] = Verts[Triangle.Index1].N;
+				Normals[Idx + 1] = Verts[Triangle.Index2].N;
+				Normals[Idx + 2] = Verts[Triangle.Index3].N;
 			}
+			
+			return new FMeshData(Vertices, Normals, TexCoords, AllTriangles.ToArray());
+		}
 
-			MeshData.Triangles = AllTriangles.ToArray();
-
-			return MeshData;
+		private int ComputeHashCode()
+		{
+			int Hash = 0;
+			
+			foreach (FVec3 Vertex in Vertices)
+			{
+				Hash ^= Vertex.GetHashCode();
+			}
+			
+			foreach (FVec3 Normal in Normals)
+			{
+				Hash ^= Normal.GetHashCode();
+			}
+			
+			foreach (FVec2 TexCoord in TexCoords)
+			{
+				Hash ^= TexCoord.GetHashCode();
+			}
+			
+			foreach (FTriangle Triangle in Triangles)
+			{
+				Hash ^= Triangle.GetHashCode();
+			}
+			
+			return Hash;
+		}
+		
+		public override string ToString()
+		{
+			StringBuilder S = new StringBuilder();
+			
+			S.Append($"HashCode: {ComputeHashCode()}");
+			S.AppendLine();
+			
+			S.Append("Vertices:");
+			S.AppendLine();
+			if (Vertices != null)
+			{
+				foreach (FVec3 V in Vertices)
+				{
+					S.Append($"{V}");
+					S.AppendLine();
+				}
+			}
+			else
+			{
+				S.Append("<null>");
+				S.AppendLine();
+			}
+			S.Append("Normals:");
+			S.AppendLine();
+			if (Vertices != null)
+			{
+				foreach (FVec3 N in Normals)
+				{
+					S.Append($"{N}");
+					S.AppendLine();
+				}
+			}
+			else
+			{
+				S.Append("<null>");
+				S.AppendLine();
+			}
+			S.Append("TexCoords:");
+			S.AppendLine();
+			if (Vertices != null)
+			{
+				foreach (FVec2 T in TexCoords)
+				{
+					S.Append($"{T}");
+					S.AppendLine();
+				}
+			}
+			else
+			{
+				S.Append("<null>");
+				S.AppendLine();
+			}
+			S.Append("Triangles:");
+			S.AppendLine();
+			if (Triangles != null)
+			{
+				foreach (FTriangle T in Triangles)
+				{
+					S.Append($"{T}");
+					S.AppendLine();
+				}
+			}
+			else
+			{
+				S.Append("<null>");
+				S.AppendLine();
+			}
+			return S.ToString();
 		}
-	}
+    }
 }
diff --git a/Engine/Source/Programs/Enterprise/Datasmith/DatasmithSolidworksExporter/Private/MetadataManager.cs b/Engine/Source/Programs/Enterprise/Datasmith/DatasmithSolidworksExporter/Private/MetadataManager.cs
--- a/Engine/Source/Programs/Enterprise/Datasmith/DatasmithSolidworksExporter/Private/MetadataManager.cs
+++ b/Engine/Source/Programs/Enterprise/Datasmith/DatasmithSolidworksExporter/Private/MetadataManager.cs
@@ -206,7 +206,6 @@ namespace DatasmithSolidworks
 
 		static private bool ExportFaceProperties(int InFaceindex, Face2 InFace, FMetadata InMetadata)
 		{
-			//InMetadata.AddPair("Face_" + InFaceindex, FDocument.GetFaceId(InFace).ToString());
 			return true;
 		}
 
diff --git a/Engine/Source/Programs/Enterprise/Datasmith/DatasmithSolidworksExporter/Private/ObjectMaterials.cs b/Engine/Source/Programs/Enterprise/Datasmith/DatasmithSolidworksExporter/Private/ObjectMaterials.cs
--- a/Engine/Source/Programs/Enterprise/Datasmith/DatasmithSolidworksExporter/Private/ObjectMaterials.cs
+++ b/Engine/Source/Programs/Enterprise/Datasmith/DatasmithSolidworksExporter/Private/ObjectMaterials.cs
@@ -5,6 +5,7 @@ using SolidWorks.Interop.sldworks;
 using SolidWorks.Interop.swconst;
 using System.Linq;
 using System.Collections.Concurrent;
+using System.Diagnostics;
 using DatasmithSolidworks.Names;
 using static DatasmithSolidworks.Addin;
 using static System.Windows.Forms.VisualStyles.VisualStyleElement;
@@ -17,6 +18,8 @@ namespace DatasmithSolidworks
 	// see "Appearance Hierarchy" in SW docs
 	public class FObjectMaterials
 	{
+		private FDocumentTracker DocumentTracker;
+
 		// Note: part document might be a different document than the owner one:
 		// for example, when loading component materials, owner will be the document 
 		// that component resides in, while part doc will be the document that component references!
@@ -31,8 +34,9 @@ namespace DatasmithSolidworks
 
 		private readonly Dictionary<int, FMaterial> GlobalMaterialsMap;  // All collected materials
 
-		public FObjectMaterials(PartDoc InPartDocument, Dictionary<int, FMaterial> InOutMaterialsMap)
+		public FObjectMaterials(FDocumentTracker InDocumentTracker, PartDoc InPartDocument, Dictionary<int, FMaterial> InOutMaterialsMap)
 		{
+			DocumentTracker = InDocumentTracker;
 			GlobalMaterialsMap = InOutMaterialsMap;
 			PartDocument = InPartDocument;
 		}
@@ -94,11 +98,11 @@ namespace DatasmithSolidworks
 			{
 				return GetMaterial(ComponentMaterialID);
 			}
-
-			uint FaceId = unchecked((uint)(InFace?.GetFaceId() ?? 0));
-			if (FDocumentTracker.IsValidFaceId(FaceId) && FaceMaterialsMap.ContainsKey(FaceId.ToString()))
+			
+			FDocumentTracker.FFaceId FaceId = DocumentTracker.GetFaceId(InFace);
+			if (FDocumentTracker.IsValidFaceId(FaceId) && FaceMaterialsMap.ContainsKey(FaceId.GetKey()))
 			{
-				int MatId = FaceMaterialsMap[FaceId.ToString()];
+				int MatId = FaceMaterialsMap[FaceId.GetKey()];
 				return GetMaterial(MatId);
 			}
 
@@ -200,7 +204,7 @@ namespace DatasmithSolidworks
 
 			object[] ObjMaterials = Ext.GetRenderMaterials2((int)InDisplayState, InDisplayStateNames);
 
-			FObjectMaterials PartMaterials = new FObjectMaterials(InPartDoc, MaterialsMap);
+			FObjectMaterials PartMaterials = new FObjectMaterials(InOwnerDoc, InPartDoc, MaterialsMap);
 
 			foreach (object ObjMat in ObjMaterials)
 			{
@@ -233,9 +237,9 @@ namespace DatasmithSolidworks
 						}
 						case IFace2 Face:
 						{
-							uint FaceId = InOwnerDoc.GetFaceId(Face);
-							LogDebug($"    Face '{Face.GetFaceId()}'");
-							PartMaterials.RegisterMaterial(PartMaterials.FaceMaterialsMap, Doc, RenderMat, FaceId.ToString());
+							FDocumentTracker.FFaceId FaceId = InOwnerDoc.GetOrAssignFaceId(Face);
+							LogDebug($"    Face '{FaceId}'");
+							PartMaterials.RegisterMaterial(PartMaterials.FaceMaterialsMap, Doc, RenderMat, FaceId.GetKey());
 							continue;
 						}
 						case IFeature Feat:
@@ -391,15 +395,16 @@ namespace DatasmithSolidworks
 
 			object[] ObjMaterials = InComponent.GetRenderMaterials2((int)InDisplayState, InDisplayStateNames);
 
-			FObjectMaterials ComponentMaterials = new FObjectMaterials(ComponentDoc as PartDoc, MaterialsMap);
+			FObjectMaterials ComponentMaterials = new FObjectMaterials(InComponentOwner, ComponentDoc as PartDoc, MaterialsMap);
 
 			if (ObjMaterials != null)
 			{
 				foreach (object ObjMat in ObjMaterials)
 				{
 					RenderMaterial RenderMat = ObjMat as RenderMaterial;
-					int NumUsers = RenderMat.GetEntitiesCount();
 
+					int NumUsers = RenderMat.GetEntitiesCount();
+					
 					LogDebug($"  FileName: {RenderMat.FileName}");
 					LogDebug($"    Users({NumUsers}):");
 
@@ -424,10 +429,12 @@ namespace DatasmithSolidworks
 							case IFace2 Face:
 							{
 								// note: plugins sets FaceId on each face to identify per-face materials later in CreateMeshData
-								uint FaceId = InComponentOwner.GetFaceId(Face);
-								LogDebug($"    Face '{Face.GetFaceId()}'");
+								FDocumentTracker.FFaceId FaceId = InComponentOwner.GetOrAssignFaceId(Face);
+								LogDebug($"    Face '{FaceId}'");
+								LogDebug($"          Body: '{((Body2)Face.GetBody())?.Name}'");
+
 								ComponentMaterials.RegisterMaterial(ComponentMaterials.FaceMaterialsMap,
-									InComponentOwner.SwDoc, RenderMat, FaceId.ToString());
+									InComponentOwner.SwDoc, RenderMat, FaceId.GetKey());
 								break;
 							}
 							case IFeature Feat:
@@ -515,6 +522,7 @@ namespace DatasmithSolidworks
 
 					if (Comp.GetModelDoc2() == null)
 					{
+						LogDedent();
 						// Component's model doc might be null if component is suppressed/lightweight (in which case we treat it as hidden)
 						// or has already been deleted
 						return false;
@@ -604,7 +612,7 @@ namespace DatasmithSolidworks
 								if (Comp.GetModelDoc2() is PartDoc CompDoc)
 								{
 									LogDebug($"        registering material for Part Component in DocMaterials");
-									FObjectMaterials ComponentMaterials = new FObjectMaterials(CompDoc, MaterialsMap);
+									FObjectMaterials ComponentMaterials = new FObjectMaterials(InAsmDoc, CompDoc, MaterialsMap);
 
 									LogIndent();
 									ComponentMaterials.SetComponentMaterial(RenderMat, InAsmDoc.SwDoc);
diff --git a/Engine/Source/Programs/Enterprise/Datasmith/DatasmithSolidworksExporter/Private/StripGeometry.cs b/Engine/Source/Programs/Enterprise/Datasmith/DatasmithSolidworksExporter/Private/StripGeometry.cs
--- a/Engine/Source/Programs/Enterprise/Datasmith/DatasmithSolidworksExporter/Private/StripGeometry.cs
+++ b/Engine/Source/Programs/Enterprise/Datasmith/DatasmithSolidworksExporter/Private/StripGeometry.cs
@@ -3,9 +3,12 @@
 using System.Collections.Concurrent;
 using System.Collections.Generic;
 using System.Diagnostics;
+using System.Linq;
 using System.Runtime.InteropServices;
 using System.Threading.Tasks;
 
+using SolidWorks.Interop.sldworks;
+
 namespace DatasmithSolidworks
 {
 	public class FStripGeometryBody
@@ -90,21 +93,29 @@ namespace DatasmithSolidworks
 			}
 
 			FStripGeometry StripGeom = new FStripGeometry();
-
-			Parallel.ForEach(InBodies, Body =>
+				
+			foreach(FBody Body in InBodies)
 			{
 				FStripGeometryBody StripBody = new FStripGeometryBody();
 				StripBody.Bounds = Body.Bounds;
 				StripGeom.Bodies.Add(StripBody);
-
-				foreach (FBody.FBodyFace Face in Body.Faces)
+				
+				object[] ArrFaces = Body.Body.GetFaces();
+				
+					
+				if (ArrFaces != null)
 				{
-					FStripGeometryFace StripFace = new FStripGeometryFace();
-					StripFace.Strip = Face.ExtractGeometry();
-					StripFace.Material = InMaterials?.GetMaterial(Face.Face);
-					StripBody.Faces.Add(StripFace);
+					foreach (object ObjFace in ArrFaces)
+					{
+						FBody.FBodyFace Face = new FBody.FBodyFace(ObjFace as Face2);
+						
+						FStripGeometryFace StripFace = new FStripGeometryFace();
+						StripFace.Strip = Face.ExtractGeometry();
+						StripFace.Material = InMaterials?.GetMaterial(Face.Face);
+						StripBody.Faces.Add(StripFace);
+					}
 				}
-			});
+			}
 
 			FMeshData MeshData = null;
 			try
diff --git a/Engine/Source/Programs/Enterprise/Datasmith/DatasmithSolidworksExporter/Private/Triangle.cs b/Engine/Source/Programs/Enterprise/Datasmith/DatasmithSolidworksExporter/Private/Triangle.cs
--- a/Engine/Source/Programs/Enterprise/Datasmith/DatasmithSolidworksExporter/Private/Triangle.cs
+++ b/Engine/Source/Programs/Enterprise/Datasmith/DatasmithSolidworksExporter/Private/Triangle.cs
@@ -72,7 +72,7 @@ namespace DatasmithSolidworks
 
 		public override string ToString()
 		{
-			return "" + Index1 + "," + Index2 + "," + Index3;
+			return "" + Index1 + "," + Index2 + "," + Index3 + ";" + MaterialID;
 		}
 	}
 }
