[Rhino] Not failing sync/export when internal texture has no path

diff --git a/Engine/Source/Programs/Enterprise/Datasmith/DatasmithRhinoExporter/Private/Utils/DatasmithRhinoUtilities.cs b/Engine/Source/Programs/Enterprise/Datasmith/DatasmithRhinoExporter/Private/Utils/DatasmithRhinoUtilities.cs
--- a/Engine/Source/Programs/Enterprise/Datasmith/DatasmithRhinoExporter/Private/Utils/DatasmithRhinoUtilities.cs
+++ b/Engine/Source/Programs/Enterprise/Datasmith/DatasmithRhinoExporter/Private/Utils/DatasmithRhinoUtilities.cs
@@ -130,14 +130,14 @@ namespace DatasmithRhino.Utils
 
 			string FullPath = RhinoTexture.FileReference.FullPath;
 			string FilePath;
-			if (FullPath.Length != 0)
+			if (!string.IsNullOrEmpty(FullPath))
 			{
 				FilePath = FullPath;
 			}
 			else
 			{
 				string RelativePath = RhinoTexture.FileReference.RelativePath;
-				if (RelativePath.Length == 0)
+				if (string.IsNullOrEmpty(RelativePath))
 				{
 					//No valid path found, skip the texture.
 					return;
