[Archicad] Plugin is installed on START(in addition to FULL and SOLO) of every laguage(from 28 existing) of Archicad product found on the system

- installer UI only allows to toggle install for specific Major AC version. When Major version is enable in the UI all flavours and languages found have the plugin installed.
- note: Archicad has a separate install folder for every localization language or "flavour" or product(e.g. solo/start) and plugin had to be installed into the AC folder in order to be found

diff --git a/Engine/Source/Programs/Enterprise/Datasmith/DatasmithARCHICADExporter/Installer/CustomActionManaged/CustomAction.cs b/Engine/Source/Programs/Enterprise/Datasmith/DatasmithARCHICADExporter/Installer/CustomActionManaged/CustomAction.cs
--- a/Engine/Source/Programs/Enterprise/Datasmith/DatasmithARCHICADExporter/Installer/CustomActionManaged/CustomAction.cs
+++ b/Engine/Source/Programs/Enterprise/Datasmith/DatasmithARCHICADExporter/Installer/CustomActionManaged/CustomAction.cs
@@ -125,11 +125,11 @@ namespace CustomActionManaged
 				string[] MajorVersions = new string[] { "27", "26", "25", "24", "23", };
 
 				string[] MinorVersions = new string[] { "0", };
-
-				string[] Products = new string[] { "FULL", "SOLO", };
+				
+				string[] Products = new string[] { "FULL", "SOLO", "START" };
 
 				string[] Releases = new string[] { "R1-1", };
-
+				
 				string[] Languages = new string[] {
 					"INT", "GER", "USA", "UKI", "HUN", "FRA", "CHE", "AUT", "CHI", "CZE", "FIN", "GRE",
 					"NED", "ITA", "JPN", "KOR", "POL", "POR", "RUS", "SPA", "SWE", "TAI", "NOR", "NZE",
@@ -137,13 +137,15 @@ namespace CustomActionManaged
 				};
 
 				string RegistryKeyFormat = "SOFTWARE\\GRAPHISOFT Installers\\ARCHICAD\\ARCHICAD {0}.{1} {2} {3} {4}";
+				
+				bool bFoundArchicad = false;
 
 				foreach(string MajorVersion in MajorVersions)
 				{
 					foreach (string Product in Products)
 					{
 						// Check if ArchiCAD for this major version and this product is installed on the computer
-						bool bKeyEntryFound = false;
+						bool bMajorProductFound = false;
 
 						foreach (string MinorVersion in MinorVersions)
 						{
@@ -153,21 +155,50 @@ namespace CustomActionManaged
 								{
 									// Build key path from version, language and product
 									string RegistryKeyPath = string.Format(RegistryKeyFormat, MajorVersion, MinorVersion, Language, Product, Release);
+									session.Log($"Checking Archicad product by regkey '{RegistryKeyPath}'");
 
 									UIntPtr Key = UIntPtr.Zero;
-									StringBuilder DataSB = new StringBuilder(1024);
-
+									
 									try
 									{
 										if (RegOpenKeyEx(HKEY_LOCAL_MACHINE, RegistryKeyPath, 0, QueryValue | WOW64_64Key, out Key) == 0)
 										{
+											session.Log($"  '{RegistryKeyPath}' found'");
+
 											uint Type = 0;
 											uint DataSize = 1024;
+											StringBuilder DataSB = new StringBuilder(1024);
 											int Ret = RegQueryValueEx(Key, "InstallLocation", 0, ref Type, DataSB, ref DataSize);
-
-											bKeyEntryFound = DataSB.Length > 0;
-
-											bKeyEntryFound = true;
+											
+											if (DataSB.Length > 0)
+											{
+												string InstallLocationDir = DataSB.ToString();
+												session.Log($"  'InstallLocation' key value is '{InstallLocationDir}'");
+												
+												string[] ColladaInOutFiles = Directory.GetFiles(InstallLocationDir, "Collada In-Out.apx", SearchOption.AllDirectories);
+												if(ColladaInOutFiles.Count() > 0)
+												{
+													bMajorProductFound = true;
+
+													string AddOnInOutDir = Path.GetDirectoryName(ColladaInOutFiles[0]);
+													session.Log($"  addons dir found is '{AddOnInOutDir}'");
+													
+													// Update property storing install directory
+													string PropertyToUpdate = string.Format("ARCHICAD{0}{1}DIR", MajorVersion, Product);
+													SetSessionProperty(session, PropertyToUpdate, InstallLocationDir);
+													
+													PropertyToUpdate = string.Format("ARCHICAD{0}{1}ADDONSDIR", MajorVersion, Product);
+													SetSessionProperty(session, PropertyToUpdate, AddOnInOutDir);
+													
+													SetSessionProperty(session, $"ARCHICAD{MajorVersion}{Product}{Language}VALID", "1");
+													SetSessionProperty(session, $"ARCHICAD{MajorVersion}{Product}{Language}DIR", InstallLocationDir);
+													SetSessionProperty(session, $"ARCHICAD{MajorVersion}{Product}{Language}ADDONSDIR", AddOnInOutDir);
+												}
+											}
+											else
+											{
+												session.Log("  WARNING: no 'InstallLocation' key found");
+											}
 										}
 									}
 									finally
@@ -178,42 +209,24 @@ namespace CustomActionManaged
 										}
 									}
 
-									if (bKeyEntryFound)
-									{
-										string InstallLocationDir = DataSB.ToString();
-
-										string[] ColladaInOutFiles = Directory.GetFiles(InstallLocationDir, "Collada In-Out.apx", SearchOption.AllDirectories);
-										if(ColladaInOutFiles.Count() > 0)
-										{
-											string AddOnInOutDir = Path.GetDirectoryName(ColladaInOutFiles[0]);
-
-											// Update property storing install directory
-											string PropertyToUpdate = string.Format("ARCHICAD{0}{1}DIR", MajorVersion, Product);
-											SetSessionProperty(session, PropertyToUpdate, InstallLocationDir);
-
-											PropertyToUpdate = string.Format("ARCHICAD{0}{1}ADDONSDIR", MajorVersion, Product);
-											SetSessionProperty(session, PropertyToUpdate, AddOnInOutDir);
-										}
-
-										break;
-									}
-								}
-
-								// Property has been no need to look ay further
-								if (bKeyEntryFound == true)
-								{
-									break;
 								}
-							}
 
-							// Property has been no need to look ay further
-							if (bKeyEntryFound == true)
-							{
-								break;
 							}
 						}
+						
+						if (bMajorProductFound)
+						{
+							SetSessionProperty(session, $"ARCHICAD{MajorVersion}VALID", "1");
+						}						
+						bFoundArchicad |= bMajorProductFound;
 					}
 				}
+				
+				if (!bFoundArchicad)
+				{
+					SetSessionProperty(session, "NOTVALIDARCHICAD", "1");
+					SetSessionProperty(session, "NOTFOUNDARCHICAD", "1");
+				}				
 
 				session.Log("Successfully update option's values.");
 			}
