[Rhino] Fixed Error on resync for user model

- despite visibility is hierarchical in Rhino sometimes breps have Visible flag set even though parent geometry is not Visible

diff --git a/Engine/Source/Programs/Enterprise/Datasmith/DatasmithRhinoExporter/Private/ElementExporters/IDatasmithRhinoElementExporter.cs b/Engine/Source/Programs/Enterprise/Datasmith/DatasmithRhinoExporter/Private/ElementExporters/IDatasmithRhinoElementExporter.cs
--- a/Engine/Source/Programs/Enterprise/Datasmith/DatasmithRhinoExporter/Private/ElementExporters/IDatasmithRhinoElementExporter.cs
+++ b/Engine/Source/Programs/Enterprise/Datasmith/DatasmithRhinoExporter/Private/ElementExporters/IDatasmithRhinoElementExporter.cs
@@ -80,7 +80,7 @@ namespace DatasmithRhino.ElementExporters
 			foreach (T CurrentElementInfo in GetElementsToSynchronize())
 			{
 				DatasmithRhinoProgressManager.Instance.UpdateCurrentTaskProgress((float)(++ElementIndex) / TotalNumberOfElements);
-
+				
 				ValidateElement(CurrentElementInfo);
 
 				switch (CurrentElementInfo.DirectLinkStatus)
@@ -95,7 +95,7 @@ namespace DatasmithRhino.ElementExporters
 						}
 						else
 						{
-							//#ueent_todo Log elements who could not be exported in Datasmith logging API.
+							DatasmithRhinoPlugin.Instance.LogManager.AddLog(DatasmithRhinoLogType.Warning, $"SynchronizeSynchronously: Can't create Actor for '{CurrentElementInfo.BaseLabel}'['{CurrentElementInfo.Name}']");
 						}
 						break;
 					case DirectLinkSynchronizationStatus.Modified:
@@ -227,7 +227,7 @@ namespace DatasmithRhino.ElementExporters
 				}
 				else
 				{
-					//#ueent_todo Log elements who could not be exported in Datasmith logging API.
+					DatasmithRhinoPlugin.Instance.LogManager.AddLog(DatasmithRhinoLogType.Warning, $"SynchronizeSynchronously: CreatedElementInfos '{CreatedElementInfo.BaseLabel}'['{CreatedElementInfo.Name}'] has no ExportedElement");
 				}
 			}
 		}
diff --git a/Engine/Source/Programs/Enterprise/Datasmith/DatasmithRhinoExporter/Private/ExportContext/DatasmithActorInfo.cs b/Engine/Source/Programs/Enterprise/Datasmith/DatasmithRhinoExporter/Private/ExportContext/DatasmithActorInfo.cs
--- a/Engine/Source/Programs/Enterprise/Datasmith/DatasmithRhinoExporter/Private/ExportContext/DatasmithActorInfo.cs
+++ b/Engine/Source/Programs/Enterprise/Datasmith/DatasmithRhinoExporter/Private/ExportContext/DatasmithActorInfo.cs
@@ -70,32 +70,36 @@ namespace DatasmithRhino.ExportContext
 		/// </summary>
 		private List<int> RelativeLayerIndices = new List<int>();
 		public Layer VisibilityLayer { get; private set; } = null;
-
-		public bool bIsVisible
+		
+		public bool bIsVisible => IsVisible();
+		
+		public bool IsVisible(bool bCheckParent=true)
 		{
-			get
+			bool bVisibleLocally = true;
+			if (RhinoCommonObject is Layer RhinoLayer)
 			{
-				bool bVisibleLocally = true;
-				if (RhinoCommonObject is Layer RhinoLayer)
-				{
-					// This recursion ensure that only layer actors with an actual exported object under them are considered visible.
-					// ie. Layers containing no mesh are not visible.
-					return RhinoLayer.IsVisible
-						&& ChildrenInternal.Any(Child => Child.bIsVisible);
-				}
-				else if (RhinoCommonObject is RhinoObject CurrentRhinoObject)
-				{
-					bVisibleLocally = CurrentRhinoObject.Visible;
-				}
-				else if (bIsRoot)
-				{
-					return true;
-				}
-
-				return bVisibleLocally
-					&& (VisibilityLayer == null || VisibilityLayer.IsVisible)
-					&& (DefinitionNode == null || DefinitionNode.bIsVisible);
+				// This recursion ensure that only layer actors with an actual exported object under them are considered visible.
+				// ie. Layers containing no mesh are not visible.
+				return RhinoLayer.IsVisible
+					&& ChildrenInternal.Any(Child => Child.IsVisible(bCheckParent=false));
+			}
+			else if (RhinoCommonObject is RhinoObject CurrentRhinoObject)
+			{
+				bVisibleLocally = CurrentRhinoObject.Visible;
+			}
+			else if (bIsRoot)
+			{
+				return true;
 			}
+			
+			// Take into account parent object visibility -
+			// Api sometimes reports brep objects composing ModelGeometry visible when the object itself is not visible
+			bool bParentVisible = !bCheckParent || Parent == null || Parent.bIsVisible;
+			
+			return bVisibleLocally
+			       && (VisibilityLayer == null || VisibilityLayer.IsVisible)
+			       && (DefinitionNode == null || DefinitionNode.bIsVisible)
+			       && bParentVisible;
 		}
 
 		public DatasmithActorInfo(Transform NodeTransform, string InName, string InUniqueLabel, string InBaseLabel)
