[3dsMax] Datasmith SkyPortalLight only when the light is a SkyDome in Max

- in Max pre 2024(checked 2021) Arnold light by default had a portal_mode attribute set to 0. Exported code used this attr alone to identify it as a portal, which is wrokg as this attr only used in UI for SkyDome light shape.
- also, fixed Disk shape type to export as proper area light(not as comment falsely stated that it was converted to spot)

diff --git a/Engine/Source/Programs/Enterprise/Datasmith/DatasmithMaxExporter/Private/DatasmithMaxSceneExporter.cpp b/Engine/Source/Programs/Enterprise/Datasmith/DatasmithMaxExporter/Private/DatasmithMaxSceneExporter.cpp
--- a/Engine/Source/Programs/Enterprise/Datasmith/DatasmithMaxExporter/Private/DatasmithMaxSceneExporter.cpp
+++ b/Engine/Source/Programs/Enterprise/Datasmith/DatasmithMaxExporter/Private/DatasmithMaxSceneExporter.cpp
@@ -362,6 +362,8 @@ TSharedPtr< IDatasmithLightActorElement > FDatasmithMaxSceneExporter::CreateLigh
 		LightType = EDatasmithElementType::DirectionalLight;
 	}
 
+	bool bIsArnoldPortal = false;
+
 	int NumParamBlocks = Light->NumParamBlocks();
 
 	for (int j = 0; j < NumParamBlocks; j++)
@@ -406,36 +408,45 @@ TSharedPtr< IDatasmithLightActorElement > FDatasmithMaxSceneExporter::CreateLigh
 			// Arnold Light transform
 			else if (FCString::Stricmp(ParamDefinition.int_name, TEXT("shapeType")) == 0 && LightClass == EMaxLightClass::ArnoldLight)
 			{
+				// Quad
 				if (ParamBlock2->GetInt(ParamDefinition.ID, GetCOREInterface()->GetTime()) == 3)
 				{
 					LightType = EDatasmithElementType::AreaLight;
 				}
-				// disc is transformed to spot
+				// Disk
+				if (ParamBlock2->GetInt(ParamDefinition.ID, GetCOREInterface()->GetTime()) == 4)
+				{
+					LightType = EDatasmithElementType::AreaLight;
+				}
+				// Spot
 				else if (ParamBlock2->GetInt(ParamDefinition.ID, GetCOREInterface()->GetTime()) == 2)
 				{
 					LightType = EDatasmithElementType::SpotLight;
 				}
+				// Distant
 				else if (ParamBlock2->GetInt(ParamDefinition.ID, GetCOREInterface()->GetTime()) == 1)
 				{
 					LightType = EDatasmithElementType::DirectionalLight;
 				}
+				// SkyDome
 				else if (ParamBlock2->GetInt(ParamDefinition.ID, GetCOREInterface()->GetTime()) == 6)
 				{
 					LightType = EDatasmithElementType::EnvironmentLight;
 				}
 			}
-
-			// check Arnold portal (not available on 3dsmax interface yet)
 			else if (FCString::Stricmp(ParamDefinition.int_name, TEXT("portal_mode")) == 0 && LightClass == EMaxLightClass::ArnoldLight)
 			{
 				int IsPortal = ParamBlock2->GetInt(ParamDefinition.ID, GetCOREInterface()->GetTime());
-				if (IsPortal != 0)
-				{
-					LightType = EDatasmithElementType::LightmassPortal;
-				}
+				bIsArnoldPortal = IsPortal != 0;
 			}
 		}
 
+		// Arnold portal attribute is only relevant for SkyDome shape
+		if (bIsArnoldPortal && (LightType == EDatasmithElementType::EnvironmentLight))
+		{
+			LightType = EDatasmithElementType::LightmassPortal;
+		}
+
 		ParamBlock2->ReleaseDesc();
 	}
 
@@ -1656,7 +1667,6 @@ bool FDatasmithMaxSceneExporter::ParseLightParameters(DatasmithMaxDirectLink::FL
 			// Arnold Light transform
 			if (FCString::Stricmp(ParamDefinition.int_name, TEXT("shapeType")) == 0 && LightClass == EMaxLightClass::ArnoldLight)
 			{
-				// disc is transformed to spot
 				if (ParamBlock2->GetInt(ParamDefinition.ID, GetCOREInterface()->GetTime()) == 2 && LightElement->IsA( EDatasmithElementType::SpotLight ) )
 				{
 					TSharedPtr< IDatasmithSpotLightElement > SpotLightElement = StaticCastSharedRef< IDatasmithSpotLightElement >( LightElement );
