[SketchUp] Baking every mesh with non-supported transform.

- Unreal doesn't support 'shear' transformation. Or scaling non-uniformly a rorated child. Meshes with such transformations are exported with this transform baked into its vertices.

diff --git a/Engine/Source/Programs/Enterprise/Datasmith/DatasmithSketchUpRubyExporter/Private/DatasmithSketchUpComponent.cpp b/Engine/Source/Programs/Enterprise/Datasmith/DatasmithSketchUpRubyExporter/Private/DatasmithSketchUpComponent.cpp
--- a/Engine/Source/Programs/Enterprise/Datasmith/DatasmithSketchUpRubyExporter/Private/DatasmithSketchUpComponent.cpp
+++ b/Engine/Source/Programs/Enterprise/Datasmith/DatasmithSketchUpRubyExporter/Private/DatasmithSketchUpComponent.cpp
@@ -27,11 +27,13 @@
 #include "SketchUpAPI/model/layer_folder.h"
 #endif
 
+#include "DatasmithSceneExporter.h"
 #include "DatasmithSketchUpSDKCeases.h"
 
 #include "IDatasmithSceneElements.h"
 #include "DatasmithSceneFactory.h"
 #include "DatasmithUtils.h"
+#include "Algo/Compare.h"
 
 #include "Misc/Paths.h"
 
@@ -89,11 +91,14 @@ void FEntityWithEntities::UpdateOccurrenceMeshActors(FExportContext& Context, FN
 		return;
 	}
 
-	Node.MeshActors.Reset(EntitiesGeometry->GetMeshCount());
+	FEntitiesGeometry::FExportedGeometry& ExportedGeometry = EntitiesGeometry->GetOccurrenceExportedGeometry(Node);
+
+	Node.MeshActors.Reset(ExportedGeometry.GetMeshCount());
 
 	FString ComponentActorName = Node.GetActorName();
+
 	
-	for (int32 MeshIndex = 0; MeshIndex < EntitiesGeometry->GetMeshCount(); ++MeshIndex)
+	for (int32 MeshIndex = 0; MeshIndex < ExportedGeometry.GetMeshCount(); ++MeshIndex)
 	{
 		FString MeshActorName = FString::Printf(TEXT("%ls_%d"), *ComponentActorName, MeshIndex + 1); // Count meshes/mesh actors from 1
 
@@ -118,22 +123,15 @@ void FEntityWithEntities::UpdateOccurrenceMeshActors(FExportContext& Context, FN
 		// ADD_TRACE_LINE(TEXT("Actor %ls: %ls %ls %ls"), *MeshActorLabel, *ComponentDepthTag, *DefinitionGUIDTag, *InstancePathTag);
 
 		// Set the Datasmith mesh element used by the mesh actor.
-		DMeshActorPtr->SetStaticMeshPathName(EntitiesGeometry->GetMeshElementName(MeshIndex));
+		DMeshActorPtr->SetStaticMeshPathName(ExportedGeometry.GetMeshElementName(MeshIndex));
 	}
 }
 
 void FNodeOccurence::UpdateVisibility(FExportContext& Context)
 {
-	if (bHierarchyInvalidated)
-	{
-		// todo: move hierarchy creation here?
-		bHierarchyInvalidated = false;
-	}
-
 	if (bVisibilityInvalidated)
 	{
 		Entity.UpdateOccurrenceVisibility(Context, *this);
-
 		bVisibilityInvalidated = false;
 	}
 
@@ -143,6 +141,19 @@ void FNodeOccurence::UpdateVisibility(FExportContext& Context)
 	}
 }
 
+void FNodeOccurence::UpdateTransformations(FExportContext& Context)
+{
+	if (bPropertiesInvalidated)
+	{
+		Entity.UpdateOccurrenceTransformation(Context, *this);
+	}
+
+	for (FNodeOccurence* ChildNode : Children)
+	{
+		ChildNode->UpdateTransformations(Context);
+	}
+}
+
 void FNodeOccurence::Update(FExportContext& Context)
 {
 	// todo: Is it possible not to traverse whole scene when only part of it changes?
@@ -260,9 +271,11 @@ void FNodeOccurence::ResetMetadataElement(FExportContext& Context)
 	DatasmithMetadataElement->ResetProperties();
 }
 
-void FNodeOccurence::SetVisibility(bool bValue)
+bool FNodeOccurence::SetVisibility(bool bValue)
 {
+	bool  bChanged = bVisible != bValue;
 	bVisible = bValue;
+	return bChanged;
 }
 
 void FNodeOccurence::RemoveDatasmithActorHierarchy(FExportContext& Context)
@@ -341,7 +354,7 @@ void FModelDefinition::Parse(FExportContext& Context)
 
 void FModelDefinition::UpdateGeometry(FExportContext& Context)
 {
-	Entities->UpdateGeometry(Context);
+	Entities->UpdateGeometry(Context, {Context.RootNode.Get()}, {});
 }
 
 void FModelDefinition::UpdateMetadata(FExportContext& Context)
@@ -377,13 +390,6 @@ FString FModelDefinition::GetSketchupSourceId()
 	return GetSketchupSourceGUID();
 }
 
-SUTransformation FModelDefinition::GetMeshBakedTransform()
-{
-	SUTransformation Transform;
-	SUTransformationScale(&Transform, 1.0);
-	return Transform;
-}
-
 bool FModelDefinition::UpdateModel(FExportContext& Context)
 {
 	// SketchUp API has no notification of Geolocation change so retrieve it every time and compare to check if we need to set DirectLink update
@@ -493,8 +499,30 @@ void FComponentInstance::SetupActor(FExportContext& Context, FNodeOccurence& Nod
 
 void FComponentDefinition::UpdateGeometry(FExportContext& Context)
 {
-	Entities->UpdateGeometry(Context);
-	bBakeTransformIntoMesh  = ShouldBakeTransformIntoMesh();
+	// Some occurrences geometry should have its transformation baked into exported mesh, when that transformation can't be converted to UE(i.e. shear transform)
+	TArray<FNodeOccurence*> NodesToInstance;
+	TArray<FNodeOccurence*> NodesToBake;
+
+	for (FComponentInstance* Instance: Instances)
+	{
+		// todo: might add check to ComponentInstance visibility
+		for (FNodeOccurence* NodeOccurence: Instance->Occurrences)
+		{
+			if (NodeOccurence->bVisible)
+			{
+				if (NodeOccurence->bTransformSupportedByUE)
+				{
+					NodesToInstance.Add(NodeOccurence);
+				}
+				else
+				{
+					NodesToBake.Add(NodeOccurence);
+				}
+			}
+		}
+	}
+
+	Entities->UpdateGeometry(Context, NodesToInstance, NodesToBake);
 }
 
 void FComponentDefinition::UpdateMetadata(FExportContext& Context)
@@ -508,13 +536,13 @@ void FComponentInstance::BuildNodeNames(FNodeOccurence& Node)
 	int64 SketchupPersistentID = Node.Entity.GetPersistentId();
 	Node.DatasmithActorName = FString::Printf(TEXT("%ls_%lld"), *Node.ParentNode->GetActorName(), SketchupPersistentID);
 
-	FString EntityName = Node.Entity.GetName();
+	FString EntityName = Node.Entity.GetEntityName();
 	Node.DatasmithActorLabel = FDatasmithUtils::SanitizeObjectName(EntityName.IsEmpty() ? GetDefinition()->GetSketchupSourceName() : EntityName);
 }
 
 void FComponentDefinition::InvalidateInstancesGeometry(FExportContext& Context)
 {
-	// todo: keep all instances or incapsulate enumeration(duplicated) of FComponentInstance
+	// todo: keep all instances or encapsulate enumeration(duplicated) of FComponentInstance
 	size_t InstanceCount = 0;
 	SUComponentDefinitionGetNumInstances(ComponentDefinitionRef, &InstanceCount);
 
@@ -580,118 +608,6 @@ FString FComponentDefinition::GetSketchupSourceId()
 	return FMD5::HashAnsiString(*GetSketchupSourceGUID());
 }
 
-// Bake transform into the mesh only when
-// - component has no child components, 
-// - only SINGLE instance
-// - local transform has non-trivial rotation component
-// Baking transform into mesh might be useful to convert transformation which is not representable in Unreal
-// for example rotated geometry with scaled parent. In Unreal scaling would only be applied in local static mesh space
-// So that scaled(on parent node) rotated cube that should be a rhombus in SU in unreal would become just a box because scaling would apply along local axes(i.e. before rotation)
-// But baking rotation into mesh itself makes scaling apply properly
-bool FComponentDefinition::ShouldBakeTransformIntoMesh()
-{
-	size_t ComponentInstanceCount = 0;
-	SUEntitiesGetNumInstances(GetEntities().EntitiesRef, &ComponentInstanceCount);
-
-	size_t GroupCount = 0;
-	SUEntitiesGetNumGroups(GetEntities().EntitiesRef, &GroupCount);
-
-	if ((Instances.Num() != 1) || ((ComponentInstanceCount + GroupCount) != 0))
-	{
-		return false;
-	}
-
-	// Check the (only) instance's transform
-	FComponentInstance* ComponentInstance = Instances.Array()[0];
-
-	SUTransformation Transform;
-	SUComponentInstanceGetTransform(ComponentInstance->GetComponentInstanceRef(), &Transform);
-
-	// Grab 3x3 rotation submatrix and check if it is not an identity transform
-	const double R[] = {
-		Transform.values[0], Transform.values[1], Transform.values[2],
-		Transform.values[4], Transform.values[5], Transform.values[6],
-		Transform.values[8], Transform.values[9], Transform.values[10]
-	};
-
-	static const double Identity[] = {
-		1, 0, 0,
-		0, 1, 0,
-		0, 0, 1,
-	};
-
-	static_assert(sizeof(Identity)/sizeof(Identity[0]) == sizeof(R)/sizeof(R[0]));
-
-	for (int32 I = 0; I < sizeof(Identity)/sizeof(Identity[0]); ++I)
-	{
-		if (!FMath::IsNearlyEqual(Identity[I], R[I]))
-		{
-			return true;
-		}
-	}
-
-	return false;
-}
-
-SUTransformation FComponentDefinition::GetMeshBakedTransform()
-{
-	if (!ShouldBakeTransformIntoMesh())
-	{
-		SUTransformation Transform;
-		SUTransformationScale(&Transform, 1.0);
-		return Transform;
-	}
-
-	FComponentInstance* ComponentInstance = Instances.Array()[0];
-	SUTransformation Transform;
-	SUComponentInstanceGetTransform(ComponentInstance->GetComponentInstanceRef(), &Transform);
-
-	// Don't bake translation 
-	Transform.values[12] = 0;
-	Transform.values[13] = 0;
-	Transform.values[14] = 0;
-	// and uniform scale
-	Transform.values[15] = 1;
-	// There's no need to bake translation and uniform scale into mesh and keeping those (at least translation) at actor transform is more convenient
-
-	return Transform;
-}
-
-void FComponentDefinition::GetLocalTransform(FComponentInstance& ComponentInstance, SUTransformation& SComponentInstanceLocalTransform)
-{
-	if (ShouldBakeTransformIntoMesh())
-	{
-		SUTransformation SComponentInstanceTransform;
-		SUComponentInstanceGetTransform(ComponentInstance.GetComponentInstanceRef(), &SComponentInstanceTransform);
-
-		// Local instance transform rotation part is baked into mesh, scale and translation aren't
-		double UniformScale = SComponentInstanceTransform.values[15];
-		const SUVector3D Translation{SComponentInstanceTransform.values[12], SComponentInstanceTransform.values[13], SComponentInstanceTransform.values[14]};
-
-		SUTransformation Transform;
-		SUTransformationTranslation(&Transform, &Translation);
-		Transform.values[15] = UniformScale;
-
-		SComponentInstanceLocalTransform = Transform;
-	}
-	else
-	{
-		SUComponentInstanceGetTransform(ComponentInstance.GetComponentInstanceRef(), &SComponentInstanceLocalTransform);
-	}
-}
-
-void FComponentDefinition::ComponentInstancePropertiesInvalidated()
-{
-	// If mesh transform baking is affecting geometry then invalidate geometry
-	// i.e. if transform should be baked or need to bake it has changed
-	bool bShouldBakeTransformIntoMesh = ShouldBakeTransformIntoMesh();
-	if (bShouldBakeTransformIntoMesh || (bShouldBakeTransformIntoMesh != bBakeTransformIntoMesh))
-	{
-		InvalidateDefinitionGeometry();
-	}
-}
-
-
 FString FComponentDefinition::GetSketchupSourceGUID()
 {
 	// Retrieve the SketchUp component definition IFC GUID.
@@ -775,6 +691,8 @@ void FEntityWithEntities::UpdateOccurrence(FExportContext& Context, FNodeOccuren
 		Node.InheritedMaterialID = Node.ParentNode->InheritedMaterialID;
 	}
 
+	FEntitiesGeometry::FExportedGeometry& ExportedGeometry = EntitiesGeometry.GetOccurrenceExportedGeometry(Node);
+
 	FString MeshActorLabel = Node.GetActorLabel();
 	// Update Datasmith Mesh Actors
 	for (int32 MeshIndex = 0; MeshIndex < Node.MeshActors.Num(); ++MeshIndex)
@@ -785,7 +703,7 @@ void FEntityWithEntities::UpdateOccurrence(FExportContext& Context, FNodeOccuren
 
 		// Update Override(Inherited)  Material
 		// todo: set inherited material only on mesh actors that have faces with default material, right now setting on every mesh, hot harmful but excessive
-		if (EntitiesGeometry.IsMeshUsingInheritedMaterial(MeshIndex))
+		if (ExportedGeometry.IsMeshUsingInheritedMaterial(MeshIndex))
 		{
 			Context.Materials.SetMeshActorOverrideMaterial(Node, EntitiesGeometry, MeshActor);
 		}
@@ -877,6 +795,11 @@ void FEntityWithEntities::EntityOccurrenceVisible(FNodeOccurence* Node, bool bVi
 	GetDefinition()->EntityVisible(this, VisibleNodes.Num() > 0);
 }
 
+FComponentInstance::FComponentInstance(SUEntityRef InEntityRef, FComponentDefinition& InDefinition)
+	: Super(InEntityRef)
+	, Definition(InDefinition)
+{}
+
 FComponentInstance::~FComponentInstance()
 {
 }
@@ -901,6 +824,108 @@ bool FComponentInstance::GetAssignedMaterial(FMaterialIDType& MaterialId)
 	return false;
 }
 
+void FComponentInstance::UpdateOccurrenceTransformation(FExportContext& Context, FNodeOccurence& Node)
+{
+	// Compute the world transform of the SketchUp component instance.
+	SUTransformation LocalTransform;
+	SUComponentInstanceGetTransform(GetComponentInstanceRef(), &LocalTransform);
+
+	SUTransformation WorldTransform;
+	SUTransformationMultiply(&Node.ParentNode->WorldTransformSource, &LocalTransform, &WorldTransform);
+
+	bool bTransformChanged = !DatasmithSketchUpUtils::CompareSUTransformations(WorldTransform, Node.WorldTransformSource);
+
+	Node.WorldTransformSource = WorldTransform;
+	Node.WorldTransform = WorldTransform;
+
+
+	auto GetNodePath = [&Node]()
+	{
+		TArray<FString> NamePath;
+
+		for (const FNodeOccurence* N = &Node; N; N = N->ParentNode)
+		{
+			NamePath.Insert(N->Entity.GetEntityLabel(), 0);
+		}
+
+		return FString::Join(NamePath, TEXT("_"));
+	};
+
+	FVector Translation{};
+	FQuat Rotation{};
+	FVector Scale{};
+	FVector Shear{};
+
+	bool bTransformSupportedByUE = true;
+	if (DatasmithSketchUpUtils::DecomposeTransform(LocalTransform, Translation, Rotation, Scale, Shear))
+	{
+		if (!Shear.IsNearlyZero())
+		{
+			bTransformSupportedByUE = false;
+			DatasmithSketchUpUtils::ToRuby::LogWarn(
+				FString::Printf(TEXT("Entity '%s' has shear in local transform"), *GetNodePath()));
+		}
+	}
+	else
+	{
+		DatasmithSketchUpUtils::ToRuby::LogWarn(
+			FString::Printf(TEXT("Entity %s has zero scaling"), *GetNodePath()));
+	}
+
+	if (DatasmithSketchUpUtils::DecomposeTransform(Node.WorldTransform, Translation, Rotation, Scale, Shear))
+	{
+		
+		if (!Shear.IsNearlyZero())
+		{
+			bTransformSupportedByUE = false;
+			DatasmithSketchUpUtils::ToRuby::LogWarn(
+				FString::Printf(TEXT("Entity %s has shear in world transform"), *GetNodePath()));
+		}
+		//  Non-uniform with children not supported as children might be rotated and this would skew them
+		// todo: worth checking down the subtree for actual rotation present to support these edge cases without extra mesh export!
+		else if (!Scale.IsUniform() && !Node.Children.IsEmpty())
+		{
+			bTransformSupportedByUE = false;
+			DatasmithSketchUpUtils::ToRuby::LogWarn(
+				FString::Printf(TEXT("Entity %s has non-uniform scaling in world transform"), *GetNodePath()));
+		}
+	}
+	else
+	{
+		DatasmithSketchUpUtils::ToRuby::LogWarn(
+			FString::Printf(TEXT("Entity %s has zero scaling"), *GetNodePath()));
+	}
+
+	bTransformSupportedByUE = bTransformSupportedByUE && Node.ParentNode->bTransformSupportedByUE;
+
+	if (!bTransformSupportedByUE)
+	{
+
+		SUTransformation ActorTransform;
+		SUTransformation MeshActorWorldTransform;
+		SUTransformation BakeTransform;
+
+		DatasmithSketchUpUtils::SplitTransform(Node.WorldTransformSource, ActorTransform, MeshActorWorldTransform, BakeTransform);
+
+		Node.WorldTransform = ActorTransform;
+		Node.MeshActorWorldTransform = MeshActorWorldTransform;
+		Node.BakeTransform = BakeTransform;
+	}
+
+	// If node's transform is not supported by UE(therefore it was baked/needs baking into mesh)
+	// and transform itself was changed this means that geometry should be re-exported(as exported geometry was baked with old transform)
+	bool bNeedInvalidateBakedGeometry = 
+		(static_cast<bool>(Node.bTransformSupportedByUE) != bTransformSupportedByUE)
+		|| (!Node.bTransformSupportedByUE && bTransformChanged);
+
+	Node.bTransformSupportedByUE = bTransformSupportedByUE;
+
+	if (bNeedInvalidateBakedGeometry)
+	{
+		GetDefinition()->InvalidateDefinitionGeometry();
+	}
+}
+
 void FComponentInstance::UpdateOccurrence(FExportContext& Context, FNodeOccurence& Node)
 {
 	Node.EffectiveLayerRef = DatasmithSketchUpUtils::GetEffectiveLayer(GetComponentInstanceRef(), Node.ParentNode->EffectiveLayerRef);
@@ -920,13 +945,6 @@ void FComponentInstance::UpdateOccurrence(FExportContext& Context, FNodeOccurenc
 	// Set the Datasmith actor layer name.
 	Node.DatasmithActorElement->SetLayer(*FDatasmithUtils::SanitizeObjectName(SEffectiveLayerName));
 
-	// Compute the world transform of the SketchUp component instance.
-	SUTransformation LocalTransform;
-	Definition.GetLocalTransform(*this, LocalTransform);
-	SUTransformation WorldTransform;
-	SUTransformationMultiply(&Node.ParentNode->WorldTransform, &LocalTransform, &WorldTransform);
-	Node.WorldTransform = WorldTransform; // Store world transform to be used by children to compute its
-
 	// Set the Datasmith actor world transform.
 	DatasmithSketchUpUtils::SetActorTransform(Node.DatasmithActorElement, Node.WorldTransform);
 
@@ -938,10 +956,17 @@ void FComponentInstance::UpdateOccurrence(FExportContext& Context, FNodeOccurenc
 	{
 		const TSharedPtr<IDatasmithMeshActorElement>& MeshActor = Node.MeshActors[MeshIndex];
 
-		// Set mesh actor transform after node transform
-		MeshActor->SetScale(Node.DatasmithActorElement->GetScale());
-		MeshActor->SetRotation(Node.DatasmithActorElement->GetRotation());
-		MeshActor->SetTranslation(Node.DatasmithActorElement->GetTranslation());
+		if (Node.bTransformSupportedByUE)
+		{
+			// Set mesh actor transform after node transform
+			MeshActor->SetScale(Node.DatasmithActorElement->GetScale());
+			MeshActor->SetRotation(Node.DatasmithActorElement->GetRotation());
+			MeshActor->SetTranslation(Node.DatasmithActorElement->GetTranslation());
+		}
+		else
+		{
+			DatasmithSketchUpUtils::SetActorTransform(MeshActor, Node.MeshActorWorldTransform);
+		}
 	}
 
 	Super::UpdateOccurrence(Context, Node);
@@ -953,13 +978,19 @@ int64 FComponentInstance::GetPersistentId()
 	return DatasmithSketchUpUtils::GetComponentPID(ComponentInstanceRef);
 }
 
-FString FComponentInstance::GetName()
+FString FComponentInstance::GetEntityName()
 {
 	SUComponentInstanceRef InSComponentInstanceRef = GetComponentInstanceRef();
 	FString SComponentInstanceName;
 	return SuGetString(SUComponentInstanceGetName, InSComponentInstanceRef);
 }
 
+FString FComponentInstance::GetEntityLabel()
+{
+	FString EntityName = GetEntityName();
+	return EntityName.IsEmpty() ? GetDefinition()->GetSketchupSourceName() : EntityName;
+}
+
 void FComponentInstance::UpdateMetadata(FExportContext& Context)
 {
 	ParsedMetadata = MakeUnique<FMetadata>(SUComponentInstanceToEntity(GetComponentInstanceRef()));
@@ -969,7 +1000,7 @@ void FComponentInstance::UpdateEntityProperties(FExportContext& Context)
 {
 	if (bPropertiesInvalidated)
 	{
-		Definition.ComponentInstancePropertiesInvalidated();
+		// todo: update metadata here
 	}
 	
 	FEntity::UpdateEntityProperties(Context);
@@ -1094,7 +1125,7 @@ void FComponentInstance::FillOccurrenceActorMetadata(FNodeOccurence& Node)
 
 	// Add original instance/component names to metadata
 	TSharedPtr<IDatasmithKeyValueProperty> EntityName = FDatasmithSceneFactory::CreateKeyValueProperty(TEXT("Instance"));
-	EntityName->SetValue(*GetName());
+	EntityName->SetValue(*GetEntityName());
 	Node.DatasmithMetadataElement->AddProperty(EntityName);
 
 	TSharedPtr<IDatasmithKeyValueProperty> DefinitionName = FDatasmithSceneFactory::CreateKeyValueProperty(TEXT("Definition"));
@@ -1131,7 +1162,7 @@ void FImageCollection::LayerModified(FEntityIDType LayerId)
 void FComponentInstance::UpdateOccurrenceVisibility(FExportContext& Context, FNodeOccurence& Node)
 {
 	// Parent node, component instance and layer - all should be visible to have node visible
-	Node.SetVisibility(Node.ParentNode->bVisible && !bHidden && bLayerVisible);
+	bool bVisibilityChanged = Node.SetVisibility(Node.ParentNode->bVisible && !bHidden && bLayerVisible);
 
 	EntityOccurrenceVisible(&Node, Node.bVisible);
 
@@ -1142,12 +1173,18 @@ void FComponentInstance::UpdateOccurrenceVisibility(FExportContext& Context, FNo
 	}
 	else
 	{
+		// Making component instance occurrence invisible  needs to invalidate geometry export
+		// for different reasons: this occurrence could have its own baked mesh, it could be a single used of an instanced mesh
+		GetDefinition()->InvalidateDefinitionGeometry();
 		Node.RemoveDatasmithActorHierarchy(Context);
 	}
 
 	for (FNodeOccurence* ChildNode : Node.Children)
 	{
-		ChildNode->bVisibilityInvalidated = true;
+		// Invalidate Visibility for child nodes when parent's was changed
+		// as visibility is hierarchical so children should update even
+		// if they weren't invalidated directly
+		ChildNode->bVisibilityInvalidated |= bVisibilityChanged;
 	}
 }
 
@@ -1189,7 +1226,12 @@ int64 FModel::GetPersistentId()
 	return 0;
 }
 
-FString FModel::GetName()
+FString FModel::GetEntityName()
+{
+	return "";
+}
+   
+FString FModel::GetEntityLabel()
 {
 	return "";
 }
diff --git a/Engine/Source/Programs/Enterprise/Datasmith/DatasmithSketchUpRubyExporter/Private/DatasmithSketchUpComponent.h b/Engine/Source/Programs/Enterprise/Datasmith/DatasmithSketchUpRubyExporter/Private/DatasmithSketchUpComponent.h
--- a/Engine/Source/Programs/Enterprise/Datasmith/DatasmithSketchUpRubyExporter/Private/DatasmithSketchUpComponent.h
+++ b/Engine/Source/Programs/Enterprise/Datasmith/DatasmithSketchUpRubyExporter/Private/DatasmithSketchUpComponent.h
@@ -40,10 +40,10 @@ namespace DatasmithSketchUp
 			: ParentNode(nullptr)
 			, Entity(InEntity)
 			, Depth(0)
-			,  bVisibilityInvalidated(true)
-			,  bPropertiesInvalidated(true)
-			,  bMeshActorsInvalidated(true)
-			,  bHierarchyInvalidated(true)
+			, bVisibilityInvalidated(true)
+			, bPropertiesInvalidated(true)
+			, bMeshActorsInvalidated(true)
+			, bTransformSupportedByUE(true)
 		{}
 
 		FNodeOccurence(FNodeOccurence* InParentNode, FEntity& InEntity)
@@ -53,11 +53,15 @@ namespace DatasmithSketchUp
 			, bVisibilityInvalidated(true)
 			, bPropertiesInvalidated(true)
 			, bMeshActorsInvalidated(true)
-			, bHierarchyInvalidated(true)
+			, bTransformSupportedByUE(true)
 		{}
 
-		// Update node hierarchy, bForceUpdate - to make update node and descendants when transform changes
+		// Update visibility of each node
 		void UpdateVisibility(FExportContext& Context);
+		// Parse tree to update transformation info:
+		// is transform can be directly converted to Unreal? Unreal only supports translation/rotation/scaling, with non-uniform scaling is only for leaf(bottommost) nodes when rotation is present
+		// todo: better name?
+		void UpdateTransformations(FExportContext& Context);
 		void Update(FExportContext& Context);
 
 		void ResetNodeActors(FExportContext& Context); // Reset actors before update
@@ -65,7 +69,7 @@ namespace DatasmithSketchUp
 
 		void InvalidateProperties(); // Invalidate name and transform. Invalidate propagates down the hierarchy - child transforms depend on the parent
 		void InvalidateMeshActors();
-		void SetVisibility(bool);
+		bool SetVisibility(bool);
 
 
 		FString GetActorName();
@@ -84,7 +88,17 @@ namespace DatasmithSketchUp
 
 		// Data that is computed from the hierarchy where Entity occurrence resides
 		int32 Depth;
+
+		// Original transform the node had in SketchUp
+		SUTransformation WorldTransformSource;
+		// Transform to have on a Datasmith actor
 		SUTransformation WorldTransform;
+		// Transform to have on a Datasmith Mesh actor
+		SUTransformation MeshActorWorldTransform;
+		// Local Transform to bake mesh with, so that MeshActorWorldTransform * BakeTransform = WorldTransformSource
+		// BakeTransform contains that part of the source transform which is not supported by Unreal(i.e. "skew"/"shear")
+		SUTransformation BakeTransform;
+
 		FMaterialIDType InheritedMaterialID;
 		SULayerRef EffectiveLayerRef = SU_INVALID;
 		bool bVisible = true; // Computed visibility for this occurrence(affecting descendants)
@@ -105,6 +119,7 @@ namespace DatasmithSketchUp
 		uint8 bPropertiesInvalidated:1; // Whether this occurrence properties(transform, name) need to be updated
 		uint8 bMeshActorsInvalidated:1; // Whether this occurrence MeshActors need updating. Happens initially when node was added and when node geometry is invalidated
 		uint8 bHierarchyInvalidated:1; // Children need to be rebuilt
+		uint8 bTransformSupportedByUE:1;
 	};
 
 	// For SketchUp's Definition that provides access to Entities and converts to Datasmith
@@ -123,7 +138,7 @@ namespace DatasmithSketchUp
 		void EntityVisible(FEntity* Entity, bool bVisible);
 
 		// Modification methods
-		virtual void AddInstance(FExportContext& Context, TSharedPtr<FComponentInstance> Instance) = 0; // Register child CompoenntInstance Entity of Definition's Entities
+		virtual void AddInstance(FExportContext& Context, TSharedPtr<FComponentInstance> Instance) = 0; // Register child ComponentInstance Entity of Definition's Entities
 		virtual void AddImage(FExportContext& Context, TSharedPtr<FImage> Image) = 0;
 
 		virtual void InvalidateInstancesGeometry(FExportContext& Context) = 0; // Mark that all instances(and their occurrences) needed to be updated
@@ -134,8 +149,6 @@ namespace DatasmithSketchUp
 		virtual FString GetSketchupSourceName() = 0; // Name used for label
 		virtual FString GetSketchupSourceId() = 0; // Unique name identifier
 
-		virtual SUTransformation GetMeshBakedTransform() = 0; // Transform that may be used to bake geometry
-
 		FEntities& GetEntities()
 		{
 			return *Entities;
@@ -186,8 +199,6 @@ namespace DatasmithSketchUp
 		virtual FString GetSketchupSourceGUID() override;
 		virtual FString GetSketchupSourceName() override;
 		virtual FString GetSketchupSourceId() override;
-
-		virtual SUTransformation GetMeshBakedTransform() override;
 		// End FDefinition
 
 		// Register/unregister instanced of this definition
@@ -195,10 +206,6 @@ namespace DatasmithSketchUp
 		void UnlinkComponentInstance(FComponentInstance* ComponentInstance);
 		void RemoveComponentDefinition(FExportContext& Context);
 
-		bool ShouldBakeTransformIntoMesh();
-		void GetLocalTransform(FComponentInstance& ComponentInstance, SUTransformation& SComponentInstanceLocalTransform);
-		void ComponentInstancePropertiesInvalidated();
-
 		// Source SketchUp component ID.
 		FComponentDefinitionIDType SketchupSourceID;
 		TSet<FComponentInstance*> Instances; // Tracked instances of this ComponentDefinition
@@ -210,8 +217,6 @@ namespace DatasmithSketchUp
 		SUComponentDefinitionRef ComponentDefinitionRef;
 
 		TUniquePtr<FMetadata> ParsedMetadata; // Shared metadata parsed from source SU component to be added to each occurrence actor's datasmith metatada element
-
-		bool bBakeTransformIntoMesh = false;
 	};
 
 
@@ -236,8 +241,6 @@ namespace DatasmithSketchUp
 		virtual FString GetSketchupSourceGUID() override;
 		virtual FString GetSketchupSourceName()  override;
 		virtual FString GetSketchupSourceId() override;
-
-		virtual SUTransformation GetMeshBakedTransform() override;
 		// End FDefinition
 
 		bool UpdateModel(FExportContext& Context);
@@ -256,7 +259,7 @@ namespace DatasmithSketchUp
 
 		FEntities(FDefinition& InDefinition) : Definition(InDefinition) {}
 
-		void UpdateGeometry(FExportContext& Context);
+		void UpdateGeometry(FExportContext& Context, TArray<FNodeOccurence*> NodesToInstance, TArray<FNodeOccurence*> NodesToBake);
 		void AddMeshesToDatasmithScene(FExportContext& Context);
 		void RemoveMeshesFromDatasmithScene(FExportContext& Context);
 
@@ -275,34 +278,53 @@ namespace DatasmithSketchUp
 	{
 	public:
 
-		int32 GetMeshCount()
+
+		class FExportedGeometry
+		{
+		public:
+			int32 GetMeshCount()
+			{
+				return Meshes.Num();
+			}
+
+			const TCHAR* GetMeshElementName(int32 MeshIndex);
+			bool IsMeshUsingInheritedMaterial(int32 MeshIndex);
+
+			TArray<TSharedPtr<FDatasmithInstantiatedMesh>> Meshes;
+		};
+
+		FExportedGeometry& GetOccurrenceExportedGeometry(FNodeOccurence& Node)
 		{
-			return Meshes.Num();
+			if (FExportedGeometry* Found = ExportedGeometryForNode.Find(&Node))
+			{
+				// todo: remove check
+				check(!Node.bTransformSupportedByUE);
+				return *Found;
+			}
+
+			check(Node.bTransformSupportedByUE || 0==ExportedGeometryForInstances.GetMeshCount());
+			return ExportedGeometryForInstances;
 		}
 
-		const TCHAR* GetMeshElementName(int32 MeshIndex);
-		bool IsMeshUsingInheritedMaterial(int32 MeshIndex);
 		int32 GetInheritedMaterialOverrideSlotId();
 
-		TArray<TSharedPtr<FDatasmithInstantiatedMesh>> Meshes;
+		void SetMaterial(const TCHAR* MaterialName, TFunctionRef<bool(const FDatasmithInstantiatedMesh& Mesh, int32& OutSlotId)> SlotMapping);
+
+		void ForEachExportedMesh(TFunctionRef<void(FDatasmithInstantiatedMesh&)> Callback);
+		void ExportOneMesh(FExportContext& Context, const TSharedPtr<class FDatasmithSketchUpMesh>& ExtractedMesh, FExportedGeometry& ExportedGeometry, int32 MeshIndex, const FString& MeshElementName, const FString& MeshLabel, SUTransformation Transform);
+
+		// Geometry exported for instanced nodes, which doesn't need baking transform into the exported geometry
+		FExportedGeometry ExportedGeometryForInstances;
+		// Geometry exported for nodes, which require baking transform into the exported geometry for proper result
+		TArray<TPair<SUTransformation, FExportedGeometry>>  ExportedGeometryForTransform;
+		TMap<FNodeOccurence*, FExportedGeometry>  ExportedGeometryForNode;
+
+		// Extracted data
 		TSet<int32> FaceIds; // EntityId of all the VISIBLE faces composing the mesh
 		TSet<DatasmithSketchUp::FEntityIDType> Layers; // EntityId of all layers assigned to geometry faces(needed to identify if geometry needs to be rebuilt when layer visibility changes)
-
 		TSet<FMaterial*> MaterialsUsed;
 		bool bDefaultMaterialUsed = false;
 
-		bool IsDefaultLayerMaterialUsed()
-		{
-			return bDefaultMaterialUsed;
-		}
-
-		void SetDefaultLayerMaterialUsed()
-		{
-			bDefaultMaterialUsed = true;
-		}
-
-		// todo: update reusing datasmith elements? 
-		// todo: update occurrences that used this entities - MeshActors need to be refreshed in accordance to OR this could be done on a level higher?
 	};
 
 	// Interface to implement SketchUp Entity Node - i.e. an instance of a ComponentDefinition(ComponentInstance or Group), Model, Image
@@ -319,7 +341,8 @@ namespace DatasmithSketchUp
 		virtual ~FEntity() {}
 
 		virtual int64 GetPersistentId() = 0;
-		virtual FString GetName() = 0;
+		virtual FString GetEntityName() = 0;
+		virtual FString GetEntityLabel() = 0;
 
 		virtual void ApplyOverrideMaterialToNode(FNodeOccurence& Node, FMaterialOccurrence& Material) = 0;
 
@@ -327,6 +350,8 @@ namespace DatasmithSketchUp
 		virtual void UpdateOccurrenceVisibility(FExportContext& Context, FNodeOccurence&) = 0;  // Re-evaluate visibility of entity's occurrence
 		virtual void UpdateOccurrenceMeshActors(FExportContext& Context, FNodeOccurence& Node) = 0; // Rebuild datasmith actors of entity's occurrence
 		virtual void ResetOccurrenceActors(FExportContext& Context, FNodeOccurence& Node) = 0; // Remove datasmith actors of entity's occurrence from datasmith scene
+		virtual void UpdateOccurrenceTransformation(FExportContext& Context, FNodeOccurence& Node) {}
+
 
 		virtual void InvalidateOccurrencesGeometry(FExportContext& Context) = 0;
 		virtual void InvalidateOccurrencesProperties(FExportContext& Context) = 0;
@@ -397,10 +422,7 @@ namespace DatasmithSketchUp
 	public:
 		FComponentDefinition& Definition;
 
-		FComponentInstance(SUEntityRef InEntityRef, FComponentDefinition& InDefinition)
-			: Super(InEntityRef)
-			, Definition(InDefinition)
-		{}
+		FComponentInstance(SUEntityRef InEntityRef, FComponentDefinition& InDefinition);
 
 		virtual ~FComponentInstance() override;
 
@@ -409,10 +431,12 @@ namespace DatasmithSketchUp
 		virtual bool GetAssignedMaterial(FMaterialIDType& MaterialId) override;
 		virtual void InvalidateOccurrencesGeometry(FExportContext& Context) override;
 		virtual void UpdateOccurrence(FExportContext& Context, FNodeOccurence& Node) override;
+		virtual void UpdateOccurrenceTransformation(FExportContext& Context, FNodeOccurence& Node);
 
 		virtual void InvalidateOccurrencesProperties(FExportContext& Context) override;
 		virtual int64 GetPersistentId() override;
-		virtual FString GetName() override;
+		virtual FString GetEntityName() override;
+		virtual FString GetEntityLabel() override;
 		virtual void UpdateOccurrenceVisibility(FExportContext& Context, FNodeOccurence&) override;
 		virtual void UpdateMetadata(FExportContext& Context) override;
 		virtual void UpdateEntityProperties(FExportContext& Context) override;
@@ -447,12 +471,12 @@ namespace DatasmithSketchUp
 
 		FImage(SUImageRef InEntityRef);
 
-		virtual ~FImage() override
-		{
-		}
+		virtual ~FImage() override;
 
 		virtual int64 GetPersistentId() override;
-		virtual FString GetName() override;
+		virtual FString GetEntityName() override;
+		virtual FString GetEntityLabel() override;
+
 		virtual void ApplyOverrideMaterialToNode(FNodeOccurence& Node, FMaterialOccurrence& Material) override;
 		virtual void UpdateOccurrence(FExportContext& Context, FNodeOccurence& Node) override;
 		void RemoveImageFromDatasmithScene(FExportContext& Context);
@@ -506,7 +530,8 @@ namespace DatasmithSketchUp
 		// void UpdateOccurrence(FExportContext& Context, FNodeOccurence& Node) override;
 		virtual void InvalidateOccurrencesProperties(FExportContext& Context) override;
 		virtual int64 GetPersistentId() override;
-		virtual FString GetName() override;
+		virtual FString GetEntityName() override;
+		virtual FString GetEntityLabel() override;
 		virtual void UpdateOccurrenceVisibility(FExportContext& Context, FNodeOccurence&) override;
 
 		virtual void UpdateMetadata(FExportContext& Context) override;
diff --git a/Engine/Source/Programs/Enterprise/Datasmith/DatasmithSketchUpRubyExporter/Private/DatasmithSketchUpExportContext.cpp b/Engine/Source/Programs/Enterprise/Datasmith/DatasmithSketchUpRubyExporter/Private/DatasmithSketchUpExportContext.cpp
--- a/Engine/Source/Programs/Enterprise/Datasmith/DatasmithSketchUpRubyExporter/Private/DatasmithSketchUpExportContext.cpp
+++ b/Engine/Source/Programs/Enterprise/Datasmith/DatasmithSketchUpRubyExporter/Private/DatasmithSketchUpExportContext.cpp
@@ -87,6 +87,7 @@ void FExportContext::Populate()
 	// Setup root Node, based on Model
 	Model = MakeShared<FModel>(*ModelDefinition);
 	RootNode = MakeShared<FNodeOccurence>(*Model);
+	RootNode->WorldTransformSource = WorldTransform;
 	RootNode->WorldTransform = WorldTransform;
 	RootNode->EffectiveLayerRef = DefaultLayerRef;
 	// Name and label for root loose mesh actors
@@ -175,19 +176,26 @@ bool FExportContext::Update(bool bModifiedHint)
 	}
 
 	// Invalidate occurrences for changed instances first
+	// These methods analyse the state of entities and make invalidation state of different relevant properties up to date
+	// to use further to actually update the exported data
 	Model->UpdateEntityProperties(*this);
 	ComponentInstances.UpdateProperties();
 	Images.UpdateProperties();
 
+	// Update visibility and transformations for occurrences after all properties were checked for invalidation
 	// Update occurrences visibility(before updating meshes to make sure to skip updating unused meshes)
 	RootNode->UpdateVisibility(*this);
+	// Update occurrences transformations before updating definitions to find out what to do with geometry for different occurrences
+	// e.g. make instanced meshes or bake separate meshes for occurrences with unsupported transformations
+	// after visibility pass - visibility information is needed to make decision on mesh baking
+	RootNode->UpdateTransformations(*this);
 
 	// Update Datasmith Meshes after their usage was refreshed(in visibility update) and before node hierarchy update(where Mesh Actors are updated for meshes)
 	ModelDefinition->UpdateDefinition(*this);
 	ComponentDefinitions.Update();
 	Images.Update();
 
-	// ComponentInstances will invalidate occurrences 
+	// ComponentInstances will invalidate occurrences
 	Model->UpdateEntityGeometry(*this);
 	ComponentInstances.UpdateGeometry();
 
@@ -197,9 +205,12 @@ bool FExportContext::Update(bool bModifiedHint)
 	Materials.Update();
 
 	// Wait for mesh export to complete
-	for(TFuture<bool>& Task: MeshExportTasks)
+	for(TSharedFuture<bool>& Task: MeshExportTasks)
 	{
-		Task.Get();
+		if (!Task.Get())
+		{
+			DatasmithSketchUpUtils::ToRuby::LogWarn(TEXT("A mesh failed to export, see UnrealDatasmithSketchUp/log.txt"));
+		}
 	}
 	MeshExportTasks.Reset();
 
@@ -432,8 +443,6 @@ bool FLayerCollection::IsLayerVisible(SULayerRef LayerRef)
 	return Found ? *Found : true;
 }
 
-
-
 SULayerRef FLayerCollection::GetLayer(FLayerIDType LayerId)
 {
 	size_t LayerCount = 0;
diff --git a/Engine/Source/Programs/Enterprise/Datasmith/DatasmithSketchUpRubyExporter/Private/DatasmithSketchUpExportContext.h b/Engine/Source/Programs/Enterprise/Datasmith/DatasmithSketchUpRubyExporter/Private/DatasmithSketchUpExportContext.h
--- a/Engine/Source/Programs/Enterprise/Datasmith/DatasmithSketchUpRubyExporter/Private/DatasmithSketchUpExportContext.h
+++ b/Engine/Source/Programs/Enterprise/Datasmith/DatasmithSketchUpRubyExporter/Private/DatasmithSketchUpExportContext.h
@@ -373,6 +373,6 @@ namespace DatasmithSketchUp
 		bool bColorByLayer = false;
 		bool bColorByLayerInvaliated = true;
 
-		TArray<TFuture<bool>> MeshExportTasks;
+		TArray<TSharedFuture<bool>> MeshExportTasks;
 	};
 }
diff --git a/Engine/Source/Programs/Enterprise/Datasmith/DatasmithSketchUpRubyExporter/Private/DatasmithSketchUpExporter.cpp b/Engine/Source/Programs/Enterprise/Datasmith/DatasmithSketchUpRubyExporter/Private/DatasmithSketchUpExporter.cpp
--- a/Engine/Source/Programs/Enterprise/Datasmith/DatasmithSketchUpRubyExporter/Private/DatasmithSketchUpExporter.cpp
+++ b/Engine/Source/Programs/Enterprise/Datasmith/DatasmithSketchUpRubyExporter/Private/DatasmithSketchUpExporter.cpp
@@ -48,6 +48,7 @@
 // disable(SU2024): Dereferencing NULL pointer
 #pragma warning(disable: 6011)
 
+#define RB_BLOCK_CALL_FUNC_STRICT 1
 
 #undef DEPRECATED
 #include <ruby.h>
@@ -584,17 +585,17 @@ VALUE UnrealStringToRuby(const FString& InStr)
 {
 	FTCHARToUTF8 Converter(*InStr, InStr.Len());
 
-	return rb_str_new(Converter.Get(), Converter.Length());
+	return rb_utf8_str_new(Converter.Get(), Converter.Length());
 }
 
 typedef VALUE(*RubyFunctionType)(ANYARGS);
 
 #pragma warning(push)
 // disable: 'reinterpret_cast': unsafe conversion from 'F' to 'RubyFunctionType'
-// This IS unsafe but so is Ruby C API
+// This IS unsafe but so is Ruby C API when ANYARGS function is used
 #pragma warning(disable: 4191)
 template<typename F>
-RubyFunctionType ToRuby(F f)
+RubyFunctionType ToRubyAnyargs(F f)
 {
 	return reinterpret_cast<RubyFunctionType>(f);
 }
@@ -938,44 +939,83 @@ extern "C" DLLEXPORT void Init_DatasmithSketchUp()
 	VALUE EpicGames = rb_define_module("EpicGames");
 	VALUE Datasmith = rb_define_module_under(EpicGames, "DatasmithBackend");
 
-	rb_define_module_function(Datasmith, "on_load", ToRuby(on_load), 2);
-	rb_define_module_function(Datasmith, "on_unload", ToRuby(on_unload), 0);
+	rb_define_module_function(Datasmith, "on_load", ToRubyAnyargs(on_load), 2);
+	rb_define_module_function(Datasmith, "on_unload", ToRubyAnyargs(on_unload), 0);
 
-	rb_define_module_function(Datasmith, "open_directlink_ui", ToRuby(open_directlink_ui), 0);
-	rb_define_module_function(Datasmith, "get_directlink_cache_directory", ToRuby(get_directlink_cache_directory), 0);
+	rb_define_module_function(Datasmith, "open_directlink_ui", ToRubyAnyargs(open_directlink_ui), 0);
+	rb_define_module_function(Datasmith, "get_directlink_cache_directory", ToRubyAnyargs(get_directlink_cache_directory), 0);
 
 	DatasmithSketchUpDirectLinkExporterCRubyClass = rb_define_class_under(Datasmith, "DatasmithSketchUpDirectLinkExporter", rb_cObject);
 
-	rb_define_singleton_method(DatasmithSketchUpDirectLinkExporterCRubyClass, "new", ToRuby(DatasmithSketchUpDirectLinkExporter_new), 3);
-	rb_define_method(DatasmithSketchUpDirectLinkExporterCRubyClass, "start", ToRuby(DatasmithSketchUpDirectLinkExporter_start), 0);
-	rb_define_method(DatasmithSketchUpDirectLinkExporterCRubyClass, "stop", ToRuby(DatasmithSketchUpDirectLinkExporter_stop), 0);
+	rb_define_singleton_method(DatasmithSketchUpDirectLinkExporterCRubyClass, "new", ToRubyAnyargs(DatasmithSketchUpDirectLinkExporter_new), 3);
+	rb_define_method(DatasmithSketchUpDirectLinkExporterCRubyClass, "start", ToRubyAnyargs(DatasmithSketchUpDirectLinkExporter_start), 0);
+	rb_define_method(DatasmithSketchUpDirectLinkExporterCRubyClass, "stop", ToRubyAnyargs(DatasmithSketchUpDirectLinkExporter_stop), 0);
 
 #ifndef SKP_SDK_2019
-	rb_define_method(DatasmithSketchUpDirectLinkExporterCRubyClass, "on_component_instance_changed", ToRuby(DatasmithSketchUpDirectLinkExporter_on_component_instance_changed), 1);
-	rb_define_method(DatasmithSketchUpDirectLinkExporterCRubyClass, "on_entity_added", ToRuby(DatasmithSketchUpDirectLinkExporter_on_entity_added), 2);
-	rb_define_method(DatasmithSketchUpDirectLinkExporterCRubyClass, "on_layer_modified", ToRuby(DatasmithSketchUpDirectLinkExporter_on_layer_modified), 1);
+	rb_define_method(DatasmithSketchUpDirectLinkExporterCRubyClass, "on_component_instance_changed", ToRubyAnyargs(DatasmithSketchUpDirectLinkExporter_on_component_instance_changed), 1);
+	rb_define_method(DatasmithSketchUpDirectLinkExporterCRubyClass, "on_entity_added", ToRubyAnyargs(DatasmithSketchUpDirectLinkExporter_on_entity_added), 2);
+	rb_define_method(DatasmithSketchUpDirectLinkExporterCRubyClass, "on_layer_modified", ToRubyAnyargs(DatasmithSketchUpDirectLinkExporter_on_layer_modified), 1);
 #endif
 
-	rb_define_method(DatasmithSketchUpDirectLinkExporterCRubyClass, "on_entity_modified_by_id", ToRuby(DatasmithSketchUpDirectLinkExporter_on_entity_modified_by_id), 1);
-	rb_define_method(DatasmithSketchUpDirectLinkExporterCRubyClass, "on_geometry_modified_by_id", ToRuby(DatasmithSketchUpDirectLinkExporter_on_geometry_modified_by_id), 1);
-	rb_define_method(DatasmithSketchUpDirectLinkExporterCRubyClass, "on_entity_added_by_id", ToRuby(DatasmithSketchUpDirectLinkExporter_on_entity_added_by_id), 2);
-	rb_define_method(DatasmithSketchUpDirectLinkExporterCRubyClass, "on_material_added_by_id", ToRuby(DatasmithSketchUpDirectLinkExporter_on_material_added_by_id), 1);
+	rb_define_method(DatasmithSketchUpDirectLinkExporterCRubyClass, "on_entity_modified_by_id", ToRubyAnyargs(DatasmithSketchUpDirectLinkExporter_on_entity_modified_by_id), 1);
+	rb_define_method(DatasmithSketchUpDirectLinkExporterCRubyClass, "on_geometry_modified_by_id", ToRubyAnyargs(DatasmithSketchUpDirectLinkExporter_on_geometry_modified_by_id), 1);
+	rb_define_method(DatasmithSketchUpDirectLinkExporterCRubyClass, "on_entity_added_by_id", ToRubyAnyargs(DatasmithSketchUpDirectLinkExporter_on_entity_added_by_id), 2);
+	rb_define_method(DatasmithSketchUpDirectLinkExporterCRubyClass, "on_material_added_by_id", ToRubyAnyargs(DatasmithSketchUpDirectLinkExporter_on_material_added_by_id), 1);
 
-	rb_define_method(DatasmithSketchUpDirectLinkExporterCRubyClass, "on_entity_removed", ToRuby(DatasmithSketchUpDirectLinkExporter_on_entity_removed), 2);
+	rb_define_method(DatasmithSketchUpDirectLinkExporterCRubyClass, "on_entity_removed", ToRubyAnyargs(DatasmithSketchUpDirectLinkExporter_on_entity_removed), 2);
 
-	rb_define_method(DatasmithSketchUpDirectLinkExporterCRubyClass, "on_style_changed", ToRuby(DatasmithSketchUpDirectLinkExporter_on_style_changed), 0);
-	rb_define_method(DatasmithSketchUpDirectLinkExporterCRubyClass, "on_color_by_layer_changed", ToRuby(DatasmithSketchUpDirectLinkExporter_on_color_by_layer_changed), 0);
+	rb_define_method(DatasmithSketchUpDirectLinkExporterCRubyClass, "on_style_changed", ToRubyAnyargs(DatasmithSketchUpDirectLinkExporter_on_style_changed), 0);
+	rb_define_method(DatasmithSketchUpDirectLinkExporterCRubyClass, "on_color_by_layer_changed", ToRubyAnyargs(DatasmithSketchUpDirectLinkExporter_on_color_by_layer_changed), 0);
 
-	rb_define_method(DatasmithSketchUpDirectLinkExporterCRubyClass, "update", ToRuby(DatasmithSketchUpDirectLinkExporter_update), 1);
-	rb_define_method(DatasmithSketchUpDirectLinkExporterCRubyClass, "send_update", ToRuby(DatasmithSketchUpDirectLinkExporter_send_update), 0);
-	rb_define_method(DatasmithSketchUpDirectLinkExporterCRubyClass, "export_current_datasmith_scene", ToRuby(DatasmithSketchUpDirectLinkExporter_export_current_datasmith_scene), 0);
+	rb_define_method(DatasmithSketchUpDirectLinkExporterCRubyClass, "update", ToRubyAnyargs(DatasmithSketchUpDirectLinkExporter_update), 1);
+	rb_define_method(DatasmithSketchUpDirectLinkExporterCRubyClass, "send_update", ToRubyAnyargs(DatasmithSketchUpDirectLinkExporter_send_update), 0);
+	rb_define_method(DatasmithSketchUpDirectLinkExporterCRubyClass, "export_current_datasmith_scene", ToRubyAnyargs(DatasmithSketchUpDirectLinkExporter_export_current_datasmith_scene), 0);
 
-	rb_define_method(DatasmithSketchUpDirectLinkExporterCRubyClass, "get_connection_status", ToRuby(DatasmithSketchUpDirectLinkExporter_get_connection_status), 0);
+	rb_define_method(DatasmithSketchUpDirectLinkExporterCRubyClass, "get_connection_status", ToRubyAnyargs(DatasmithSketchUpDirectLinkExporter_get_connection_status), 0);
 
-	rb_define_method(DatasmithSketchUpDirectLinkExporterCRubyClass, "set_active_scene", ToRuby(DatasmithSketchUpDirectLinkExporter_set_active_scene), 1);
+	rb_define_method(DatasmithSketchUpDirectLinkExporterCRubyClass, "set_active_scene", ToRubyAnyargs(DatasmithSketchUpDirectLinkExporter_set_active_scene), 1);
 
 }
 
+namespace DatasmithSketchUpUtils
+{
+
+namespace ToRuby
+{
+	const FString* LogWarn_Block_Value = nullptr;
+
+	VALUE LogWarn_Block(VALUE obj, VALUE context, int argc, const VALUE* argv, VALUE blockarg)
+	{
+		if (LogWarn_Block_Value)
+		{
+			return UnrealStringToRuby(*LogWarn_Block_Value);
+		}
+		return Qnil;
+	}
+
+	void LogWarn(const FString& Message)
+	{
+		// call: EpicGames::Datasmith.log_warn { "the message text" }
+		VALUE DatasmithRubyNamespace = rb_path2class("EpicGames::Datasmith");
+		TGuardValue BlockValueGuard(LogWarn_Block_Value, &Message);
+
+#if defined(__clang__) || defined(__GNUC__)
+#pragma GCC diagnostic push
+// disable(pre-SU024): '(' and '{' tokens introducing statement expression appear in different macro expansion contexts [-Wcompound-token-split-by-macro]
+// a problem in Ruby's rb_intern implementation
+#pragma GCC diagnostic ignored "-Wcompound-token-split-by-macro"
+#endif
+		ID LogWarnSym = rb_intern("log_warn");
+#if defined(__clang__) || defined(__GNUC__)
+#pragma GCC diagnostic pop
+#endif
+		rb_block_call(DatasmithRubyNamespace, LogWarnSym, 0, nullptr, LogWarn_Block, Qtrue);
+	}
+}
+
+}
+
+
 /* todo:
 - error reporting(Ruby Console, Log etc)
 */
diff --git a/Engine/Source/Programs/Enterprise/Datasmith/DatasmithSketchUpRubyExporter/Private/DatasmithSketchUpImage.cpp b/Engine/Source/Programs/Enterprise/Datasmith/DatasmithSketchUpRubyExporter/Private/DatasmithSketchUpImage.cpp
--- a/Engine/Source/Programs/Enterprise/Datasmith/DatasmithSketchUpRubyExporter/Private/DatasmithSketchUpImage.cpp
+++ b/Engine/Source/Programs/Enterprise/Datasmith/DatasmithSketchUpRubyExporter/Private/DatasmithSketchUpImage.cpp
@@ -38,6 +38,10 @@ FImage::FImage(SUImageRef InEntityRef)
 	MeshElementName = FString::Printf(TEXT("M%s"), *IdStr);
 }
 
+FImage::~FImage() 
+{
+}
+
 void FImage::RemoveImage(FExportContext& Context)
 {
 	RemoveOccurrences(Context);
@@ -85,7 +89,12 @@ void FImage::InvalidateImage()
 }
 
 
-FString FImage::GetName()
+FString FImage::GetEntityName()
+{
+	return FDatasmithUtils::SanitizeObjectName(GetFileName());
+}
+
+FString FImage::GetEntityLabel()
 {
 	return FDatasmithUtils::SanitizeObjectName(GetFileName());
 }
@@ -133,8 +142,8 @@ void FImage::UpdateOccurrence(FExportContext& Context, FNodeOccurence& Node)
 	SUTransformationSetFromPointAndAxes(&LocalTransform, &Point, &XAxis, &YAxis,  &ZAxis);
 
 	SUTransformation WorldTransform;
-	SUTransformationMultiply(&Node.ParentNode->WorldTransform, &LocalTransform, &WorldTransform);
-	Node.WorldTransform = WorldTransform; // Store world transform to be used by children to compute its
+	SUTransformationMultiply(&Node.ParentNode->WorldTransformSource, &LocalTransform, &WorldTransform);
+	Node.WorldTransform = WorldTransform;
 
 	// Set the Datasmith actor world transform.
 	DatasmithSketchUpUtils::SetActorTransform(Node.DatasmithActorElement, Node.WorldTransform);
@@ -212,8 +221,8 @@ void FImage::BuildNodeNames(FNodeOccurence& Node)
 	int64 SketchupPersistentID = Node.Entity.GetPersistentId();
 	Node.DatasmithActorName = FString::Printf(TEXT("%ls_%lld"), *Node.ParentNode->GetActorName(), SketchupPersistentID);
 
-	FString EntityName = Node.Entity.GetName();
-	Node.DatasmithActorLabel = FDatasmithUtils::SanitizeObjectName(EntityName.IsEmpty() ? GetName() : EntityName);
+	FString EntityName = Node.Entity.GetEntityName();
+	Node.DatasmithActorLabel = FDatasmithUtils::SanitizeObjectName(EntityName.IsEmpty() ? GetEntityName() : EntityName);
 }
 
 void FImage::SetupActor(FExportContext& Context, FNodeOccurence& Node)
diff --git a/Engine/Source/Programs/Enterprise/Datasmith/DatasmithSketchUpRubyExporter/Private/DatasmithSketchUpMaterial.cpp b/Engine/Source/Programs/Enterprise/Datasmith/DatasmithSketchUpRubyExporter/Private/DatasmithSketchUpMaterial.cpp
--- a/Engine/Source/Programs/Enterprise/Datasmith/DatasmithSketchUpRubyExporter/Private/DatasmithSketchUpMaterial.cpp
+++ b/Engine/Source/Programs/Enterprise/Datasmith/DatasmithSketchUpRubyExporter/Private/DatasmithSketchUpMaterial.cpp
@@ -255,11 +255,13 @@ namespace DatasmithSketchUp
 	{
 		FEntitiesGeometry& EntitiesGeometry = *GetEntities().EntitiesGeometry;
 
+		FEntitiesGeometry::FExportedGeometry& ExportedGeometry = EntitiesGeometry.GetOccurrenceExportedGeometry(Node);
+
 		for (int32 MeshIndex = 0; MeshIndex < Node.MeshActors.Num(); ++MeshIndex)
 		{
 			// Update Override(Inherited)  Material
 			// todo: set inherited material only on mesh actors that have faces with default material, right now setting on every mesh, hot harmful but excessive
-			if (EntitiesGeometry.IsMeshUsingInheritedMaterial(MeshIndex))
+			if (ExportedGeometry.IsMeshUsingInheritedMaterial(MeshIndex))
 			{
 				const TSharedPtr<IDatasmithMeshActorElement>& MeshActor = Node.MeshActors[MeshIndex];
 
@@ -475,6 +477,35 @@ const TCHAR* FMaterialOccurrence::GetName()
 	return DatasmithElement->GetName();
 }
 
+void FEntitiesGeometry::SetMaterial(const TCHAR* MaterialName, TFunctionRef<bool(const FDatasmithInstantiatedMesh& Mesh, int32& OutSlotId)> SlotMapping)
+{
+	ForEachExportedMesh([MaterialName, SlotMapping](const FDatasmithInstantiatedMesh& Mesh)
+	{
+		int32 SlotId;
+		if (SlotMapping(Mesh, SlotId))
+		{
+			Mesh.DatasmithMesh->SetMaterial(MaterialName, SlotId);
+		}
+	});
+}
+
+void FEntitiesGeometry::ForEachExportedMesh(TFunctionRef<void(FDatasmithInstantiatedMesh&)> Callback)
+{
+	for (const TSharedPtr<FDatasmithInstantiatedMesh>& Mesh : ExportedGeometryForInstances.Meshes)
+	{
+		Callback(*Mesh);
+	}
+
+	for (const TPair<SUTransformation, FExportedGeometry>& NodeExportedGeometry : ExportedGeometryForTransform)
+	{
+		for (TSharedPtr<FDatasmithInstantiatedMesh> Mesh : NodeExportedGeometry.Value.Meshes)
+		{
+			Callback(*Mesh);
+		}
+	}
+}
+
+
 // todo: make virtual FMaterialOccurrence - for layer/regular materials different?
 // they differ by the way they are applied to geometry/nodes -
 // VERY tied to meshes composition - they either map slots from SU Material or from SU Layer
@@ -483,14 +514,15 @@ void FMaterialOccurrence::ApplyRegularMaterial(FMaterialIDType MaterialId)
 	// Apply material to meshes
 	for (FEntitiesGeometry* Geometry : MeshesMaterialDirectlyAppliedTo)
 	{
-		//Geometry->SetMaterialElementName(DatasmithSketchUpUtils::GetMaterialID(MaterialRef), MaterialDirectlyAppliedToMeshes->GetName());
-		for (const TSharedPtr<FDatasmithInstantiatedMesh>& Mesh : Geometry->Meshes)
+		Geometry->SetMaterial(GetName(), [MaterialId](const FDatasmithInstantiatedMesh& Mesh, int32& OutSlotId)
 		{
-			if (int32* SlotIdPtr = Mesh->SlotIdForMaterialId.Find(MaterialId))
+			if (const int32* SlotIdPtr = Mesh.SlotIdForMaterialId.Find(MaterialId))
 			{
-				Mesh->DatasmithMesh->SetMaterial(GetName(), *SlotIdPtr);
+				OutSlotId = *SlotIdPtr;
+				return true;
 			}
-		}
+			return false;
+		});
 	}
 
 	// Apply material to mesh actors
@@ -505,16 +537,15 @@ void FMaterialOccurrence::ApplyLayerMaterial(FLayerIDType LayerId)
 	// Apply material to meshes
 	for (FEntitiesGeometry* Geometry : MeshesMaterialDirectlyAppliedTo)
 	{
-		//Geometry->SetMaterialElementName(DatasmithSketchUpUtils::GetMaterialID(MaterialRef), MaterialDirectlyAppliedToMeshes->GetName());
-		for (const TSharedPtr<FDatasmithInstantiatedMesh>& Mesh : Geometry->Meshes)
+		Geometry->SetMaterial(GetName(), [LayerId](const FDatasmithInstantiatedMesh& Mesh, int32& OutSlotId)
 		{
-			// todo: maybe can refactor this and unify FMaterialOccurrence for regular/layer materials?
-			// By moving this into the Mesh? It has this occurrence anyway so it can keep the association of material<->slot there?
-			if (int32* SlotIdPtr = Mesh->SlotIdForLayerId.Find(LayerId))
+			if (const int32* SlotIdPtr = Mesh.SlotIdForLayerId.Find(LayerId))
 			{
-				Mesh->DatasmithMesh->SetMaterial(GetName(), *SlotIdPtr);
+				OutSlotId = *SlotIdPtr;
+				return true;
 			}
-		}
+			return false;
+		});
 	}
 
 	// Apply material to mesh actors
@@ -770,7 +801,7 @@ public:
 
 	FImageMaterial(FExportContext& Context, FImage& Image, FImageFile& ImageFile)
 	{
-		FString NameBase = Image.GetName();
+		FString NameBase = Image.GetEntityName();
 
 		FString MaterialName = FString::Printf(TEXT("image_material_%s"), *NameBase);  
 		FString MaterialLabel = NameBase;
diff --git a/Engine/Source/Programs/Enterprise/Datasmith/DatasmithSketchUpRubyExporter/Private/DatasmithSketchUpMesh.cpp b/Engine/Source/Programs/Enterprise/Datasmith/DatasmithSketchUpRubyExporter/Private/DatasmithSketchUpMesh.cpp
--- a/Engine/Source/Programs/Enterprise/Datasmith/DatasmithSketchUpRubyExporter/Private/DatasmithSketchUpMesh.cpp
+++ b/Engine/Source/Programs/Enterprise/Datasmith/DatasmithSketchUpRubyExporter/Private/DatasmithSketchUpMesh.cpp
@@ -70,7 +70,7 @@ namespace DatasmithSketchUp
 	public:
 
 		// Convert the combined mesh into a Datasmith mesh.
-		void ConvertMeshToDatasmith(FExportContext& Context, SUTransformation Transform, FDatasmithMesh& OutDMesh) const;
+		void ConvertMeshToDatasmith(FExportContext& Context, const SUTransformation& Transform, FDatasmithMesh& OutDMesh) const;
 
 		// Tessellate a SketchUp face into a triangle mesh merged into the combined mesh.
 		void AddFace(FExportContext& Context, SUFaceRef InSFaceRef, FLayerIDType LayerId);
@@ -326,7 +326,7 @@ namespace DatasmithSketchUp
 		}
 	}
 
-	void FDatasmithSketchUpMesh::ConvertMeshToDatasmith(FExportContext& Context, SUTransformation Transform, FDatasmithMesh& OutDMesh) const
+	void FDatasmithSketchUpMesh::ConvertMeshToDatasmith(FExportContext& Context, const SUTransformation& Transform, FDatasmithMesh& OutDMesh) const
 	{
 		// Get the number of mesh vertices (must be > 0).
 		int32 VertexCount = MeshVertexPoints.Num();
@@ -334,6 +334,9 @@ namespace DatasmithSketchUp
 		// Set the number of vertices of the exported Datasmith mesh.
 		OutDMesh.SetVerticesCount(VertexCount);
 
+		bool bIsMirroredTransform = false;
+		SUTransformationIsMirrored(&Transform, &bIsMirroredTransform);
+
 		for (int32 VertexNo = 0; VertexNo < VertexCount; VertexNo++)
 		{
 			SUPoint3D VertexPoint = MeshVertexPoints[VertexNo];
@@ -399,18 +402,26 @@ namespace DatasmithSketchUp
 
 			// Set the triangle vertex indices in the exported Datasmith mesh.
 			SMeshTriangleIndices const& TriangleIndices = MeshTriangleIndices[TriangleNo];
-			OutDMesh.SetFace(TriangleNo, int32(TriangleIndices.IndexA), int32(TriangleIndices.IndexB), int32(TriangleIndices.IndexC), MeshTriangleSlotIds[TriangleNo]);
+			int32 IndexA = static_cast<int32>(TriangleIndices.IndexA);
+			int32 IndexB = static_cast<int32>(TriangleIndices.IndexB);
+			int32 IndexC = static_cast<int32>(TriangleIndices.IndexC);
+			if (bIsMirroredTransform)
+			{
+				// Reorient triangle if transform is changing handedness
+				std::swap(IndexB, IndexC);
+			}
+			OutDMesh.SetFace(TriangleNo, IndexA, IndexB, IndexC, MeshTriangleSlotIds[TriangleNo]);
 
 			// Set the triangle vertex normals in the exported Datasmith mesh.
-			SMeshTriangleNormals TriangleNormals = { MeshVertexNormalsBaked[TriangleIndices.IndexA],
-													 MeshVertexNormalsBaked[TriangleIndices.IndexB],
-													 MeshVertexNormalsBaked[TriangleIndices.IndexC] };
+			SMeshTriangleNormals TriangleNormals = { MeshVertexNormalsBaked[IndexA],
+													 MeshVertexNormalsBaked[IndexB],
+													 MeshVertexNormalsBaked[IndexC] };
 			OutDMesh.SetNormal(NormalNo++, float(TriangleNormals.NormalA.x), float(-TriangleNormals.NormalA.y), float(TriangleNormals.NormalA.z));
 			OutDMesh.SetNormal(NormalNo++, float(TriangleNormals.NormalB.x), float(-TriangleNormals.NormalB.y), float(TriangleNormals.NormalB.z));
 			OutDMesh.SetNormal(NormalNo++, float(TriangleNormals.NormalC.x), float(-TriangleNormals.NormalC.y), float(TriangleNormals.NormalC.z));
 
 			// Set the triangle UV coordinate indices in the exported Datasmith mesh.
-			OutDMesh.SetFaceUV(TriangleNo, 0, int32(TriangleIndices.IndexA), int32(TriangleIndices.IndexB), int32(TriangleIndices.IndexC));
+			OutDMesh.SetFaceUV(TriangleNo, 0, IndexA, IndexB, IndexC);
 		}
 	}
 
@@ -517,7 +528,7 @@ namespace DatasmithSketchUp
 }
 
 
-bool FEntitiesGeometry::IsMeshUsingInheritedMaterial(int32 MeshIndex)
+bool FEntitiesGeometry::FExportedGeometry::IsMeshUsingInheritedMaterial(int32 MeshIndex)
 {
 	return Meshes[MeshIndex]->bIsUsingInheritedMaterial;
 }
@@ -528,15 +539,85 @@ int32 FEntitiesGeometry::GetInheritedMaterialOverrideSlotId()
 }
 
 
-const TCHAR* FEntitiesGeometry::GetMeshElementName(int32 MeshIndex)
+const TCHAR* FEntitiesGeometry::FExportedGeometry::GetMeshElementName(int32 MeshIndex)
 {
 	return Meshes[MeshIndex]->DatasmithMesh->GetName();
 }
 
+void FEntitiesGeometry::ExportOneMesh(FExportContext& Context, const TSharedPtr<FDatasmithSketchUpMesh>& ExtractedMesh, FExportedGeometry& ExportedGeometry, int32 MeshIndex, const FString& MeshElementName, const FString& MeshLabel, SUTransformation Transform)
+{
+	TSharedPtr<FDatasmithInstantiatedMesh> Mesh;
+	// Create MeshElement or reuse existing
+	if (MeshIndex < ExportedGeometry.Meshes.Num())
+	{
+		Mesh = ExportedGeometry.Meshes[MeshIndex];
+		Mesh->SlotIdForMaterialId.Reset();
+		Mesh->SlotIdForLayerId.Reset();
+	}
+	else
+	{
+		Mesh = MakeShared<FDatasmithInstantiatedMesh>();
+		ExportedGeometry.Meshes.Add(Mesh);
+	}
+	// todo: reuse DatasmithMesh when it allows to reset material slots
+	Mesh->DatasmithMesh = FDatasmithSceneFactory::CreateMesh(TEXT(""));
+
+	Mesh->DatasmithMesh->SetName(*MeshElementName);
+	Mesh->DatasmithMesh->SetLabel(*MeshLabel);
+	Mesh->bIsUsingInheritedMaterial = ExtractedMesh->bHasFacesWithDefaultMaterial;
+
+	// Add the non-inherited materials used by the combined mesh triangles.
+	if (Context.bColorByLayer)
+	{
+		for (int32 SlotId = 0;SlotId < ExtractedMesh->LayerIDForSlotId.Num(); ++SlotId)
+		{
+			FLayerIDType LayerID = ExtractedMesh->LayerIDForSlotId[SlotId];
+			Mesh->SlotIdForLayerId.FindOrAdd(LayerID, SlotId);
+
+			if (FMaterialOccurrence* Material = Context.Materials.LayerMaterials.RegisterGeometryForLayer(LayerID, this))
+			{
+				Mesh->DatasmithMesh->SetMaterial(Material->GetName(), SlotId);
+			}
+		}
+	}
+	else
+	{
+		for (int32 SlotId = 0;SlotId < ExtractedMesh->MaterialIDForSlotId.Num(); ++SlotId)
+		{
+			if (SlotId == 0 && !ExtractedMesh->bHasFacesWithDefaultMaterial)
+			{
+				continue; // Skip adding Default material slot if it's not used
+			}
+
+			FMaterialIDType MeshMaterialID = ExtractedMesh->MaterialIDForSlotId[SlotId];
+			Mesh->SlotIdForMaterialId.FindOrAdd(MeshMaterialID, SlotId);
+			// Default or (somehow)missing materials are also assigned to mesh(as a default material)
+			if (FMaterialOccurrence* Material = Context.Materials.RegularMaterials.RegisterGeometry(MeshMaterialID, this))
+			{
+				Mesh->DatasmithMesh->SetMaterial(Material->GetName(), SlotId);
+			}
+		}
+	}
+	
+	Mesh->ExportFuture = Context.MeshExportTasks.Emplace_GetRef(Async(
+		EAsyncExecution::ThreadPool,
+		[&Context, Mesh, ExtractedMesh, Transform]()
+		{
+			FDatasmithMeshExporter DatasmithMeshExporter;
+			FDatasmithMesh DatasmithMesh;
+			ExtractedMesh->ConvertMeshToDatasmith(Context, Transform, DatasmithMesh);
+
+			FGCScopeGuard GCGuard; // Prevent GC from running while UDatasmithMesh is created in ExportToUObject.
+
+			return DatasmithMeshExporter.ExportToUObject(Mesh->DatasmithMesh, Context.GetAssetsOutputPath(), DatasmithMesh, nullptr, FDatasmithExportOptions::LightmapUV);
+		}
+	).Share());
+}
+
 void ScanSketchUpEntitiesFaces(FExportContext& Context, SUEntitiesRef EntitiesRef, FEntitiesGeometry& Geometry, TFunctionRef<void(TSharedPtr<FDatasmithSketchUpMesh> ExtractedMesh)> OnNewExtractedMesh);
 void CombineSketchUpEntitiesFaces(FExportContext& Context, SUEntitiesRef EntitiesRef, FEntitiesGeometry& Geometry, TFunctionRef<void(TSharedPtr<FDatasmithSketchUpMesh> ExtractedMesh)> OnNewExtractedMesh);
 
-void FEntities::UpdateGeometry(FExportContext& Context)
+void FEntities::UpdateGeometry(FExportContext& Context, TArray<FNodeOccurence*> NodesToInstance, TArray<FNodeOccurence*> NodesToBake)
 {
 	if (EntitiesGeometry.IsValid())
 	{
@@ -547,6 +628,9 @@ void FEntities::UpdateGeometry(FExportContext& Context)
 		Context.EntitiesObjects.UnregisterEntities(*this);
 		EntitiesGeometry->FaceIds.Reset();
 		EntitiesGeometry->Layers.Reset();
+
+		EntitiesGeometry->ExportedGeometryForNode.Reset();
+		EntitiesGeometry->ExportedGeometryForTransform.Reset();
 	}
 	else
 	{
@@ -555,83 +639,54 @@ void FEntities::UpdateGeometry(FExportContext& Context)
 
 	int32 MeshCount = 0;
 
-	TFunction<void(TSharedPtr<FDatasmithSketchUpMesh> ExtractedMesh)> ProcessExtractedMesh = [&Context, this, &MeshCount](TSharedPtr<FDatasmithSketchUpMesh> ExtractedMeshPtr)
+	TFunction<void(TSharedPtr<FDatasmithSketchUpMesh> ExtractedMesh)> ProcessExtractedMesh = [&Context, this, &MeshCount,
+		&NodesToInstance, &NodesToBake](TSharedPtr<FDatasmithSketchUpMesh> ExtractedMeshPtr)
 	{
 		if (ExtractedMeshPtr->ContainsGeometry())
 		{
-			FString MeshElementName = FString::Printf(TEXT("M%ls_%d"), *Definition.GetSketchupSourceId(), MeshCount + 1); // Count meshes from 1
 			FString MeshLabel = FDatasmithUtils::SanitizeObjectName(Definition.GetSketchupSourceName());
 
-			TSharedPtr<FDatasmithInstantiatedMesh> Mesh;
-			// Create MeshElement or reuse existing
-			if (MeshCount < EntitiesGeometry->Meshes.Num())
-			{
-				Mesh = EntitiesGeometry->Meshes[MeshCount];
-				Mesh->SlotIdForMaterialId.Reset();
-				Mesh->SlotIdForLayerId.Reset();
-			}
-			else
+			// Export mesh for 'regular' instances without applying a transform
+			if (!NodesToInstance.IsEmpty())
 			{
-				Mesh = MakeShared<FDatasmithInstantiatedMesh>();
-				EntitiesGeometry->Meshes.Add(Mesh);
+				FString MeshElementName = FString::Printf(TEXT("M%ls_%d"), *Definition.GetSketchupSourceId(), MeshCount + 1); // Count meshes from 1
+				SUTransformation Transform;
+				SUTransformationScale(&Transform, 1.0);
+				FEntitiesGeometry::FExportedGeometry& ExportedGeometry = EntitiesGeometry->ExportedGeometryForInstances;
+
+				EntitiesGeometry->ExportOneMesh(Context, ExtractedMeshPtr, ExportedGeometry, MeshCount, MeshElementName, MeshLabel, Transform);
 			}
-			// todo: reuse DatasmithMesh when it allows to reset material slots
-			Mesh->DatasmithMesh = FDatasmithSceneFactory::CreateMesh(TEXT(""));
 
-			MeshCount++;
+			int32 NodeBaked = 0;
+			// Export mesh for nodes with skewed transform, applying this transform to exported mesh vertices
+			for (FNodeOccurence* Node : NodesToBake)
+			{
+				// Make name for baked meshes different from the 'regular' mesh by adding an extra index
+				NodeBaked += 1;
+				FString MeshElementName = FString::Printf(TEXT("M%ls_%d_%d"), *Definition.GetSketchupSourceId(), MeshCount + 1, NodeBaked); // Count meshes from 1
+				SUTransformation Transform = Node->BakeTransform;
 
-			Mesh->DatasmithMesh->SetName(*MeshElementName);
-			Mesh->DatasmithMesh->SetLabel(*MeshLabel);
-			Mesh->bIsUsingInheritedMaterial = ExtractedMeshPtr->bHasFacesWithDefaultMaterial;
 
-			// Add the non-inherited materials used by the combined mesh triangles.
-			if (Context.bColorByLayer)
-			{
-				for (int32 SlotId = 0;SlotId < ExtractedMeshPtr->LayerIDForSlotId.Num(); ++SlotId)
+				// Find exported geometry with the same Bake transform and reuse it
+				bool bFoundGeometry = false;
+				for (const TPair<SUTransformation, FEntitiesGeometry::FExportedGeometry>& OtherTransformAndGeometry: EntitiesGeometry->ExportedGeometryForTransform)
 				{
-					FLayerIDType LayerID = ExtractedMeshPtr->LayerIDForSlotId[SlotId];
-					Mesh->SlotIdForLayerId.FindOrAdd(LayerID, SlotId);
-
-					if (FMaterialOccurrence* Material = Context.Materials.LayerMaterials.RegisterGeometryForLayer(LayerID, EntitiesGeometry.Get()))
+					if (DatasmithSketchUpUtils::CompareSUTransformations(OtherTransformAndGeometry.Key, Transform))
 					{
-						Mesh->DatasmithMesh->SetMaterial(Material->GetName(), SlotId);
+						EntitiesGeometry->ExportedGeometryForNode.Add(Node, OtherTransformAndGeometry.Value);
+						bFoundGeometry = true;
 					}
 				}
-			}
-			else
-			{
-				for (int32 SlotId = 0;SlotId < ExtractedMeshPtr->MaterialIDForSlotId.Num(); ++SlotId)
-				{
-					if (SlotId == 0 && !ExtractedMeshPtr->bHasFacesWithDefaultMaterial)
-					{
-						continue; // Skip adding Default material slot if it's not used
-					}
 
-					FMaterialIDType MeshMaterialID = ExtractedMeshPtr->MaterialIDForSlotId[SlotId];
-					Mesh->SlotIdForMaterialId.FindOrAdd(MeshMaterialID, SlotId);
-					// Default or (somehow)missing materials are also assigned to mesh(as a default material)
-					if (FMaterialOccurrence* Material = Context.Materials.RegularMaterials.RegisterGeometry(MeshMaterialID, EntitiesGeometry.Get()))
-					{
-						Mesh->DatasmithMesh->SetMaterial(Material->GetName(), SlotId);
-					}
+				if (!bFoundGeometry)
+				{
+					FEntitiesGeometry::FExportedGeometry& ExportedGeometry = EntitiesGeometry->ExportedGeometryForNode.Add(Node);
+					EntitiesGeometry->ExportOneMesh(Context, ExtractedMeshPtr, ExportedGeometry, MeshCount, MeshElementName, MeshLabel, Transform);
+					EntitiesGeometry->ExportedGeometryForTransform.Add({Transform, ExportedGeometry});
 				}
 			}
-			
-			SUTransformation Transform = Definition.GetMeshBakedTransform();
 
-			Context.MeshExportTasks.Emplace(Async(
-				EAsyncExecution::ThreadPool,
-				[&Context, Mesh, ExtractedMeshPtr, Transform]()
-				{
-					FDatasmithMeshExporter DatasmithMeshExporter;
-					FDatasmithMesh DatasmithMesh;
-					ExtractedMeshPtr->ConvertMeshToDatasmith(Context, Transform, DatasmithMesh);
-
-					FGCScopeGuard GCGuard; // Prevent GC from running while UDatasmithMesh is created in ExportToUObject.
-
-					return DatasmithMeshExporter.ExportToUObject(Mesh->DatasmithMesh, Context.GetAssetsOutputPath(), DatasmithMesh, nullptr, FDatasmithExportOptions::LightmapUV);
-				}
-			));
+			MeshCount++;
 		}
 	};
 
@@ -644,25 +699,25 @@ void FEntities::UpdateGeometry(FExportContext& Context)
 		CombineSketchUpEntitiesFaces(Context, EntitiesRef, *EntitiesGeometry, ProcessExtractedMesh);
 	}
 
-	EntitiesGeometry->Meshes.SetNum(MeshCount);
+	EntitiesGeometry->ExportedGeometryForInstances.Meshes.SetNum(NodesToInstance.IsEmpty() ? 0 : MeshCount);
 
 	Context.EntitiesObjects.RegisterEntities(*this);
 }
 
 void FEntities::AddMeshesToDatasmithScene(FExportContext& Context)
 {
-	for (TSharedPtr<FDatasmithInstantiatedMesh> Mesh : EntitiesGeometry->Meshes)
+	EntitiesGeometry->ForEachExportedMesh([&Context](const FDatasmithInstantiatedMesh& Mesh)
 	{
-		Context.DatasmithScene->AddMesh(Mesh->DatasmithMesh);
-	}
+		Context.DatasmithScene->AddMesh(Mesh.DatasmithMesh);
+	});
 }
 
 void FEntities::RemoveMeshesFromDatasmithScene(FExportContext& Context)
 {
-	for (TSharedPtr<FDatasmithInstantiatedMesh> Mesh : EntitiesGeometry->Meshes)
+	EntitiesGeometry->ForEachExportedMesh([&Context](const FDatasmithInstantiatedMesh& Mesh)
 	{
-		Context.DatasmithScene->RemoveMesh(Mesh->DatasmithMesh);
-	}
+		Context.DatasmithScene->RemoveMesh(Mesh.DatasmithMesh);
+	});
 }
 
 TArray<SUGroupRef> FEntities::GetGroups()
@@ -896,7 +951,7 @@ void FImage::UpdateGeometry(FExportContext& Context)
 	DatasmithMeshElement = FDatasmithSceneFactory::CreateMesh(TEXT(""));
 
 	DatasmithMeshElement->SetName(GetMeshElementName());
-	DatasmithMeshElement->SetLabel(*GetName());
+	DatasmithMeshElement->SetLabel(*GetEntityName());
 
 	FGCScopeGuard GCGuard; // Prevent GC from running while UDatasmithMesh is created in ExportToUObject. 
 	bool bResult = DatasmithMeshExporter.ExportToUObject(DatasmithMeshElement, Context.GetAssetsOutputPath(), DatasmithMesh, nullptr, FDatasmithExportOptions::LightmapUV);
diff --git a/Engine/Source/Programs/Enterprise/Datasmith/DatasmithSketchUpRubyExporter/Private/DatasmithSketchUpMesh.h b/Engine/Source/Programs/Enterprise/Datasmith/DatasmithSketchUpRubyExporter/Private/DatasmithSketchUpMesh.h
--- a/Engine/Source/Programs/Enterprise/Datasmith/DatasmithSketchUpRubyExporter/Private/DatasmithSketchUpMesh.h
+++ b/Engine/Source/Programs/Enterprise/Datasmith/DatasmithSketchUpRubyExporter/Private/DatasmithSketchUpMesh.h
@@ -15,6 +15,7 @@
 #include "Containers/Set.h"
 #include "Containers/UnrealString.h"
 #include "Templates/SharedPointer.h"
+#include "Async/Future.h"
 
 
 namespace DatasmithSketchUp
@@ -26,6 +27,7 @@ namespace DatasmithSketchUp
 		TMap<FMaterialIDType, int32> SlotIdForMaterialId;
 		TMap<FLayerIDType, int32> SlotIdForLayerId;
 		bool bIsUsingInheritedMaterial; // Whether mesh has faces without material assigned(so override/inherited material will apply)
+		TSharedFuture<bool> ExportFuture;
 	};
 }
 
diff --git a/Engine/Source/Programs/Enterprise/Datasmith/DatasmithSketchUpRubyExporter/Private/DatasmithSketchUpUtils.cpp b/Engine/Source/Programs/Enterprise/Datasmith/DatasmithSketchUpRubyExporter/Private/DatasmithSketchUpUtils.cpp
--- a/Engine/Source/Programs/Enterprise/Datasmith/DatasmithSketchUpRubyExporter/Private/DatasmithSketchUpUtils.cpp
+++ b/Engine/Source/Programs/Enterprise/Datasmith/DatasmithSketchUpRubyExporter/Private/DatasmithSketchUpUtils.cpp
@@ -34,6 +34,8 @@
 #include "DatasmithSketchUpSDKCeases.h"
 
 // Imath third party library.
+#include "Algo/Compare.h"
+#include "Algo/Copy.h"
 #include "Imath/ImathMatrixAlgo.h"
 
 
@@ -196,9 +198,45 @@ SUMaterialRef DatasmithSketchUpUtils::GetMaterial(
 }
 
 void DatasmithSketchUpUtils::SetActorTransform(
-	TSharedPtr<IDatasmithActorElement> InActorElement,
+	const TSharedPtr<IDatasmithActorElement>& InActorElement,
 	SUTransformation const& InWorldTransform
 )
+{
+	FVector Translation{};
+	FQuat Rotation{};
+	FVector Scale{};
+	FVector Shear{};
+
+	if (DecomposeTransform(InWorldTransform, Translation, Rotation, Scale, Shear))
+	{
+		if (Rotation.IsNormalized())
+		{
+			InActorElement->SetTranslation(Translation.X, Translation.Y, Translation.Z, false);
+			InActorElement->SetRotation(Rotation.X, Rotation.Y, Rotation.Z, Rotation.W, false);
+			InActorElement->SetScale(Scale.X, Scale.Y, Scale.Z, false);
+		}
+		else
+		{
+			ToRuby::LogWarn(
+				FString::Printf(TEXT("WARNING: Actor %ls (%ls) has unnormalized rotation"), 
+					InActorElement->GetName(), InActorElement->GetLabel()));
+		}
+
+	}
+	else
+	{
+		ToRuby::LogWarn(
+			FString::Printf(TEXT("WARNING: Actor %ls (%ls) has some zero scaling or rotation can't be converted"), 
+				InActorElement->GetName(), InActorElement->GetLabel()));
+	}
+}
+
+bool DatasmithSketchUpUtils::DecomposeTransform(
+	SUTransformation const& InWorldTransform, // SketchUp world transform to convert
+	FVector& OutTranslation,
+	FQuat& OutRotation,
+	FVector& OutScale,
+	FVector& OutShear)
 {
 	// We use Imath::extractAndRemoveScalingAndShear() because FMatrix::ExtractScaling() is deemed unreliable.
 
@@ -212,20 +250,21 @@ void DatasmithSketchUpUtils::SetActorTransform(
 	// Remove any scaling from the matrix and get the scale vector that was initially present.
 	Imath::Vec3<float> Scale;
 	Imath::Vec3<float> Shear;
-	bool bExtracted = Imath::extractAndRemoveScalingAndShear<float>(Matrix, Scale, Shear, false);
-
-	if (!bExtracted)
+	if (!Imath::extractAndRemoveScalingAndShear<float>(Matrix, Scale, Shear, false))
 	{
-		ADD_SUMMARY_LINE(TEXT("WARNING: Actor %ls (%ls) has some zero scaling"), InActorElement->GetName(), InActorElement->GetLabel());
-		return;
+		// Scaling is zero
+		return false;
 	}
 
 	if (SMatrix[15] != 1.0)
 	{
 		// Apply the extra SketchUp uniform scaling factor.
-		Scale *= float(SMatrix[15]);
+		Scale *= SMatrix[15];
 	}
 
+	OutScale = FVector(Scale.x, Scale.y, Scale.z);
+	OutShear = FVector(Shear.x, Shear.y, Shear.z);
+
 	// Initialize a rotation quaternion with the rotation matrix.
 	Imath::Quat<float> Quaternion = Imath::extractQuat<float>(Matrix);
 
@@ -238,22 +277,76 @@ void DatasmithSketchUpUtils::SetActorTransform(
 	// Make sure Unreal will be able to handle the rotation quaternion.
 	float              Angle = Quaternion.angle();
 	Imath::Vec3<float> Axis = Quaternion.axis();
-	FQuat Rotation(FVector(Axis.x, Axis.y, Axis.z), Angle);
-
-	ensure(Rotation.IsNormalized());
+	OutRotation = FQuat(FVector(Axis.x, Axis.y, Axis.z), Angle);
 
 	// Convert the SketchUp right-handed Z-up coordinate translation into an Unreal left-handed Z-up coordinate translation.
 	// To avoid perturbating X, which is forward in Unreal, the handedness conversion is done by flipping the side vector Y.
 	// SketchUp uses inches as internal system unit for all 3D coordinates in the model while Unreal uses centimeters.
+	OutTranslation = FromSketchUp::ConvertPosition(SMatrix[12], SMatrix[13], SMatrix[14]);
+
+	return true;
+}
+
+bool DatasmithSketchUpUtils::SplitTransform(
+	SUTransformation const& InWorldTransform, // SketchUp world transform to convert
+	SUTransformation& OutWorldTransform,
+	SUTransformation& OutMeshActorWorldTransform,
+	SUTransformation& OutBakeTransform
+)
+{
+	// Set up a scaling and rotation matrix.
+	auto& SMatrix = InWorldTransform.values;
+	Imath::Matrix44<double> Matrix(
+		SMatrix[0], SMatrix[1],  SMatrix[2], 0.0,
+		SMatrix[4], SMatrix[5],  SMatrix[6], 0.0,
+		SMatrix[8], SMatrix[9], SMatrix[10], 0.0,
+		       0.0,        0.0,         0.0, 1.0);
+
+	// Remove any scaling from the matrix and get the scale vector that was initially present.
+	Imath::Vec3<double> Scale;
+	Imath::Vec3<double> Shear;
+	if (!Imath::extractAndRemoveScalingAndShear<double>(Matrix, Scale, Shear, false))
+	{
+		// Scaling is zero
+		return false;
+	}
+	// Scale/Shear removed, now it's just Rotation
+	Imath::Matrix44<double> RotationMatrix(Matrix);
+
+	Imath::Matrix44<double> TranslationMatrix;
+	TranslationMatrix.translate(Imath::Vec3<double>(SMatrix[12], SMatrix[13], SMatrix[14]));
+
+	if (SMatrix[15] > UE_SMALL_NUMBER)
+	{
+		// Apply the extra SketchUp uniform (inverse)scaling factor.
+		Scale /= SMatrix[15];
+	}
 
-	FVector3f Translation = DatasmithSketchUpUtils::FromSketchUp::ConvertPosition(SMatrix[12], SMatrix[13], SMatrix[14]);
+	Imath::Matrix44<double> ScaleMatrix;
+	ScaleMatrix.scale(Scale);
 
-	// Set the world transform of the Datasmith actor.
-	InActorElement->SetTranslation(Translation.X, Translation.Y, Translation.Z, false);
-	InActorElement->SetRotation((float)Rotation.X, (float)Rotation.Y, (float)Rotation.Z, (float)Rotation.W, false);
-	InActorElement->SetScale(Scale.x, Scale.y, Scale.z, false);
+	Imath::Matrix44<double> ShearMatrix;
+	ShearMatrix.shear(Shear);
+
+	Imath::Matrix44<double> ActorMatrix = RotationMatrix*TranslationMatrix;
+	Imath::Matrix44<double> MeshActorMatrix = RotationMatrix*TranslationMatrix;
+
+	auto CopyImathMatrixToSU = [](const Imath::Matrix44<double>& Source, SUTransformation& Dest)
+	{
+		for (int32 Index=0; Index<16; ++Index)
+		{
+			Dest.values[Index] = Source.getValue()[Index];
+		}
+	};
+
+	CopyImathMatrixToSU(ActorMatrix, OutWorldTransform);
+	CopyImathMatrixToSU(MeshActorMatrix, OutMeshActorWorldTransform);
+	CopyImathMatrixToSU(ScaleMatrix*ShearMatrix, OutBakeTransform);
+
+	return true;
 }
 
+
 SUTransformation DatasmithSketchUpUtils::GetComponentInstanceTransform(SUComponentInstanceRef InComponentInstanceRef, SUTransformation const& InWorldTransform)
 {
 	// Get the SketchUp component instance transform.
@@ -266,3 +359,11 @@ SUTransformation DatasmithSketchUpUtils::GetComponentInstanceTransform(SUCompone
 	return SComponentInstanceWorldTransform;
 }
 
+bool DatasmithSketchUpUtils::CompareSUTransformations(const SUTransformation& A, const SUTransformation& B)
+{
+	return Algo::Compare(
+		MakeArrayView(A.values, std::size(A.values)), 
+		MakeArrayView(B.values, std::size(B.values)),
+		[](const double& A, const double& B) {  return FMath::Abs(A-B) < UE_KINDA_SMALL_NUMBER; });
+}
+
diff --git a/Engine/Source/Programs/Enterprise/Datasmith/DatasmithSketchUpRubyExporter/Private/DatasmithSketchUpUtils.h b/Engine/Source/Programs/Enterprise/Datasmith/DatasmithSketchUpRubyExporter/Private/DatasmithSketchUpUtils.h
--- a/Engine/Source/Programs/Enterprise/Datasmith/DatasmithSketchUpRubyExporter/Private/DatasmithSketchUpUtils.h
+++ b/Engine/Source/Programs/Enterprise/Datasmith/DatasmithSketchUpRubyExporter/Private/DatasmithSketchUpUtils.h
@@ -14,22 +14,16 @@ namespace DatasmithSketchUpUtils
 			return FVector(float(V.x), float(-V.y), float(V.z));
 		}
 
-		static FORCEINLINE FVector3f ConvertPosition(double X, double Y, double Z)
+		static FORCEINLINE FVector ConvertPosition(double X, double Y, double Z)
 		{
 			const float UnitScaleSketchupToUnreal = 2.54; // centimeters per inch
-			return FVector3f(float(X * UnitScaleSketchupToUnreal), float(-Y * UnitScaleSketchupToUnreal), float(Z * UnitScaleSketchupToUnreal));
+			return FVector(float(X * UnitScaleSketchupToUnreal), float(-Y * UnitScaleSketchupToUnreal), float(Z * UnitScaleSketchupToUnreal));
 		}
 
-		static FORCEINLINE FVector3f ConvertPosition(const SUPoint3D& V)
+		static FORCEINLINE FVector ConvertPosition(const SUPoint3D& V)
 		{
 			return ConvertPosition(V.x, V.y, V.z);
 		}
-
-		static FORCEINLINE float ConvertDistance(const float& D)
-		{
-			const float UnitScaleSketchupToUnreal = 2.54; // centimeters per inch
-			return D * UnitScaleSketchupToUnreal;
-		}
 	}
 
 	DatasmithSketchUp::FEntityIDType GetEntityID(
@@ -96,11 +90,45 @@ namespace DatasmithSketchUpUtils
 		SUComponentInstanceRef InSComponentInstanceRef // valid SketckUp component instance
 	);
 
+	bool CompareSUTransformations(const SUTransformation& A, const SUTransformation& B);
+
 	// Set the world transform of a Datasmith actor.
 	void SetActorTransform(
-		TSharedPtr<IDatasmithActorElement> IODActorPtr,      // Datasmith actor to transform
-		SUTransformation const& InSWorldTransform // SketchUp world transform to apply
+		const TSharedPtr<IDatasmithActorElement>& IODActorPtr,      // Datasmith actor to transform
+		SUTransformation const& InWorldTransform // SketchUp world transform to apply
+	);
+
+	bool DecomposeTransform(
+		SUTransformation const& InWorldTransform, // SketchUp world transform to convert
+		FVector& OutTranslation,
+		FQuat& OutRotation,
+		FVector& OutScale,
+		FVector& OutShear
+	);
+
+	// Split a source SketchUp transformation to a set of transformations supported by Unreal
+	// Transform which comes from SketchUp can be any affine transform, represented like: T*R*H*S
+	// TRS - are supported by Unreal TranslationRotationScaling
+	// and H is the 'Shear'/'Skew' unsupported by unreal
+	// In order to correctly display geometry with Shear the H*S part of the transform needs to be 'baked' into the
+	// exported geometry, meaning that the vertices need to be pre-transformed by the S*H matrix
+	//
+	// @param OutWorldTransform Is a transform without SCALE and SHEAR to set on the Actor
+	// @param OutMeshActorWorldTransform Is a transform without SHEAR(but scaling is kept) to set on the MeshActor
+	// @param OutBakeTransform Is just SHEAR to apply to mesh vertices before export, to 'bake' it into the exported mesh
+	bool SplitTransform(
+		SUTransformation const& InWorldTransform, // SketchUp world transform to convert
+		SUTransformation& OutWorldTransform,
+		SUTransformation& OutMeshActorWorldTransform,
+		SUTransformation& OutBakeTransform
 	);
 
 	SUTransformation GetComponentInstanceTransform(SUComponentInstanceRef InSComponentInstanceRef, SUTransformation const& InSWorldTransform);
+
+	// Call into Ruby code
+	namespace ToRuby
+	{
+		// Add Warning that will be show in SketchUp UI(plugins Messages dialog)
+		void LogWarn(const FString& Message);
+	};
 }
