[Solidworks] Parts updated when changing configs with DirectLink Sync

- fixes MicAssembly Asset does not update when adjusting configuration and synchronizing

diff --git a/Engine/Source/Programs/Enterprise/Datasmith/DatasmithSolidworksExporter/Private/AssemblyDocument.cs b/Engine/Source/Programs/Enterprise/Datasmith/DatasmithSolidworksExporter/Private/AssemblyDocument.cs
--- a/Engine/Source/Programs/Enterprise/Datasmith/DatasmithSolidworksExporter/Private/AssemblyDocument.cs
+++ b/Engine/Source/Programs/Enterprise/Datasmith/DatasmithSolidworksExporter/Private/AssemblyDocument.cs
@@ -169,7 +169,13 @@ namespace DatasmithSolidworks
 
 		public override FMeshes GetMeshes(string ActiveConfigName)
 		{
-			return SyncState.Meshes ?? (SyncState.Meshes = new FMeshes(ActiveConfigName));
+			FMeshes MeshesCurrent = SyncState.Meshes;
+			if (MeshesCurrent != null && (MeshesCurrent.MainConfigurationName == ActiveConfigName))
+			{
+				return MeshesCurrent;
+			}
+
+			return SyncState.Meshes = new FMeshes(ActiveConfigName);
 		}
 
 		// Export to datasmith scene extracted configurations data
diff --git a/Engine/Source/Programs/Enterprise/Datasmith/DatasmithSolidworksExporter/Private/ConfigurationExporter.cs b/Engine/Source/Programs/Enterprise/Datasmith/DatasmithSolidworksExporter/Private/ConfigurationExporter.cs
--- a/Engine/Source/Programs/Enterprise/Datasmith/DatasmithSolidworksExporter/Private/ConfigurationExporter.cs
+++ b/Engine/Source/Programs/Enterprise/Datasmith/DatasmithSolidworksExporter/Private/ConfigurationExporter.cs
@@ -41,7 +41,7 @@ namespace DatasmithSolidworks
 
 	public class FMeshes
 	{
-		private readonly string MainConfigurationName;
+		public readonly string MainConfigurationName;
 		private Dictionary<string, FConfiguration> Configurations = new Dictionary<string, FConfiguration>();
 
 		// Store unique meshes(identical mesh from different component will use the same FMesh object)
@@ -131,7 +131,7 @@ namespace DatasmithSolidworks
 		{
 			MainConfigurationName = InMainConfigurationName;
 		}
-
+		
 		public FConfiguration GetMeshesConfiguration(string ConfigurationName)
 		{
 			if (Configurations.TryGetValue(ConfigurationName, out var Configuration))
