[Solidworks] Fixed mesh actor visibility when exploded view has parts different in different configs

- was reproduced by MicAssemblyShort.SLDASM from Sandbox\QASourceAssets\Studio\Datasmith\SolidWorks\Mic_Assembly - exporting it made mic cone parts(from every variant) invisible by default(until variant is switched on). Expected - active configuration should be visible by default

diff --git a/Engine/Source/Programs/Enterprise/Datasmith/DatasmithSolidworksExporter/Private/AssemblyDocument.cs b/Engine/Source/Programs/Enterprise/Datasmith/DatasmithSolidworksExporter/Private/AssemblyDocument.cs
--- a/Engine/Source/Programs/Enterprise/Datasmith/DatasmithSolidworksExporter/Private/AssemblyDocument.cs
+++ b/Engine/Source/Programs/Enterprise/Datasmith/DatasmithSolidworksExporter/Private/AssemblyDocument.cs
@@ -409,6 +409,8 @@ namespace DatasmithSolidworks
 
 					Exporter.ExportOrUpdateActor(ActorExportInfo);
 				}
+				
+				Dictionary<FActorName, FDatasmithActorExportInfo> ConfigurationsActorExportInfo = new Dictionary<FActorName, FDatasmithActorExportInfo>();
 
 				foreach (FConfigurationTree.FComponentConfig ComponentConfig in InNode.Configurations)
 				{
@@ -422,22 +424,32 @@ namespace DatasmithSolidworks
 
 					ActorExportInfo.ParentName = ParentName;
 
-					ActorExportInfo.bVisible = true;
-					ActorExportInfo.Type = Exporter.GetExportedActorType(ActorName) ?? EActorType.SimpleActor;
 					ActorExportInfo.Transform = Transform;
 
 					SyncState.ComponentsTransformsMap[InNode.ComponentName] = ActorExportInfo.Transform;
-
+					
 					ActorExportInfo.bVisible = ComponentConfig.bVisible && (ComponentConfig.ConfigName == ActiveConfigName);
-
 					ActorExportInfo.Type = EActorType.MeshActor;
 
 					ConfigurationExporter.AddActorForMesh(ActorExportInfo.Name, ComponentConfig.ConfigName, InNode.ComponentName);
-
-					Exporter.ExportOrUpdateActor(ActorExportInfo);
+					
+					if (ConfigurationsActorExportInfo.TryGetValue(ActorName, out FDatasmithActorExportInfo ExistingInfo))
+					{
+						Debug.Assert(ActorExportInfo.ParentName == ExistingInfo.ParentName);
+						Debug.Assert(ActorExportInfo.Type == ExistingInfo.Type);
+						Debug.Assert(ActorExportInfo.Type == ExistingInfo.Type);
+					}
+					else
+					{
+						// Only update actor once - for the first configuration in the list
+						// To keep actor's exported state in the default active configuration
+						Exporter.ExportOrUpdateActor(ActorExportInfo);
+						ConfigurationsActorExportInfo.Add(ActorName, ActorExportInfo);
+					}
 				}
 			}
 
+
 			SyncState.CleanComponents.Add(InNode.ComponentName);
 
 			// Export component children
diff --git a/Engine/Source/Programs/Enterprise/Datasmith/DatasmithSolidworksExporter/Private/DatasmithExporter.cs b/Engine/Source/Programs/Enterprise/Datasmith/DatasmithSolidworksExporter/Private/DatasmithExporter.cs
--- a/Engine/Source/Programs/Enterprise/Datasmith/DatasmithSolidworksExporter/Private/DatasmithExporter.cs
+++ b/Engine/Source/Programs/Enterprise/Datasmith/DatasmithSolidworksExporter/Private/DatasmithExporter.cs
@@ -227,7 +227,7 @@ namespace DatasmithSolidworks
 
 		public override string ToString()
 		{
-			return $"FDatasmithActorExportInfo(Name={Name}, Type={Type}, MeshName={MeshName}, ParentName={ParentName})";
+			return $"FDatasmithActorExportInfo(Name={Name}, Label={Label}, Type={Type}, MeshName={MeshName}, ParentName={ParentName}, bVisible={bVisible})";
 		}
 	};
 
@@ -350,7 +350,6 @@ namespace DatasmithSolidworks
 			Debug.Assert(Info.MeshName.IsValid());
 			Debug.Assert(Info.MeshElement != null);
 
-
 			LogIndent();
 			RemoveMesh(Info.MeshName);
 			LogDedent();
