[Datasmith] FacadeCSharp regenerated,  to make SerializeLevelSequences method available in C#

- //UE5/Main - Compile Datasmith Solidworks Exporter - Failed to build plugin

diff --git a/Engine/Source/Programs/Enterprise/Datasmith/DatasmithFacadeCSharp/Private/DatasmithFacadeCSharp.cpp b/Engine/Source/Programs/Enterprise/Datasmith/DatasmithFacadeCSharp/Private/DatasmithFacadeCSharp.cpp
--- a/Engine/Source/Programs/Enterprise/Datasmith/DatasmithFacadeCSharp/Private/DatasmithFacadeCSharp.cpp
+++ b/Engine/Source/Programs/Enterprise/Datasmith/DatasmithFacadeCSharp/Private/DatasmithFacadeCSharp.cpp
@@ -6161,6 +6161,14 @@ SWIGEXPORT unsigned int SWIGSTDCALL CSharp_FDatasmithFacadeScene_ExportScene__SW
 }
 
 
+SWIGEXPORT void SWIGSTDCALL CSharp_FDatasmithFacadeScene_SerializeLevelSequences(void * jarg1) {
+  FDatasmithFacadeScene *arg1 = (FDatasmithFacadeScene *) 0 ;
+  
+  arg1 = (FDatasmithFacadeScene *)jarg1; 
+  (arg1)->SerializeLevelSequences();
+}
+
+
 SWIGEXPORT void SWIGSTDCALL CSharp_FDatasmithFacadeScene_SetLabel(void * jarg1, wchar_t * jarg2) {
   FDatasmithFacadeScene *arg1 = (FDatasmithFacadeScene *) 0 ;
   TCHAR *arg2 = (TCHAR *) 0 ;
diff --git a/Engine/Source/Programs/Enterprise/Datasmith/DatasmithFacadeCSharp/Public/DatasmithFacadeCSharpPINVOKE.cs b/Engine/Source/Programs/Enterprise/Datasmith/DatasmithFacadeCSharp/Public/DatasmithFacadeCSharpPINVOKE.cs
--- a/Engine/Source/Programs/Enterprise/Datasmith/DatasmithFacadeCSharp/Public/DatasmithFacadeCSharpPINVOKE.cs
+++ b/Engine/Source/Programs/Enterprise/Datasmith/DatasmithFacadeCSharp/Public/DatasmithFacadeCSharpPINVOKE.cs
@@ -1703,6 +1703,9 @@ class DatasmithFacadeCSharpPINVOKE {
   [global::System.Runtime.InteropServices.DllImport("DatasmithFacadeCSharp", EntryPoint="CSharp_FDatasmithFacadeScene_ExportScene__SWIG_3")]
   public static extern bool FDatasmithFacadeScene_ExportScene__SWIG_3(global::System.Runtime.InteropServices.HandleRef jarg1);
 
+  [global::System.Runtime.InteropServices.DllImport("DatasmithFacadeCSharp", EntryPoint="CSharp_FDatasmithFacadeScene_SerializeLevelSequences")]
+  public static extern void FDatasmithFacadeScene_SerializeLevelSequences(global::System.Runtime.InteropServices.HandleRef jarg1);
+
   [global::System.Runtime.InteropServices.DllImport("DatasmithFacadeCSharp", EntryPoint="CSharp_FDatasmithFacadeScene_SetLabel")]
   public static extern void FDatasmithFacadeScene_SetLabel(global::System.Runtime.InteropServices.HandleRef jarg1, [global::System.Runtime.InteropServices.MarshalAs(global::System.Runtime.InteropServices.UnmanagedType.LPWStr)]string jarg2);
 
diff --git a/Engine/Source/Programs/Enterprise/Datasmith/DatasmithFacadeCSharp/Public/FDatasmithFacadeScene.cs b/Engine/Source/Programs/Enterprise/Datasmith/DatasmithFacadeCSharp/Public/FDatasmithFacadeScene.cs
--- a/Engine/Source/Programs/Enterprise/Datasmith/DatasmithFacadeCSharp/Public/FDatasmithFacadeScene.cs
+++ b/Engine/Source/Programs/Enterprise/Datasmith/DatasmithFacadeCSharp/Public/FDatasmithFacadeScene.cs
@@ -396,6 +396,10 @@ public class FDatasmithFacadeScene : global::System.IDisposable {
     return ret;
   }
 
+  public void SerializeLevelSequences() {
+    DatasmithFacadeCSharpPINVOKE.FDatasmithFacadeScene_SerializeLevelSequences(swigCPtr);
+  }
+
   public void SetLabel(string InSceneLabel) {
     DatasmithFacadeCSharpPINVOKE.FDatasmithFacadeScene_SetLabel(swigCPtr, InSceneLabel);
   }
