[Solidworks] Metadata export and sync

- Each component's Document's metadata is exported

- Generic metadata - component doc type, filename,
- Summary info - author, creation/save date, comments
- Custom Properties
- Configuration Properties(or per-configuration custom prjperties)
- Display state transparency
- Comments
- Bill Of Materials table

diff --git a/Engine/Source/Programs/Enterprise/Datasmith/DatasmithSolidworksExporter/Private/AssemblyDocument.cs b/Engine/Source/Programs/Enterprise/Datasmith/DatasmithSolidworksExporter/Private/AssemblyDocument.cs
--- a/Engine/Source/Programs/Enterprise/Datasmith/DatasmithSolidworksExporter/Private/AssemblyDocument.cs
+++ b/Engine/Source/Programs/Enterprise/Datasmith/DatasmithSolidworksExporter/Private/AssemblyDocument.cs
@@ -143,6 +143,31 @@ namespace DatasmithSolidworks
 				new Dictionary<FMeshName, HashSet<FComponentName>>();
 
 			public FMeshes Meshes = null;
+			
+			public void ActiveConfigChanged()
+			{
+				CleanComponents?.Clear();
+				DirtyComponents?.Clear();
+				CollectedComponentsMap?.Clear();
+			}
+			
+			public void SetComponentDirty(FComponentName InComponent, EComponentDirtyState InState)
+			{
+				if (CleanComponents.Contains(InComponent))
+				{
+					CleanComponents.Remove(InComponent);
+				}
+				
+				uint DirtyFlags = 0u;
+				DirtyComponents.TryGetValue(InComponent, out DirtyFlags);
+				DirtyFlags |= (1u << (int)InState);
+				DirtyComponents[InComponent] = DirtyFlags;
+			}
+			
+			public void DirtyComponentsUpdated()
+			{
+				DirtyComponents.Clear();
+			}
 		}
 
 		public AssemblyDoc SwAsmDoc { get; private set; } = null;
@@ -191,7 +216,8 @@ namespace DatasmithSolidworks
 			Debug.Assert(ConfigurationExporter.CombinedTree.Children.Count == 1);
 			ExportComponentRecursive(ActiveVariantName, ConfigurationExporter,
 				ConfigurationExporter.CombinedTree.Children[0], null);
-			SyncState.DirtyComponents.Clear();
+			// All components are exported 
+			SyncState.DirtyComponentsUpdated();
 
 			ConfigurationExporter.FinalizeExport(this);
 
@@ -279,7 +305,7 @@ namespace DatasmithSolidworks
 			}
 			SyncState.ComponentsMaterialsMap[ComponentName] = Materials;
 		}
-
+		
 		public override void AddMeshForComponent(FComponentName ComponentName, FMeshName MeshName)
 		{
 			SyncState.ComponentNameToMeshNameMap.FindOrAdd(ComponentName).Add(MeshName);
@@ -348,26 +374,22 @@ namespace DatasmithSolidworks
 			FConfigurationTree.FComponentConfig ActiveConfig = InNode.Configurations?.Find(Config => Config.ConfigName == ActiveConfigName);
 			FConvertedTransform Transform = GetComponentDatasmithTransform(InNode, ActiveConfig);
 			SyncState.ComponentsTransformsMap[InNode.ComponentName] = Transform;
-
+			
+			FDatasmithActorExportInfo ActorExportInfo = new FDatasmithActorExportInfo();
+			ActorExportInfo.Label = InNode.ComponentName.GetLabel();
+			FActorName ActorName = Exporter.GetComponentActorName(InNode.ComponentName);
+			ActorExportInfo.Name = ActorName;
+			if (InParent != null)
+			{
+				ActorExportInfo.ParentName = Exporter.GetComponentActorName(InParent.ComponentName);
+			}
+			ActorExportInfo.Transform = Transform;
+			ActorExportInfo.bVisible = ActiveConfig?.bVisible ?? InNode.CommonConfig.bVisible;
+			ActorExportInfo.Metadata = InNode.Metadata;
+			
 			if (InNode.bGeometrySame)
 			{
-				FActorName ActorName = Exporter.GetComponentActorName(InNode.ComponentName);
-
-				FDatasmithActorExportInfo ActorExportInfo = new FDatasmithActorExportInfo();
-
-				ActorExportInfo.Label = InNode.ComponentName.GetLabel();
-				ActorExportInfo.Name = ActorName;
-
-				if (InParent != null)
-				{
-					ActorExportInfo.ParentName = Exporter.GetComponentActorName(InParent.ComponentName);
-				}
-
-				ActorExportInfo.bVisible = true;
 				ActorExportInfo.Type = Exporter.GetExportedActorType(ActorName) ?? EActorType.SimpleActor;
-				ActorExportInfo.Transform = Transform;
-
-				ActorExportInfo.bVisible = ActiveConfig?.bVisible ?? InNode.CommonConfig.bVisible;
 
 				if (InNode.IsPartComponent())
 				{
@@ -384,67 +406,47 @@ namespace DatasmithSolidworks
 			}
 			else
 			{
-				Debug.Assert(InNode.IsPartComponent());  // Expecting only Part components to have 'mesh' variants
-
-				FActorName ParentName;
-				{
-					FActorName ActorName = Exporter.GetComponentActorName(InNode.ComponentName);
-					ParentName = ActorName;
-
-					FDatasmithActorExportInfo ActorExportInfo = new FDatasmithActorExportInfo();
-
-					ActorExportInfo.Label = InNode.ComponentName.GetLabel();
-					ActorExportInfo.Name = ActorName;
-
-					if (InParent != null)
-					{
-						ActorExportInfo.ParentName = Exporter.GetComponentActorName(InParent.ComponentName);
-					}
+				Debug.Assert(InNode.IsPartComponent());  // Seriously, expecting only Part components to have 'mesh' variants
 
-					ActorExportInfo.bVisible = true;
-					ActorExportInfo.Type = EActorType.SimpleActor;  // Actor for Component with Mesh Variants is a 'simple' actor(i.e. just a node which has children)
-					ActorExportInfo.Transform = Transform;
+				ActorExportInfo.Type = EActorType.SimpleActor;  // Actor for Component with Mesh Variants is a 'simple' actor(i.e. just a node which has children)
 
-					ActorExportInfo.bVisible = ActiveConfig?.bVisible ?? InNode.CommonConfig.bVisible;
-
-					Exporter.ExportOrUpdateActor(ActorExportInfo);
-				}
+				Exporter.ExportOrUpdateActor(ActorExportInfo);
 				
 				Dictionary<FActorName, FDatasmithActorExportInfo> ConfigurationsActorExportInfo = new Dictionary<FActorName, FDatasmithActorExportInfo>();
 
 				foreach (FConfigurationTree.FComponentConfig ComponentConfig in InNode.Configurations)
 				{
-					FActorName ActorName = ConfigurationExporter.GetMeshActorName(ComponentConfig.ConfigName, InNode.ComponentName);
-					string Label = ActorName.GetString();
+					FActorName MeshActorName = ConfigurationExporter.GetMeshActorName(ComponentConfig.ConfigName, InNode.ComponentName);
+					string MeshActorLabel = MeshActorName.GetString();
 
-					FDatasmithActorExportInfo ActorExportInfo = new FDatasmithActorExportInfo();
+					FDatasmithActorExportInfo MeshActorExportInfo = new FDatasmithActorExportInfo();
 
-					ActorExportInfo.Label = Label;
-					ActorExportInfo.Name = ActorName;
-
-					ActorExportInfo.ParentName = ParentName;
+					MeshActorExportInfo.Label = MeshActorLabel;
+					MeshActorExportInfo.Name = MeshActorName;
+					
+					MeshActorExportInfo.ParentName = ActorName;
 
-					ActorExportInfo.Transform = Transform;
+					MeshActorExportInfo.Transform = Transform;
 
-					SyncState.ComponentsTransformsMap[InNode.ComponentName] = ActorExportInfo.Transform;
+					SyncState.ComponentsTransformsMap[InNode.ComponentName] = MeshActorExportInfo.Transform;
 					
-					ActorExportInfo.bVisible = ComponentConfig.bVisible && (ComponentConfig.ConfigName == ActiveConfigName);
-					ActorExportInfo.Type = EActorType.MeshActor;
+					MeshActorExportInfo.bVisible = ComponentConfig.bVisible && (ComponentConfig.ConfigName == ActiveConfigName);
+					MeshActorExportInfo.Type = EActorType.MeshActor;
 
-					ConfigurationExporter.AddActorForMesh(ActorExportInfo.Name, ComponentConfig.ConfigName, InNode.ComponentName);
+					ConfigurationExporter.AddActorForMesh(MeshActorExportInfo.Name, ComponentConfig.ConfigName, InNode.ComponentName);
 					
-					if (ConfigurationsActorExportInfo.TryGetValue(ActorName, out FDatasmithActorExportInfo ExistingInfo))
+					if (ConfigurationsActorExportInfo.TryGetValue(MeshActorName, out FDatasmithActorExportInfo ExistingInfo))
 					{
-						Debug.Assert(ActorExportInfo.ParentName == ExistingInfo.ParentName);
-						Debug.Assert(ActorExportInfo.Type == ExistingInfo.Type);
-						Debug.Assert(ActorExportInfo.Type == ExistingInfo.Type);
+						Debug.Assert(MeshActorExportInfo.ParentName == ExistingInfo.ParentName);
+						Debug.Assert(MeshActorExportInfo.Type == ExistingInfo.Type);
+						Debug.Assert(MeshActorExportInfo.Type == ExistingInfo.Type);
 					}
 					else
 					{
 						// Only update actor once - for the first configuration in the list
 						// To keep actor's exported state in the default active configuration
-						Exporter.ExportOrUpdateActor(ActorExportInfo);
-						ConfigurationsActorExportInfo.Add(ActorName, ActorExportInfo);
+						Exporter.ExportOrUpdateActor(MeshActorExportInfo);
+						ConfigurationsActorExportInfo.Add(MeshActorName, MeshActorExportInfo);
 					}
 				}
 			}
@@ -514,15 +516,7 @@ namespace DatasmithSolidworks
 
 		public void SetComponentDirty(FComponentName InComponent, EComponentDirtyState InState)
 		{
-			if (SyncState.CleanComponents.Contains(InComponent))
-			{
-				SyncState.CleanComponents.Remove(InComponent);
-			}
-
-			uint DirtyFlags = 0u;
-			SyncState.DirtyComponents.TryGetValue(InComponent, out DirtyFlags);
-			DirtyFlags |= (1u << (int)InState);
-			SyncState.DirtyComponents[InComponent] = DirtyFlags;
+			SyncState.SetComponentDirty(InComponent, InState);
 
 			SetDirty(true);
 		}
@@ -545,9 +539,7 @@ namespace DatasmithSolidworks
 		// Active configuration changed
 		public void ActiveConfigChanged()
 		{
-			SyncState.CleanComponents?.Clear();
-			SyncState.DirtyComponents?.Clear();
-			SyncState.CollectedComponentsMap?.Clear();
+			SyncState.ActiveConfigChanged();
 		}
 
 		public void Tick()
diff --git a/Engine/Source/Programs/Enterprise/Datasmith/DatasmithSolidworksExporter/Private/ConfigurationExporter.cs b/Engine/Source/Programs/Enterprise/Datasmith/DatasmithSolidworksExporter/Private/ConfigurationExporter.cs
--- a/Engine/Source/Programs/Enterprise/Datasmith/DatasmithSolidworksExporter/Private/ConfigurationExporter.cs
+++ b/Engine/Source/Programs/Enterprise/Datasmith/DatasmithSolidworksExporter/Private/ConfigurationExporter.cs
@@ -369,7 +369,8 @@ namespace DatasmithSolidworks
 						VariantName = VariantName.LinkedDisplayStateVariant(DisplayStates[0]);
 					}
 				}
-
+				
+				// Node, to collect data for the configuration
 				FConfigurationTree.FComponentTreeNode ConfigNode = new FConfigurationTree.FComponentTreeNode();
 				ConfigNode.ComponentInfo.ComponentName = VariantName.GetRootComponentName();
 
@@ -381,7 +382,7 @@ namespace DatasmithSolidworks
 				// Use GetRootComponent3() with Resolve = true to ensure suppressed components will be loaded
 				// todo: docs says that Part document SW returns null. Not the case. But probably better to add a guard
 				LogDebug($"Components:");
-				CollectComponentsRecursive(InDoc, swConfiguration.GetRootComponent3(true), ConfigNode, MeshesConfiguration);
+				CollectComponentsRecursive(InDoc, swConfiguration.GetRootComponent3(true), CfgName, ConfigNode, MeshesConfiguration);
 
 				ExportedVariantNames.Add(CfgName, new List<FVariantName>(){VariantName});
 
@@ -783,6 +784,7 @@ namespace DatasmithSolidworks
 		}
 
 		private static void CollectComponentsRecursive(FDocumentTracker InDoc, Component2 InComponent,
+			string CfgName,
 			FComponentTreeNode InParentNode, FMeshes.FConfiguration Meshes)
 		{
 			LogDebug($"'{InComponent.Name2}'");
@@ -797,7 +799,8 @@ namespace DatasmithSolidworks
 			// ComponentDoc is null if component is suppressed or lightweight
 			ModelDoc2 ModelDoc = (ModelDoc2)InComponent.GetModelDoc2();
 			NewNode.ComponentInfo.PartPath = (ModelDoc is PartDoc) ? ModelDoc.GetPathName() : null;  // Identify whether the component is a Part component
-
+			
+			NewNode.Metadata = InDoc.GetComponentMetadata(InComponent, CfgName);
 
 			NewNode.CommonConfig.bVisible = InComponent.Visible != (int)swComponentVisibilityState_e.swComponentHidden;
 			NewNode.CommonConfig.bSuppressed = InComponent.IsSuppressed();
@@ -813,7 +816,7 @@ namespace DatasmithSolidworks
 				{
 					Component2 Child = (Component2)ObjChild;
 					LogIndent();
-					CollectComponentsRecursive(InDoc, Child, NewNode, Meshes);
+					CollectComponentsRecursive(InDoc, Child, CfgName, NewNode, Meshes);
 					LogDedent();
 				}
 
@@ -833,7 +836,7 @@ namespace DatasmithSolidworks
 			}
 			InDoc.AddCollectedComponent(NewNode);
 		}
-
+		
 		private static void ComputeNodeTransform(FComponentConfig ParentConfig, FComponentTreeNode InNode, FComponentConfig ComponentConfig)
 		{
 			// Read transform
diff --git a/Engine/Source/Programs/Enterprise/Datasmith/DatasmithSolidworksExporter/Private/ConfigurationTree.cs b/Engine/Source/Programs/Enterprise/Datasmith/DatasmithSolidworksExporter/Private/ConfigurationTree.cs
--- a/Engine/Source/Programs/Enterprise/Datasmith/DatasmithSolidworksExporter/Private/ConfigurationTree.cs
+++ b/Engine/Source/Programs/Enterprise/Datasmith/DatasmithSolidworksExporter/Private/ConfigurationTree.cs
@@ -115,7 +115,8 @@ namespace DatasmithSolidworks
 			public HashSet<FActorName> Meshes = null;
 
 			public List<FComponentTreeNode> Children;
-
+			public FMetadata Metadata;
+			
 			// Traverse tree passing result of function computation for each node to its children
 			// Function - receives parent's computed value and current node and returns value computed for the node(to pass to children)
 			public void Traverse<T>(T ParentValue, Func<T, FComponentTreeNode, T> Function)
@@ -225,7 +226,7 @@ namespace DatasmithSolidworks
 			}
 		};
 
-		static public void Merge(FComponentTreeNode OutCombined, FComponentTreeNode InTree, FVariantName InConfigurationName,
+		public static void Merge(FComponentTreeNode OutCombined, FComponentTreeNode InTree, FVariantName InConfigurationName,
 			bool bIsActiveConfiguration)
 		{
 			foreach (FComponentTreeNode Child in InTree.EnumChildren())
@@ -273,13 +274,26 @@ namespace DatasmithSolidworks
 				{
 					CombinedChild.Configurations.AddRange(Child.Configurations);
 				}
+				
+				if (Child.Metadata != null)
+				{
+					if (CombinedChild.Metadata == null)
+					{
+						CombinedChild.Metadata = Child.Metadata;
+						CombinedChild.Metadata.OwnerName = Child.ActorName;
+					}
+					else
+					{
+						CombinedChild.Metadata.Pairs.AddRange(Child.Metadata.Pairs);
+					}
+				}
 
 				// Recurse to children
 				Merge(CombinedChild, Child, InConfigurationName, bIsActiveConfiguration);
 			}
 		}
 
-		static public void Compress(FComponentTreeNode InNode, FConfigurationExporter ConfigurationExporter)
+		public static void Compress(FComponentTreeNode InNode, FConfigurationExporter ConfigurationExporter)
 		{
 			bool CheckMaterialsEqual(List<FComponentConfig> ConfigList)
 			{
diff --git a/Engine/Source/Programs/Enterprise/Datasmith/DatasmithSolidworksExporter/Private/DatasmithExporter.cs b/Engine/Source/Programs/Enterprise/Datasmith/DatasmithSolidworksExporter/Private/DatasmithExporter.cs
--- a/Engine/Source/Programs/Enterprise/Datasmith/DatasmithSolidworksExporter/Private/DatasmithExporter.cs
+++ b/Engine/Source/Programs/Enterprise/Datasmith/DatasmithSolidworksExporter/Private/DatasmithExporter.cs
@@ -224,7 +224,8 @@ namespace DatasmithSolidworks
 		public FMeshName MeshName;
 		public FConvertedTransform Transform;
 		public bool bVisible;
-
+		public FMetadata Metadata;
+		
 		public override string ToString()
 		{
 			return $"FDatasmithActorExportInfo(Name={Name}, Label={Label}, Type={Type}, MeshName={MeshName}, ParentName={ParentName}, bVisible={bVisible})";
@@ -340,6 +341,11 @@ namespace DatasmithSolidworks
 					MeshActor.SetMesh(DatasmithMeshName);
 				}
 			}
+			
+			if (InExportInfo.Metadata != null)
+			{
+				ExportMetadata(InExportInfo.Metadata);
+			}
 
 			return Actor;
 		}
@@ -430,6 +436,13 @@ namespace DatasmithSolidworks
 			{
 				Tuple<EActorType, FDatasmithFacadeActor> ActorInfo = ExportedActorsMap[InActorName];
 				FDatasmithFacadeActor Actor = ActorInfo.Item2;
+				
+				FDatasmithFacadeMetaData DatasmithMetaData = DatasmithScene.GetMetaData(Actor);
+				if (DatasmithMetaData != null)
+				{
+					DatasmithScene.RemoveMetaData(DatasmithMetaData);
+				}
+
 				DatasmithScene.RemoveActor(Actor);
 				ExportedActorsMap.Remove(InActorName);
 			}
@@ -609,11 +622,12 @@ namespace DatasmithSolidworks
 
 				if (DatasmithMetadata == null)
 				{
-					DatasmithMetadata = new FDatasmithFacadeMetaData("SolidWorks Document Metadata");
+					DatasmithMetadata = new FDatasmithFacadeMetaData(InMetadata.OwnerName + "_DATA");
 					DatasmithMetadata.SetAssociatedElement(Element);
 					DatasmithScene.AddMetaData(DatasmithMetadata);
 				}
-
+				
+				DatasmithMetadata.ResetProperties();
 				foreach (IMetadataPair Pair in InMetadata.Pairs)
 				{
 					Pair.WriteToDatasmithMetaData(DatasmithMetadata);
diff --git a/Engine/Source/Programs/Enterprise/Datasmith/DatasmithSolidworksExporter/Private/Document.cs b/Engine/Source/Programs/Enterprise/Datasmith/DatasmithSolidworksExporter/Private/Document.cs
--- a/Engine/Source/Programs/Enterprise/Datasmith/DatasmithSolidworksExporter/Private/Document.cs
+++ b/Engine/Source/Programs/Enterprise/Datasmith/DatasmithSolidworksExporter/Private/Document.cs
@@ -251,7 +251,58 @@ namespace DatasmithSolidworks
 		public abstract Dictionary<FComponentName, FObjectMaterials> LoadDocumentMaterials(HashSet<FComponentName> ComponentNamesToExportSet);
 		public abstract void AddComponentMaterials(FComponentName ComponentName, FObjectMaterials Materials);
 		public abstract FObjectMaterials GetComponentMaterials(Component2 Comp);
+		
+		public FMetadata GetComponentMetadata(Component2 InComponent, string CfgName)
+		{
+			ModelDoc2 ModelDoc = (ModelDoc2)InComponent.GetModelDoc2();
+			if (ModelDoc == null)
+			{
+				return new FMetadata(FMetadata.EOwnerType.Actor);
+			}
+			
+			FMetadata Metadata = new FMetadata(FMetadata.EOwnerType.Actor);
+			
+			string Doctype = "";
+			bool bIsPart = false;
+			switch (ModelDoc)
+			{
+				case AssemblyDoc _:
+				{
+					Doctype = "Assembly";
+					break;
+				}
+				case PartDoc _:
+				{
+					Doctype = "Part";
+					bIsPart = true;
+					break;
+				}
+			}
+			Metadata.AddPair("Document_Type", Doctype);
+			Metadata.AddPair("Document_Filename", System.IO.Path.GetFileName(ModelDoc.GetPathName()));
 
+			
+			Metadata.AddPair("Document_Author", ModelDoc.SummaryInfo[(int)swSummInfoField_e.swSumInfoAuthor]);
+			Metadata.AddPair("Document_Comment", ModelDoc.SummaryInfo[(int)swSummInfoField_e.swSumInfoComment]);
+			Metadata.AddPair("Document_CreateDate", ModelDoc.SummaryInfo[(int)swSummInfoField_e.swSumInfoCreateDate]);
+			Metadata.AddPair("Document_CreateDate2", ModelDoc.SummaryInfo[(int)swSummInfoField_e.swSumInfoCreateDate2]);
+			Metadata.AddPair("Document_Keywords", ModelDoc.SummaryInfo[(int)swSummInfoField_e.swSumInfoKeywords]);
+			Metadata.AddPair("Document_SaveDate", ModelDoc.SummaryInfo[(int)swSummInfoField_e.swSumInfoSaveDate]);
+			Metadata.AddPair("Document_SaveDate2", ModelDoc.SummaryInfo[(int)swSummInfoField_e.swSumInfoSaveDate2]);
+			Metadata.AddPair("Document_SavedBy", ModelDoc.SummaryInfo[(int)swSummInfoField_e.swSumInfoSavedBy]);
+			Metadata.AddPair("Document_Subject", ModelDoc.SummaryInfo[(int)swSummInfoField_e.swSumInfoSubject]);
+			Metadata.AddPair("Document_Title", ModelDoc.SummaryInfo[(int)swSummInfoField_e.swSumInfoTitle]);
+
+			FMetadataManager.ExportCustomProperties(ModelDoc, Metadata);
+			FMetadataManager.ExportCustomProperties(ModelDoc, Metadata, CfgName);
+			if (bIsPart == false)
+			{
+				FMetadataManager.AddAssemblyDisplayStateMetadata(ModelDoc as AssemblyDoc, Metadata);
+			}
+			FMetadataManager.ExportCommentsAndBom(ModelDoc, Metadata);
+			
+			return Metadata;
+		}
 
 		// Record which meshes are used aby a component
 		public abstract void AddMeshForComponent(FComponentName ComponentName, FMeshName MeshName);
@@ -301,6 +352,7 @@ namespace DatasmithSolidworks
 		{
 			Exporter.AssignMaterialsToDatasmithMeshes(CreatedMeshes);
 		}
+		
 	};
 
 
diff --git a/Engine/Source/Programs/Enterprise/Datasmith/DatasmithSolidworksExporter/Private/MetadataManager.cs b/Engine/Source/Programs/Enterprise/Datasmith/DatasmithSolidworksExporter/Private/MetadataManager.cs
--- a/Engine/Source/Programs/Enterprise/Datasmith/DatasmithSolidworksExporter/Private/MetadataManager.cs
+++ b/Engine/Source/Programs/Enterprise/Datasmith/DatasmithSolidworksExporter/Private/MetadataManager.cs
@@ -3,8 +3,8 @@
 using SolidWorks.Interop.sldworks;
 using SolidWorks.Interop.swconst;
 using System.Collections.Generic;
-using System.Runtime.InteropServices;
 using DatasmithSolidworks.Names;
+using static DatasmithSolidworks.Addin;
 
 namespace DatasmithSolidworks
 {
@@ -78,48 +78,10 @@ namespace DatasmithSolidworks
 
 	public class FMetadataManager
 	{
-
-		public FMetadataManager()
-		{
-		}
-
-		static public void AddDocumentMetadataToCommand(ModelDoc2 InModelDoc, FMetadata InMetadata)
-		{
-			string Doctype = "";
-			bool bIsPart = false;
-			if (InModelDoc is AssemblyDoc)
-			{
-				Doctype = "Assembly";
-			}
-			else if (InModelDoc is PartDoc)
-			{
-				Doctype = "Part";
-				bIsPart = true;
-			}
-			InMetadata.AddPair("Document_Type", Doctype);
-			InMetadata.AddPair("Document_Filename", System.IO.Path.GetFileName(InModelDoc.GetPathName()));
-
-			ExportCustomProperties(InModelDoc, InMetadata);
-			ExportCurrentConfigurationCustomProperties(InModelDoc, InMetadata);
-
-			if (bIsPart == false)
-			{
-				AddAssemblyDisplayStateMetadata(InModelDoc as AssemblyDoc, InMetadata);
-			}
-
-			ExportCommentsandBOM(InModelDoc, InMetadata);
-		}        
-
-		static private bool ExportCurrentConfigurationCustomProperties(ModelDoc2 InModeldoc, FMetadata InMetadata)
+		public static bool ExportCustomProperties(ModelDoc2 InModeldoc, FMetadata InMetadata, string CfgName="")
 		{
-			Configuration ActiveConf = InModeldoc.GetActiveConfiguration();            
-			CustomPropertyManager ActiveConfCustomPropertyManager = ActiveConf.CustomPropertyManager;
+			LogDebug($"ExportCustomProperties {InModeldoc.GetPathName()}, Config:{CfgName} ");
 
-			return ExportCustomPropertyManagerToMetadata(ActiveConfCustomPropertyManager, InMetadata, "Active_Configuration_Custom_Property_");
-		}
-
-		static private bool ExportCustomProperties(ModelDoc2 InModeldoc, FMetadata InMetadata)
-		{
 			ModelDocExtension Ext = InModeldoc.Extension;
 			if (Ext == null)
 			{
@@ -127,89 +89,46 @@ namespace DatasmithSolidworks
 			}
 
 			// Add document related custom properties metadata
-			CustomPropertyManager PropertyManager = Ext.get_CustomPropertyManager("");
+			CustomPropertyManager PropertyManager = Ext.get_CustomPropertyManager(CfgName);
 
 			if (PropertyManager == null)
 			{
 				return false;
 			}
-			return ExportCustomPropertyManagerToMetadata(PropertyManager, InMetadata, "Document_Custom_Property_");
+			
+			return ExportCustomPropertyManagerToMetadata(PropertyManager, InMetadata, string.IsNullOrEmpty(CfgName) ?  $"Document_Custom_Property_" : $"Document_Configuration_Property_{CfgName}_");
 		}
-
-		static private bool ExportPartProperties(PartDoc InPartDoc, FMetadata InMetadata)
+		
+		// Recursively visit all sub-features of the feature
+		// e.g. BomFeat is a suf-feature of a TableFolder feature
+		public static bool TraverseFeature(Feature Feat, System.Func<Feature, bool> Callback)
 		{
-			EnumBodies2 Enum3 = InPartDoc.EnumBodies3((int)swBodyType_e.swAllBodies, false);
-			Body2 Body = null;
-			do
+			LogDebug($"Feature: {Feat.Name} <{Feat.GetTypeName2()}>");
+			
+			if (!Callback(Feat))
 			{
-				int Fetched = 0;
-				Enum3.Next(1, out Body, ref Fetched);
-				if (Body != null)
-				{
-					ExportBodiesProperties(Fetched, Body, InMetadata);
-				}
-			} while (Body != null);
-
-			return true;
-		}
-
-		static private bool ExportBodiesProperties(int InBodyIndex, Body2 InBody, FMetadata InMetadata)
-		{
-			//SOLIDWORKS recommends that you use the IAttribute, IAttributeDef, and IParameter interfaces
-			//instead of the IBody2::AddPropertyExtension2 method
-
-			//IAttribute
-			//body.FindAttribute(attributedef
-			//How to list the attributedef ??
-			string BodyIndexStr = "Body_" + InBodyIndex;
-			InMetadata.AddPair(BodyIndexStr, InBody.Name);
-
-			//parse features
-			object[] Features = (object[]) InBody.GetFeatures();
-			Feature Feat = null;
-			for (int Idx = 0; Idx < InBody.GetFeatureCount(); Idx++)
+				return false;
+			}
+			
+			Feature SubFeat = (Feature)Feat.GetFirstSubFeature();
+			
+			while (SubFeat != null)
 			{
-				Feat = (Feature) Features[Idx];
-				InMetadata.AddPair(BodyIndexStr + "_Feature_" + Feat.GetTypeName2() , Feat.Name);
-
-				//failed attempt
-				//string featuretype = feature.GetTypeName2();
-				//if (featuretype == "Attribute")
-				//{
-				//    Attribute attr = (Attribute)feature.GetSpecificFeature2();
-				//    System.Diagnostics.Debug.WriteLine(attr.GetName());
-				//    AttributeDef def = attr.GetDefinition();
-				//    if (false == attr.GetEntityState((int)swAssociatedEntityStates_e.swIsEntityInvalid) & false == attr.GetEntityState((int)swAssociatedEntityStates_e.swIsEntitySuppressed) & false == attr.GetEntityState((int)swAssociatedEntityStates_e.swIsEntityAmbiguous) & false == attr.GetEntityState((int)swAssociatedEntityStates_e.swIsEntityDeleted))
-				//    {
-				//        bool valid = true;
-				//        //Parameter paramname = (Parameter)attr.GetParameter(????? );//no way to get attribute name from AttributeDef ????
-				//        //Parameter paramvalue = (Parameter)attr.GetParameter(???);//same for attribute value name ?
-				//        //Debug.Print("Attribute " + swParamName.GetStringValue() + " found.");
-				//        //Debug.Print("  Value = " + swParamValue.GetDoubleValue());
-
-				//    }
-				//}
-
-				dynamic Faces = Feat.GetFaces();
-				if (Faces != null)
+				using (LogScopedIndent())
 				{
-					int J = 0;
-					foreach (var Obj in Faces)
+					if (!TraverseFeature(SubFeat, Callback))
 					{
-						ExportFaceProperties(J, (Face2)Obj, InMetadata);
-						J++;
+						return false;
 					}
 				}
+				
+				SubFeat = (Feature)SubFeat.GetNextSubFeature();
 			}
+			
 			return true;
 		}
-
-		static private bool ExportFaceProperties(int InFaceindex, Face2 InFace, FMetadata InMetadata)
-		{
-			return true;
-		}
-
-		static private bool ExportCommentsandBOM(ModelDoc2 InModelDoc, FMetadata InMetadata)
+		
+		public static bool ExportCommentsAndBom(ModelDoc2 InModelDoc, FMetadata InMetadata)
 		{
 			FeatureManager FeatureMan = InModelDoc.FeatureManager;
 			if (FeatureMan == null)
@@ -218,60 +137,59 @@ namespace DatasmithSolidworks
 			}
 
 			Feature Feat = InModelDoc.FirstFeature();
-			string Featuretype = null;
-			CommentFolder CommentFolder = null;
-			object[] Comments = null;
-			Comment Comment = null;
-
-			//BomFeature bomFeature = null;
-
-			int CommentCount = 0;
+			
 			while (Feat !=  null)
 			{
-				Featuretype = Feat.GetTypeName2();
-				if (Featuretype == null)
+				TraverseFeature(Feat, F =>
 				{
-					Feat = Feat.GetNextFeature();
-					continue;
-				}                
+					ExtractFeatureMetadata(F, F.GetTypeName2(), InMetadata);
+					return true;
+				});
 
-				if (Featuretype == "CommentsFolder")
+				Feat = Feat.GetNextFeature();
+			}
+			return true;
+		}
+		
+		private static void ExtractFeatureMetadata(Feature Feat, string FeatureType, FMetadata InMetadata)
+		{
+			switch (FeatureType)
+			{
+				case "CommentsFolder":
 				{
-					CommentFolder = (CommentFolder)Feat.GetSpecificFeature2();
-					CommentCount = CommentFolder.GetCommentCount();
-					if (CommentFolder != null && CommentCount != 0)
+					CommentFolder CommentFolder = (CommentFolder)Feat.GetSpecificFeature2();
+					if (CommentFolder != null)
 					{
-						Comments = (object[])CommentFolder.GetComments();
-						if (Comments != null)
+						int CommentCount = CommentFolder.GetCommentCount();
+						if (CommentCount > 0)
 						{
-							for (int i = 0; i < CommentCount; i++)
+							object[] Comments = (object[])CommentFolder.GetComments();
+							if (Comments != null)
 							{
-								Comment = (Comment)Comments[i];
-								InMetadata.AddPair("Comment_" + Comment.Name, Comment.Text);
+								for (int I = 0; I < CommentCount; I++)
+								{
+									Comment Comment = (Comment)Comments[I];
+									InMetadata.AddPair("Comment_" + Comment.Name, Comment.Text);
+								}
 							}
 						}
 					}
+					break;
 				}
-				//BOM feature reading disabled until we
-				//handle BOM in configuration handling context
-				/*
-				//"BomFeat" is not found if BOM is added to Tables list. Is there another way to add a BOM that would create "BomFeat" feature ? 
-				//used insert > table > BOM from main menu to add a BOM
-				else if (featuretype == "BomFeat") 
+
+				case "BomFeat":
 				{
-					bomFeature = (BomFeature)feature.GetSpecificFeature2();
-					if (bomFeature != null)
+					BomFeature BomFeat = (BomFeature)Feat.GetSpecificFeature2();
+					if (BomFeat != null)
 					{
-						ExportBOMFeature(bomFeature, metadatacommand);
-					}                    
+						ExportBomFeature(BomFeat, InMetadata);
+					}
+					break;
 				}
-				*/
-				Feat = Feat.GetNextFeature();
 			}
-			return true;
 		}
-
-		private bool ExportBOMFeature(BomFeature InBomFeature, FMetadata InMetadata)
+		
+		public static bool ExportBomFeature(BomFeature InBomFeature, FMetadata InMetadata)
 		{
 			Feature Feat = InBomFeature.GetFeature();
 			string FeatureName = Feat.Name;
@@ -289,38 +207,36 @@ namespace DatasmithSolidworks
 			}
 			return true;
 		}
-
-		private void ExportTable(TableAnnotation InTable, FMetadata InMetadata, string FeatureNamePrefix)
+		
+		public static void ExportTable(TableAnnotation InTable, FMetadata InMetadata, string FeatureNamePrefix)
 		{
-			int NBHearders = InTable.GetHeaderCount();
-			if (NBHearders == 0)
+			int HeaderCount = InTable.GetHeaderCount();
+			if (HeaderCount == 0)
 			{
-				//TOD log error here
 				return;
 			}
+
 			int Index = 0;
 			int SplitCount = 0;
 			int RangeStart = 0;
 			int RangeEnd = 0;
-			int NBRows = 0;
-			int NBCols = 0;
-
+			int ColumnCount = 0;
 			swTableSplitDirection_e SplitDir = (swTableSplitDirection_e) InTable.GetSplitInformation(ref Index, ref SplitCount, ref RangeStart, ref RangeEnd);
 
 			if (SplitDir == swTableSplitDirection_e.swTableSplit_None)
 			{
-				NBRows = InTable.RowCount;
-				NBCols = InTable.ColumnCount;
-				RangeStart = NBHearders;
-				RangeEnd = NBRows - 1;
+				int RowCount = InTable.RowCount;
+				ColumnCount = InTable.ColumnCount;
+				RangeStart = HeaderCount;
+				RangeEnd = RowCount - 1;
 			}
 			else
 			{
-				NBCols = InTable.ColumnCount;
+				ColumnCount = InTable.ColumnCount;
 				if (Index == 1)
 				{
 					// Add header offset for first portion of table
-					RangeStart += NBHearders;
+					RangeStart += HeaderCount;
 				}
 			}
 
@@ -329,48 +245,34 @@ namespace DatasmithSolidworks
 				InMetadata.AddPair("BOMTable_Feature_" + FeatureNamePrefix, InTable.Title);
 			}
 
-			string[] HeadersTitles = new string[NBHearders];
-			for (int i = 0; i < NBHearders; i++)
+			string[] HeadersTitles = new string[ColumnCount];
+			for (int ColumnIndex = 0; ColumnIndex < ColumnCount; ColumnIndex++)
 			{
-				HeadersTitles[i] = (string)InTable.GetColumnTitle2(i, true);
+				HeadersTitles[ColumnIndex] = InTable.GetColumnTitle2(ColumnIndex, true);
 			}
 
-			for (int I = RangeStart; I <= RangeEnd; I++)
+			for (int RowIndex = RangeStart; RowIndex <= RangeEnd; RowIndex++)
 			{
-				for (int J = 0; J < NBCols; J++)
+				for (int ColumnIndex = 0; ColumnIndex < ColumnCount; ColumnIndex++)
 				{
-					InMetadata.AddPair("BOMTable_Feature_" + FeatureNamePrefix + "_" + HeadersTitles[J], InTable.Text2[I, J, true]);
+					string HeadersTitle = HeadersTitles[ColumnIndex];
+					InMetadata.AddPair("BOMTable_Feature_" + FeatureNamePrefix + "_" + HeadersTitle, InTable.Text2[RowIndex, ColumnIndex, true]);
 				}
 			}
-
 		}
 
-		static public bool AddAssemblyComponentMetadata(Component2 InComponent, FMetadata InMetadata)
-		{
-			//object[] varComp = (object[])(assemblydoc.GetComponents(false));
-			//int nbcomponents = assemblydoc.GetComponentCount(false);
-			//Component2 component = null;
-			//for (int i = 0; i < nbcomponents; i++)
-			//{
-			//    component = (Component2)varComp[i];
-			//}
-			InMetadata.AddPair("ComponentReference", InComponent.ComponentReference);
-
-			return true;
-		}
-
-		static private bool AddAssemblyDisplayStateMetadata(AssemblyDoc InAssemblyDoc, FMetadata InMetadata)
+		public static void AddAssemblyDisplayStateMetadata(AssemblyDoc InAssemblyDoc, FMetadata InMetadata)
 		{
 			ModelDocExtension Ext = ((ModelDoc2)InAssemblyDoc).Extension;
 			if (Ext == null)
 			{
-				return false;
+				return;
 			}
 
-			ConfigurationManager CFM = ((ModelDoc2)InAssemblyDoc).ConfigurationManager;
-			if (CFM != null)
+			ConfigurationManager Cfm = ((ModelDoc2)InAssemblyDoc).ConfigurationManager;
+			if (Cfm != null)
 			{
-				Configuration ActiveConf = CFM.ActiveConfiguration;
+				Configuration ActiveConf = Cfm.ActiveConfiguration;
 				object[] DisplayStates = ActiveConf.GetDisplayStates();
 				if (DisplayStates != null)
 				{
@@ -383,50 +285,50 @@ namespace DatasmithSolidworks
 			}
 
 			object[] VarComp = (object[])(InAssemblyDoc.GetComponents(false));
-			int NBComponents = InAssemblyDoc.GetComponentCount(false);
-			Component2[] ListComp = new Component2[NBComponents];
+			int ComponentCount = InAssemblyDoc.GetComponentCount(false);
+			Component2[] ListComp = new Component2[ComponentCount];
 			DisplayStateSetting DSS = Ext.GetDisplayStateSetting((int)swDisplayStateOpts_e.swThisDisplayState);
 			DSS.Option = (int)swDisplayStateOpts_e.swThisDisplayState;
 
-			for (int i = 0; i < NBComponents; i++)
+			for (int ComponentIndex = 0; ComponentIndex < ComponentCount; ComponentIndex++)
 			{
-				ListComp[i] = (Component2)VarComp[i];
+				ListComp[ComponentIndex] = (Component2)VarComp[ComponentIndex];
 			}
 			DSS.Entities = ListComp;
 
-			System.Array displaymodearray = (System.Array)Ext.DisplayMode[DSS];
-			System.Array transparencyarray = (System.Array)Ext.Transparency[DSS];
-			for (int Idx = 0; Idx < NBComponents; Idx++)
+			System.Array DisplayModes = (System.Array)Ext.DisplayMode[DSS];
+			System.Array Transparencies = (System.Array)Ext.Transparency[DSS];
+			for (int Idx = 0; Idx < ComponentCount; Idx++)
 			{
-				InMetadata.AddPair("Component_Display_State_DisplayMode_" + ((Component2)VarComp[Idx]).Name2, ((swDisplayMode_e)displaymodearray.GetValue(Idx)).ToString());
-				InMetadata.AddPair("Component_Display_State_Transparency_" + ((Component2)VarComp[Idx]).Name2, ((swTransparencyState_e)transparencyarray.GetValue(Idx)).ToString());
+				InMetadata.AddPair("Component_Display_State_DisplayMode_" + ((Component2)VarComp[Idx]).Name2, ((swDisplayMode_e)DisplayModes.GetValue(Idx)).ToString());
+				InMetadata.AddPair("Component_Display_State_Transparency_" + ((Component2)VarComp[Idx]).Name2, ((swTransparencyState_e)Transparencies.GetValue(Idx)).ToString());
 			}
-			return true;
+			return;
 		}
 
-		static private bool ExportCustomPropertyManagerToMetadata(CustomPropertyManager InCustomPropertyManager, FMetadata InMetadata, string InPrefix)
+		private static bool ExportCustomPropertyManagerToMetadata(CustomPropertyManager InCustomPropertyManager, FMetadata InMetadata, string InPrefix)
 		{
 			if (InCustomPropertyManager == null)
 			{
+				LogDebug($"    NULL");
 				return false;
 			}
 
 			object PropertiesNamesObject = null;
 			object PropertiesValuesObject = null;
-			string[] PropertiesNames;
-			object[] PropertiesValues;
-			int[] PropertiesTypes;
-			object PropertiestypesObject = null;
+			object PropertiesTypesObject = null;
 			object ResolvedObject = false;
-			object LinktopropObject = false;
-
-			InCustomPropertyManager.GetAll3(ref PropertiesNamesObject, ref PropertiestypesObject, ref PropertiesValuesObject, ref ResolvedObject, ref LinktopropObject);
-			PropertiesNames = (string[])PropertiesNamesObject;
-			PropertiesValues = (object[])PropertiesValuesObject;
-			PropertiesTypes = (int[])PropertiestypesObject;
-
+			object LinkToPropObject = false;
+
+			InCustomPropertyManager.GetAll3(ref PropertiesNamesObject, ref PropertiesTypesObject, ref PropertiesValuesObject, ref ResolvedObject, ref LinkToPropObject);
+			string[] PropertiesNames = (string[])PropertiesNamesObject;
+			object[] PropertiesValues = (object[])PropertiesValuesObject;
+			int[] PropertiesTypes = (int[])PropertiesTypesObject;
+			
+			LogDebug($"  Count: {InCustomPropertyManager.Count}");
 			for (int Idx = 0; Idx < InCustomPropertyManager.Count; Idx++)
 			{
+				LogDebug($"    {PropertiesNames[Idx]}: Type={PropertiesTypes[Idx]}, Value={PropertiesValues[Idx]}");
 				switch (PropertiesTypes[Idx])
 				{
 					case (int)swCustomInfoType_e.swCustomInfoUnknown:
diff --git a/Engine/Source/Programs/Enterprise/Datasmith/DatasmithSolidworksExporter/Private/PartDocument.cs b/Engine/Source/Programs/Enterprise/Datasmith/DatasmithSolidworksExporter/Private/PartDocument.cs
--- a/Engine/Source/Programs/Enterprise/Datasmith/DatasmithSolidworksExporter/Private/PartDocument.cs
+++ b/Engine/Source/Programs/Enterprise/Datasmith/DatasmithSolidworksExporter/Private/PartDocument.cs
@@ -290,6 +290,10 @@ namespace DatasmithSolidworks
 					ExportInfo.Label = ActorName.GetString();
 					ExportInfo.MeshName = MeshNameForConfiguration;
 					ExportInfo.Type = bHasMesh ? EActorType.MeshActor : EActorType.SimpleActor;
+
+					ExportInfo.Metadata = ConfigurationExporter.CombinedTree.Children[0].Metadata;
+					ExportInfo.Metadata.OwnerName = ExportInfo.Name;
+
 					Exporter.ExportOrUpdateActor(ExportInfo);
 				}
 			}
@@ -340,6 +344,9 @@ namespace DatasmithSolidworks
 				ExportInfo.Label = ActorName.GetString();
 				ExportInfo.MeshName = MeshName;
 				ExportInfo.Type = bHasMesh ? EActorType.MeshActor : EActorType.SimpleActor;
+				
+				ExportInfo.Metadata = ConfigurationExporter.CombinedTree.Children[0].Metadata;
+				ExportInfo.Metadata.OwnerName = ExportInfo.Name;
 				Exporter.ExportOrUpdateActor(ExportInfo);
 			}
 
@@ -350,7 +357,7 @@ namespace DatasmithSolidworks
 		{
 			return null;
 		}
-
+		
 		public override void AddMeshForComponent(FComponentName ComponentName, FMeshName MeshName)
 		{
 		}
