[SketchUp] Export/sync of Layer/Tag visibility takes into account Default layer visibility overrides

- When an entity(any - from face to component instance) has is "Untagged"(aka default layer) in UI
its visibility is overridden by the parent's Tag visibility.

- So: untagged faces(which is usual) contaigned in a tagged component are not hidden when Untagged Tag is hidded. Untagged components within tagged components behaves the same. This is how SU works.

diff --git a/Engine/Source/Programs/Enterprise/Datasmith/DatasmithSketchUpRubyExporter/Private/DatasmithSketchUpComponent.cpp b/Engine/Source/Programs/Enterprise/Datasmith/DatasmithSketchUpRubyExporter/Private/DatasmithSketchUpComponent.cpp
--- a/Engine/Source/Programs/Enterprise/Datasmith/DatasmithSketchUpRubyExporter/Private/DatasmithSketchUpComponent.cpp
+++ b/Engine/Source/Programs/Enterprise/Datasmith/DatasmithSketchUpRubyExporter/Private/DatasmithSketchUpComponent.cpp
@@ -129,7 +129,12 @@ void FEntityWithEntities::UpdateOccurrenceMeshActors(FExportContext& Context, FN
 
 void FNodeOccurence::UpdateVisibility(FExportContext& Context)
 {
-	if (bVisibilityInvalidated)
+	if (bPropertiesInvalidated)
+	{
+		Entity.UpdateOccurrenceLayer(Context, *this);
+	}
+
+	if (bVisibilityInvalidated || bPropertiesInvalidated)
 	{
 		Entity.UpdateOccurrenceVisibility(Context, *this);
 		bVisibilityInvalidated = false;
@@ -928,8 +933,6 @@ void FComponentInstance::UpdateOccurrenceTransformation(FExportContext& Context,
 
 void FComponentInstance::UpdateOccurrence(FExportContext& Context, FNodeOccurence& Node)
 {
-	Node.EffectiveLayerRef = DatasmithSketchUpUtils::GetEffectiveLayer(GetComponentInstanceRef(), Node.ParentNode->EffectiveLayerRef);
-
 	if (FDefinition* EntityDefinition = GetDefinition())
 	{
 		BuildNodeNames(Node);
@@ -1159,10 +1162,17 @@ void FImageCollection::LayerModified(FEntityIDType LayerId)
 	}
 }
 
+void FComponentInstance::UpdateOccurrenceLayer(FExportContext& Context, FNodeOccurence& Node)
+{
+	Node.EffectiveLayerRef = DatasmithSketchUpUtils::GetEffectiveLayer(GetComponentInstanceRef(), Node.ParentNode->EffectiveLayerRef);
+}
+
 void FComponentInstance::UpdateOccurrenceVisibility(FExportContext& Context, FNodeOccurence& Node)
 {
+	bool bEffectiveLayerVisible = Context.Layers.IsLayerVisible(Node.EffectiveLayerRef);
+
 	// Parent node, component instance and layer - all should be visible to have node visible
-	bool bVisibilityChanged = Node.SetVisibility(Node.ParentNode->bVisible && !bHidden && bLayerVisible);
+	bool bVisibilityChanged = Node.SetVisibility(Node.ParentNode->bVisible && !bHidden && bEffectiveLayerVisible);
 
 	EntityOccurrenceVisible(&Node, Node.bVisible);
 
@@ -1235,7 +1245,11 @@ FString FModel::GetEntityLabel()
 {
 	return "";
 }
-   
+
+void FModel::UpdateOccurrenceLayer(FExportContext& Context, FNodeOccurence&)
+{
+}
+
 void FModel::InvalidateOccurrencesGeometry(FExportContext& Context)
 {
 	Context.RootNode->InvalidateMeshActors();
diff --git a/Engine/Source/Programs/Enterprise/Datasmith/DatasmithSketchUpRubyExporter/Private/DatasmithSketchUpComponent.h b/Engine/Source/Programs/Enterprise/Datasmith/DatasmithSketchUpRubyExporter/Private/DatasmithSketchUpComponent.h
--- a/Engine/Source/Programs/Enterprise/Datasmith/DatasmithSketchUpRubyExporter/Private/DatasmithSketchUpComponent.h
+++ b/Engine/Source/Programs/Enterprise/Datasmith/DatasmithSketchUpRubyExporter/Private/DatasmithSketchUpComponent.h
@@ -100,6 +100,8 @@ namespace DatasmithSketchUp
 		SUTransformation BakeTransform;
 
 		FMaterialIDType InheritedMaterialID;
+		// Resolved Layer/Tag on the node after considering own and parent's
+		// The rule: Default(Layer0 or Untagged in UI) layer/tag is overridden by parent's
 		SULayerRef EffectiveLayerRef = SU_INVALID;
 		bool bVisible = true; // Computed visibility for this occurrence(affecting descendants)
 
@@ -347,6 +349,7 @@ namespace DatasmithSketchUp
 		virtual void ApplyOverrideMaterialToNode(FNodeOccurence& Node, FMaterialOccurrence& Material) = 0;
 
 		virtual void UpdateOccurrence(FExportContext& Context, FNodeOccurence& Node) = 0; // Update occurrence of this entity
+		virtual void UpdateOccurrenceLayer(FExportContext& Context, FNodeOccurence&) = 0; // Resolve effective layer for the occurrence
 		virtual void UpdateOccurrenceVisibility(FExportContext& Context, FNodeOccurence&) = 0;  // Re-evaluate visibility of entity's occurrence
 		virtual void UpdateOccurrenceMeshActors(FExportContext& Context, FNodeOccurence& Node) = 0; // Rebuild datasmith actors of entity's occurrence
 		virtual void ResetOccurrenceActors(FExportContext& Context, FNodeOccurence& Node) = 0; // Remove datasmith actors of entity's occurrence from datasmith scene
@@ -437,6 +440,7 @@ namespace DatasmithSketchUp
 		virtual int64 GetPersistentId() override;
 		virtual FString GetEntityName() override;
 		virtual FString GetEntityLabel() override;
+		virtual void UpdateOccurrenceLayer(FExportContext& Context, FNodeOccurence&) override;
 		virtual void UpdateOccurrenceVisibility(FExportContext& Context, FNodeOccurence&) override;
 		virtual void UpdateMetadata(FExportContext& Context) override;
 		virtual void UpdateEntityProperties(FExportContext& Context) override;
@@ -480,6 +484,7 @@ namespace DatasmithSketchUp
 		virtual void ApplyOverrideMaterialToNode(FNodeOccurence& Node, FMaterialOccurrence& Material) override;
 		virtual void UpdateOccurrence(FExportContext& Context, FNodeOccurence& Node) override;
 		void RemoveImageFromDatasmithScene(FExportContext& Context);
+		virtual void UpdateOccurrenceLayer(FExportContext& Context, FNodeOccurence&) override;
 		virtual void UpdateOccurrenceVisibility(FExportContext& Context, FNodeOccurence&) override;
 		virtual void UpdateOccurrenceMeshActors(FExportContext& Context, FNodeOccurence& Node) override;
 		virtual void ResetOccurrenceActors(FExportContext& Context, FNodeOccurence& Node) override;
@@ -532,6 +537,7 @@ namespace DatasmithSketchUp
 		virtual int64 GetPersistentId() override;
 		virtual FString GetEntityName() override;
 		virtual FString GetEntityLabel() override;
+		virtual void UpdateOccurrenceLayer(FExportContext& Context, FNodeOccurence&) override;
 		virtual void UpdateOccurrenceVisibility(FExportContext& Context, FNodeOccurence&) override;
 
 		virtual void UpdateMetadata(FExportContext& Context) override;
diff --git a/Engine/Source/Programs/Enterprise/Datasmith/DatasmithSketchUpRubyExporter/Private/DatasmithSketchUpExportContext.cpp b/Engine/Source/Programs/Enterprise/Datasmith/DatasmithSketchUpRubyExporter/Private/DatasmithSketchUpExportContext.cpp
--- a/Engine/Source/Programs/Enterprise/Datasmith/DatasmithSketchUpRubyExporter/Private/DatasmithSketchUpExportContext.cpp
+++ b/Engine/Source/Programs/Enterprise/Datasmith/DatasmithSketchUpRubyExporter/Private/DatasmithSketchUpExportContext.cpp
@@ -468,6 +468,13 @@ FLayerIDType FLayerCollection::GetLayerId(SULayerRef LayerRef)
 	return DatasmithSketchUpUtils::GetEntityID(SULayerToEntity(LayerRef));
 }
 
+bool FLayerCollection::IsDefault(FLayerIDType LayerID)
+{
+	SULayerRef DefaultLayerRef = SU_INVALID;
+	SUModelGetDefaultLayer(Context.ModelRef, &DefaultLayerRef);
+	return LayerID == GetLayerId(DefaultLayerRef);
+}
+
 void FComponentDefinitionCollection::PopulateFromModel(SUModelRef InModelRef)
 {
 	// Get the number of normal component definitions in the SketchUp model.
diff --git a/Engine/Source/Programs/Enterprise/Datasmith/DatasmithSketchUpRubyExporter/Private/DatasmithSketchUpExportContext.h b/Engine/Source/Programs/Enterprise/Datasmith/DatasmithSketchUpRubyExporter/Private/DatasmithSketchUpExportContext.h
--- a/Engine/Source/Programs/Enterprise/Datasmith/DatasmithSketchUpRubyExporter/Private/DatasmithSketchUpExportContext.h
+++ b/Engine/Source/Programs/Enterprise/Datasmith/DatasmithSketchUpRubyExporter/Private/DatasmithSketchUpExportContext.h
@@ -316,6 +316,7 @@ namespace DatasmithSketchUp
 		SULayerRef GetLayer(FLayerIDType LayerID);
 
 		FLayerIDType GetLayerId(SULayerRef LayerRef);
+		bool IsDefault(FLayerIDType LayerID);
 
 		TMap<FLayerIDType, bool> LayerVisibility;
 	private:
diff --git a/Engine/Source/Programs/Enterprise/Datasmith/DatasmithSketchUpRubyExporter/Private/DatasmithSketchUpImage.cpp b/Engine/Source/Programs/Enterprise/Datasmith/DatasmithSketchUpRubyExporter/Private/DatasmithSketchUpImage.cpp
--- a/Engine/Source/Programs/Enterprise/Datasmith/DatasmithSketchUpRubyExporter/Private/DatasmithSketchUpImage.cpp
+++ b/Engine/Source/Programs/Enterprise/Datasmith/DatasmithSketchUpRubyExporter/Private/DatasmithSketchUpImage.cpp
@@ -246,12 +246,14 @@ void FImage::SetupActor(FExportContext& Context, FNodeOccurence& Node)
 	}
 }
 
-
-void FImage::UpdateOccurrenceVisibility(FExportContext& Context, FNodeOccurence& Node)
+void FImage::UpdateOccurrenceLayer(FExportContext& Context, FNodeOccurence& Node)
 {
 	// Parent node, component instance and layer - all should be visible to have node visible
 	Node.SetVisibility(Node.ParentNode->bVisible && !bHidden && bLayerVisible);
+}
 
+void FImage::UpdateOccurrenceVisibility(FExportContext& Context, FNodeOccurence& Node)
+{
 	EntityOccurrenceVisible(&Node, Node.bVisible);
 
 	if (Node.bVisible)
diff --git a/Engine/Source/Programs/Enterprise/Datasmith/DatasmithSketchUpRubyExporter/Private/DatasmithSketchUpMesh.cpp b/Engine/Source/Programs/Enterprise/Datasmith/DatasmithSketchUpRubyExporter/Private/DatasmithSketchUpMesh.cpp
--- a/Engine/Source/Programs/Enterprise/Datasmith/DatasmithSketchUpRubyExporter/Private/DatasmithSketchUpMesh.cpp
+++ b/Engine/Source/Programs/Enterprise/Datasmith/DatasmithSketchUpRubyExporter/Private/DatasmithSketchUpMesh.cpp
@@ -924,7 +924,19 @@ void CombineSketchUpEntitiesFaces(FExportContext& Context, SUEntitiesRef Entitie
 		bool bFaceHidden = false;
 		SUDrawingElementGetHidden(SUFaceToDrawingElement(FaceRef), &bFaceHidden);
 
-		if (!bFaceHidden && Context.Layers.IsLayerVisible(LayerRef))
+		// Check layer visibility without considering parent components
+		// Assume default layer is visible. As it's overriden by parent components/instances and will have visibility of those layers
+		// todo: more full support for layers visibility should consider occurrences where this mesh is used
+		// if there are layers that are visible/hidden which override default layer of this specific mesh
+		//
+		// E.g. when this mesh has faces with Default layer assigned. AND other faces, with other layers.
+		// This means these are faces of a free geometry within some component(or model itself).
+		// The component can be instantiated more than once.
+		// Let's say it has instance A with LayerA and instance B with LayerB. LayerA and LayerB override those
+		// faces with Default layer and as a consequence override faces visibility.
+		bool bLayerVisibility = Context.Layers.IsLayerVisible(LayerRef) || Context.Layers.IsDefault(LayerId);
+
+		if (!bFaceHidden && bLayerVisibility)
 		{
 			ExtractedMesh.AddFace(Context, FaceRef, LayerId);
 		}
