[Navisworks] Fixex crash exporting high poly density terrain geometry

diff --git a/Engine/Source/Programs/Enterprise/Datasmith/DatasmithNavisworksExporter/Private/DatasmithNavisworksPlugin.cs b/Engine/Source/Programs/Enterprise/Datasmith/DatasmithNavisworksExporter/Private/DatasmithNavisworksPlugin.cs
--- a/Engine/Source/Programs/Enterprise/Datasmith/DatasmithNavisworksExporter/Private/DatasmithNavisworksPlugin.cs
+++ b/Engine/Source/Programs/Enterprise/Datasmith/DatasmithNavisworksExporter/Private/DatasmithNavisworksPlugin.cs
@@ -34,6 +34,7 @@ namespace DatasmithNavisworks
 		public bool USE_COMPONENTS = false;
 		public bool CREATE_INSTANCES_FOR_EMPTY_GEOMETRY = false; // Should we bother instantiating empty geometry(might want for debugging or metadata)?
 		public bool bUseClipbox = false; // This is disabled for now(can be useful for debugging, though)
+		public int MaxTriangleCount = 1000000; // limit triangle counts heavily - to avoid OOM
 	};
 
 	[Plugin("DatasmithNavisworksExporter.Ribbon", "EpicGames", 
@@ -531,7 +532,7 @@ namespace DatasmithNavisworks
 			public List<TransformMatrix> Transforms;
 			public List<Appearance> Appearances;
 			public GuiProperties Metadata;
-			public FDatasmithFacadeActorMesh MeshActor;
+			public List<FDatasmithFacadeActorMesh> MeshActors = new List<FDatasmithFacadeActorMesh>();
 			public GeometryUtil.BoundingBox BoundingBox;
 
 			public GeometryInstance(InstancedWithInstancesToExport Instanced, SceneItem Item, int MeshesCountHint)
@@ -1380,7 +1381,7 @@ namespace DatasmithNavisworks
 		{
 			InNode.GeometryDepth = InNode.HasGeometry() ? 0 : -1;
 			InNode.GeometryCountInHierarchy = InNode.HasGeometry() ? 1 : 0;
-			InNode.TriangleCountInHierarchy = InNode.GetTriangleCount();
+			InNode.TriangleCountInHierarchy = InNode.ComputeTriangleCount();
 
 			foreach (Node ChildNode in InNode.Children)
 			{
@@ -1407,8 +1408,9 @@ namespace DatasmithNavisworks
 			Node InNode)
 		{
 			// Limit meshes merging
-			const int MaxTriangleCount = 1000000; // limit triangle counts heavily - to avoid OOM
-			if (InNode.GeometryDepth > SceneContext.ExportParameters.MergeMaxDepth || InNode.GeometryCountInHierarchy < 2 || InNode.TriangleCountInHierarchy > MaxTriangleCount)
+			if (InNode.GeometryDepth > SceneContext.ExportParameters.MergeMaxDepth 
+			    || InNode.GeometryCountInHierarchy < 2 
+			    || InNode.TriangleCountInHierarchy > SceneContext.ExportParameters.MaxTriangleCount)
 			{
 				foreach (Node ChildNode in InNode.Children)
 				{
@@ -1699,48 +1701,87 @@ namespace DatasmithNavisworks
 						{
 							return false;
 						}
-
-						FDatasmithFacadeMeshElement DatasmithMeshElement = CreateDatasmithMeshElementForItem(SceneContext, Instance.SceneItem, 0);
-						int SlotCount = AssignAppearancesListToDatasmithMesh(SceneContext, DatasmithMeshElement, Instance.Appearances, out List<int> SlotRemap);
-						SceneContext.TotalMeshSectionCount += SlotCount;
-						SceneContext.TotalMeshSectionInstancedCount += SlotCount;
-
+						
 						// For each mesh that is composed from fragments with distinct transforms create its own DatasmithMesh 
-						using (FDatasmithFacadeMesh DatasmithMesh = CreateDatasmithMeshForBakedFragmentGeometry(SceneContext, Instance, SlotRemap, out Vector3d Translation))
+						// Compute bounding box of consolidated geometry to repivot merged geometry to the bbox center
+						GeometryUtil.BoundingBox BoundingBox = new GeometryUtil.BoundingBox();
+						foreach (var F in Instance.SharedInstancedForInstancesRef.InstancedGeometry.FragmentGeometries.Select((Geometry, Index) => new { Geometry, Index }))
 						{
-							if (DatasmithMesh == null)
+							OptimizeSourceGeometry(F.Geometry);
+							if (F.Geometry.GetGeometry(out Geometry Geometry))
 							{
-								Info($"Datasmith Mesh wasn't created for {Instance.SceneItem?.ComNode?.UserName}");
-								continue;
+								// Computing bbox in world space(by providing transform)
+								BoundingBox.Extend(Geometry, SceneContext.NavisworksToUnrealTransform.Multiply(Instance.GetTransform(F.Index)));
 							}
-
-							ExportDatasmithMesh(SceneContext, DatasmithMesh, DatasmithMeshElement);
-							TotalBakedMeshTriangles += DatasmithMesh.GetFacesCount();
-
-							FDatasmithFacadeActorMesh ItemMeshActor =
-								CreateDatasmithActorMeshForInstance(SceneContext, Instance, 0);
-							// Translation is expected to be in export(Unreal) coordinates already
-							TransformMatrix Transform = TransformMatrix.Identity.ConcatTranslation(Translation);
-							ItemMeshActor.SetWorldTransform(Transform.Floats);
-							ItemMeshActor.SetMesh(DatasmithMeshElement.GetName());
 						}
-
+						Instance.BoundingBox = BoundingBox;
+						
+						GetNameAndLabelFromItem(Instance.SceneItem, out string ItemName, out string ItemLabel);
+						
+						for (int FragmentStart = 0, PartIndex = 0; FragmentStart < InstancedGeometry.FragmentGeometryCount; PartIndex++)
+						{
+							int FragmentEnd = GetGeometryFragmentsToFillDatasmithMesh(Instance.SharedInstancedForInstancesRef.InstancedGeometry,
+								SceneContext.ExportParameters.MaxTriangleCount, FragmentStart, out List<Geometry> GeometryFragments, out int VerticesCount,
+								out int FacesCount);
+							// Get transforms for fragments
+							TransformMatrix[] FragmentsTransforms = Enumerable.Range(FragmentStart, FragmentEnd - FragmentStart)
+								.Select(I => SceneContext.NavisworksToUnrealTransform.Multiply(Instance.GetTransform(I))).ToArray();
+							
+							// DatasmithMesh wasn't filled with all the fragments from the InstancedGeometry
+							// because it was too big. Geometry will use multiple Datasmith meshes, continuing from the FragmentEnd 
+							FragmentStart = FragmentEnd;
+							
+							// Center of the bounding box of vertices already transformed to export(Unreal) coordinates
+							Vector3d PivotTranslation = BoundingBox.BottomCenter;
+							
+							using (FDatasmithFacadeMesh DatasmithMesh = new FDatasmithFacadeMesh())
+							{
+								if (DatasmithMesh == null)
+								{
+									Info($"Datasmith Mesh wasn't created for {Instance.SceneItem?.ComNode?.UserName}");
+									continue;
+								}
+								string MeshHashedName = FDatasmithFacadeElement.GetStringHash($"M{PartIndex}:{ItemName}");
+								DatasmithMesh.SetName(MeshHashedName);
+								
+								FDatasmithFacadeMeshElement DatasmithMeshElement = CreateDatasmithMeshElement(
+									SceneContext, MeshHashedName, ItemLabel + "_" + PartIndex);
+
+								int SlotCount = AssignAppearancesListToDatasmithMesh(SceneContext, DatasmithMeshElement, Instance.Appearances,
+									out List<int> SlotRemap);
+								SceneContext.TotalMeshSectionCount += SlotCount;
+								SceneContext.TotalMeshSectionInstancedCount += SlotCount;
+								
+								FillDatasmithMeshFromGeometryFragmentsWithTransform(DatasmithMesh,
+									SlotRemap, -PivotTranslation, GeometryFragments, FragmentsTransforms, VerticesCount, FacesCount);
+								
+								ExportDatasmithMesh(SceneContext, DatasmithMesh, DatasmithMeshElement);
+								TotalBakedMeshTriangles += DatasmithMesh.GetFacesCount();
+								
+								FDatasmithFacadeActorMesh ItemMeshActor =
+									CreateDatasmithActorMeshForInstance(SceneContext, Instance, PartIndex);
+								// Translation is expected to be in export(Unreal) coordinates already
+								TransformMatrix Transform = TransformMatrix.Identity.ConcatTranslation(PivotTranslation);
+								ItemMeshActor.SetWorldTransform(Transform.Floats);
+								ItemMeshActor.SetMesh(DatasmithMeshElement.GetName());
+							}
+						}
+						
 						TotalBakedMeshInstances++;
 					}
 
 					// Merge meshes that have same transform
 					if (Instances.WithSameTransform.Count > 0)
 					{
+						// TODO: Potential optimization - same mesh might have different set of appearances
+						// Right now we are creating separate mesh for each appearance list, might want to make override materials on mesh actor
 						foreach (var AppearanceListAndInstances in Instances.WithSameTransform.Select((Value, Index) =>
-							new {Index, AppearanceList = Value.Key, GeometryInstances = Value.Value}))
+							         new {Index, AppearanceList = Value.Key, GeometryInstances = Value.Value}))
 						{
 							MultiFragmentsGeometryInstances.AppearanceList AppearanceList =
 								AppearanceListAndInstances.AppearanceList;
 							List<GeometryInstance> GeometryInstances = AppearanceListAndInstances.GeometryInstances;
 
-							FDatasmithFacadeMeshElement DatasmithMeshElement = CreateDatasmithMeshElementForItem(SceneContext, InstancedGeometry.Item, It.Index);
-							int SlotCount = AssignAppearancesListToDatasmithMesh(SceneContext, DatasmithMeshElement, AppearanceList.Appearances, out List<int> SlotRemap);
-
 							// Compute bounding box of consolidated geometry to repivot merged geometry to the bbox center
 							// Bounding box is computed in local coordinates(before item transform applied)
 							GeometryUtil.BoundingBox BoundingBox = new GeometryUtil.BoundingBox();
@@ -1753,20 +1794,51 @@ namespace DatasmithNavisworks
 								}
 							}
 							Vector3d PivotTranslation = BoundingBox.BottomCenter;
-							TransformMatrix CenterTransform = TransformMatrix.Identity.ConcatTranslation(PivotTranslation);
-
-							// TODO: Potential optimization - same mesh might have different set of appearances
-							// Right now we are creating separate mesh for each appearance list, might want to make override materials on mesh actor
-							using (FDatasmithFacadeMesh MergedDatasmithMesh = CreateDatasmithMeshForItem(SceneContext, InstancedGeometry.Item, AppearanceListAndInstances.Index)) // Instantiate mesh with different name for each appearance
-							{
-								TotalMergedMeshes++;
-								FillDatasmithMeshFromInstancedGeometry(InstancedGeometry, MergedDatasmithMesh, SlotRemap, -PivotTranslation);
 
-								ExportDatasmithMesh(SceneContext, MergedDatasmithMesh, DatasmithMeshElement);
+							GetNameAndLabelFromItem(InstancedGeometry.Item, out string ItemName, out string ItemLabel);
+							
+							int SlotCountAll = 0;
 
-								SceneContext.TotalMeshSectionCount += SlotCount;
+							List<FDatasmithFacadeMeshElement> DatasmithMeshElements = new List<FDatasmithFacadeMeshElement>();
+							for(int FragmentStart = 0, PartIndex = 0;FragmentStart < InstancedGeometry.FragmentGeometryCount;PartIndex++)
+							{
+								int FragmentEnd = GetGeometryFragmentsToFillDatasmithMesh(
+									InstancedGeometry, SceneContext.ExportParameters.MaxTriangleCount, FragmentStart,
+									out List<Geometry> GeometryFragments, out int VerticesCount, out int FacesCount);
+								
+								// DatasmithMesh wasn't filled with all the fragments from the InstancedGeometry
+								// because it was too big. Geometry will use multiple Datasmith meshes, continuing from the FragmentEnd 
+								FragmentStart = FragmentEnd;
+								
+								using (FDatasmithFacadeMesh DatasmithMesh = new FDatasmithFacadeMesh()) // Instantiate mesh with different name for each appearance
+								{
+									TotalMergedMeshes++;
+									
+									// Hash the Datasmith mesh name to shorten it and make it valid
+									string MeshHashedName = FDatasmithFacadeElement.GetStringHash(
+										PartIndex == 0
+											? $"M{AppearanceListAndInstances.Index}:{ItemName}"
+											: $"M{AppearanceListAndInstances.Index}:{ItemName}:{PartIndex}");
+									DatasmithMesh.SetName(MeshHashedName);
+									
+									FDatasmithFacadeMeshElement DatasmithMeshElement = CreateDatasmithMeshElement(
+										SceneContext, MeshHashedName, ItemLabel + "_" + It.Index);
+									DatasmithMeshElements.Add(DatasmithMeshElement);
+									
+									int SlotCount = AssignAppearancesListToDatasmithMesh(SceneContext, DatasmithMeshElement, AppearanceList.Appearances,
+										out List<int> SlotRemap);
+									SlotCountAll += SlotCount;
+
+									FillDatasmithMeshFromGeometryFragments(DatasmithMesh, SlotRemap, -PivotTranslation, GeometryFragments,
+										VerticesCount, FacesCount);
+									
+									ExportDatasmithMesh(SceneContext, DatasmithMesh, DatasmithMeshElement);
+									
+									SceneContext.TotalMeshSectionCount += SlotCount;
+								}
 							}
-
+							
+							TransformMatrix CenterTransform = TransformMatrix.Identity.ConcatTranslation(PivotTranslation);
 
 							foreach (GeometryInstance Instance in GeometryInstances)
 							{
@@ -1774,18 +1846,22 @@ namespace DatasmithNavisworks
 								{
 									return false;
 								}
-
-								FDatasmithFacadeActorMesh ItemMeshActor = CreateDatasmithActorMeshForInstance(SceneContext, Instance, 0);
-
+								
 								// Repivot to BBox center. Note - CenterTransform applied first(i.e. to local vertex positions)
 								TransformMatrix TransformRepivoted = Instance.GetTransform(0).Multiply(CenterTransform);
 								TransformMatrix WorldTransform = SceneContext.NavisworksToUnrealTransform.Multiply(TransformRepivoted);
-								ItemMeshActor.SetWorldTransform(WorldTransform.Floats);
-								ItemMeshActor.SetMesh(DatasmithMeshElement.GetName());
 								Instance.BoundingBox = BoundingBox.Transform(WorldTransform);
 
+								foreach (var Mesh in DatasmithMeshElements.Select((DatasmithMeshElement, MeshIndex) =>
+									         new { DatasmithMeshElement, MeshIndex }))
+								{
+									FDatasmithFacadeActorMesh ItemMeshActor = CreateDatasmithActorMeshForInstance(SceneContext, Instance, Mesh.MeshIndex);
+									ItemMeshActor.SetWorldTransform(WorldTransform.Floats);
+									ItemMeshActor.SetMesh(Mesh.DatasmithMeshElement.GetName());
+								}
+								
 								TotalMergedMeshInstances++;
-								SceneContext.TotalMeshSectionInstancedCount += SlotCount;
+								SceneContext.TotalMeshSectionInstancedCount += SlotCountAll;
 							}
 						}
 					}
@@ -1808,7 +1884,7 @@ namespace DatasmithNavisworks
 
 			return true;
 		}
-
+		
 		private static TransformMatrix CreateNavisworksToUnrealTransform(SceneContext SceneContext)
 		{
 			InwOaPartition3 Partition = ((InwOaPartition3) SceneContext.State.CurrentPartition);
@@ -1966,7 +2042,7 @@ namespace DatasmithNavisworks
 				return Instance != null;
 			}
 
-			public long GetTriangleCount()
+			public long ComputeTriangleCount()
 			{
 				return HasGeometry() ? Instance.SharedInstancedForInstancesRef.InstancedGeometry.TriangleCount() : 0;
 			}
@@ -1996,15 +2072,26 @@ namespace DatasmithNavisworks
 					continue;
 				}
 
-				FDatasmithFacadeActorMesh MeshActor = CurrentNode.Instance?.MeshActor;
-				if (MeshActor != null)
+				if (CurrentNode.Instance != null)
 				{
-					Item.DatasmithActor.AddChild(MeshActor);
+					foreach (FDatasmithFacadeActorMesh MeshActor in CurrentNode.Instance.MeshActors)
+					{
+						Item.DatasmithActor.AddChild(MeshActor);
+					}
 				}
-
 				foreach (Node Child in CurrentNode.Children)
 				{
-					Item.DatasmithActor.AddChild(Child.SceneItem.DatasmithActor ?? Child.Instance.MeshActor);
+					if (Child.SceneItem.DatasmithActor != null)
+					{
+						Item.DatasmithActor.AddChild(Child.SceneItem.DatasmithActor);
+					}
+					else if (Child.Instance != null)
+					{
+						foreach (FDatasmithFacadeActorMesh MeshActor in Child.Instance.MeshActors)
+						{
+							Item.DatasmithActor.AddChild(MeshActor);
+						}
+					}
 				}
 
 				if (Item.Parent == null)
@@ -2082,16 +2169,18 @@ namespace DatasmithNavisworks
 				{
 					AddMetadataForSceneItemToDatasmithActor(SceneContext, Node.SceneItem, DatasmithActor);
 				}
-
-				FDatasmithFacadeActor DatasmithMeshActor = Node.Instance?.MeshActor;
-				if (DatasmithMeshActor != null)
+				
+				if (Node.Instance != null)
 				{
 					// Add metadata to MeshActor, if it was built for different SceneItem(this happens when scene was optimized so that single child geometry instance was moved to parent)
 					// or Node's DatasmithActor is not present(it's omitted when there's nothing else in the hierarchy except maybe a single MeshActor)
 					SceneItem InstanceSceneItem = Node.Instance.SceneItem;
-					if (((InstanceSceneItem != Node.SceneItem) || DatasmithActor == null))
+					if ((InstanceSceneItem != Node.SceneItem) || DatasmithActor == null)
 					{
-						AddMetadataForSceneItemToDatasmithActor(SceneContext, InstanceSceneItem, DatasmithMeshActor);
+						foreach (FDatasmithFacadeActor DatasmithMeshActor in Node.Instance.MeshActors)
+						{
+							AddMetadataForSceneItemToDatasmithActor(SceneContext, InstanceSceneItem, DatasmithMeshActor);
+						}
 					}
 				}
 			}
@@ -2446,7 +2535,7 @@ namespace DatasmithNavisworks
 				ItemMeshActor.SetIsComponent(true);
 			}
 
-			Instance.MeshActor = ItemMeshActor;
+			Instance.MeshActors.Add(ItemMeshActor);
 			return ItemMeshActor;
 		}
 
@@ -2468,16 +2557,22 @@ namespace DatasmithNavisworks
 		private static FDatasmithFacadeMeshElement CreateDatasmithMeshElementForItem(SceneContext SceneContext, SceneItem Item, int MeshIndex)
 		{
 			// Set mesh name/label from the Item that it's instantiated from
-			GetNameAndLabelFromItem(Item, out string Name, out string Label);
+			GetNameAndLabelFromItem(Item, out string ItemName, out string ItemLabel);
 			// Hash the Datasmith mesh name to shorten it and make it valid
-			string HashedName = FDatasmithFacadeElement.GetStringHash("M" + MeshIndex + ":" + Name);
+			string HashedName = FDatasmithFacadeElement.GetStringHash("M" + MeshIndex + ":" + ItemName);
+			return CreateDatasmithMeshElement(SceneContext, HashedName, ItemLabel + "_" + MeshIndex);
+		}
+		
+		private static FDatasmithFacadeMeshElement CreateDatasmithMeshElement(SceneContext SceneContext,
+			string HashedName, string Label)
+		{
 			FDatasmithFacadeMeshElement MeshElement = new FDatasmithFacadeMeshElement(HashedName);
-			MeshElement.SetLabel(Label + "_" + MeshIndex);
 			SceneContext.DatasmithMeshElements.Add(MeshElement); // hold on to reference
 			SceneContext.DatasmithScene.AddMesh(MeshElement);
+			MeshElement.SetLabel(Label);
 			return MeshElement;
 		}
-
+		
 		[MethodImpl(MethodImplOptions.NoInlining)]
 		private static bool ExportDatasmithMesh(SceneContext SceneContext, FDatasmithFacadeMesh DatasmithMesh, FDatasmithFacadeMeshElement DatasmithMeshElement)
 		{
@@ -2528,50 +2623,41 @@ namespace DatasmithNavisworks
 			}
 		}
 
-		private unsafe void FillDatasmithMeshFromInstancedGeometry(SceneItemInstancedGeometry InstancedGeometry,
-			FDatasmithFacadeMesh DatasmithMesh, List<int> SlotRemap, Vector3d PositionOffset)
+		private static unsafe void FillDatasmithMeshFromGeometryFragments(FDatasmithFacadeMesh DatasmithMesh, List<int> SlotRemap,
+			Vector3d PositionOffset, List<Geometry> GeometryFragments, int VerticesCount, int FacesCount)
 		{
 			const int UVChannelCount = 1;
 			const int DefaultUVChannel = 0;
-			int VerticesCount = 0;
-			int FacesCount = 0;
-			List<Geometry> GeometryFragments = new List<Geometry>(InstancedGeometry.FragmentGeometryCount);
-
-			for (int FragmentIndex = 0; FragmentIndex < InstancedGeometry.FragmentGeometryCount; ++FragmentIndex)
-			{
-				if (InstancedGeometry.GetFragmentGeometry(FragmentIndex).GetGeometry(out Geometry Geometry))
-				{
-					GeometryFragments.Add(Geometry);
-					VerticesCount += (int)Geometry.VertexCount;
-					FacesCount += (int)Geometry.TriangleCount;
-				}
-			}
-
 			DatasmithMesh.SetVerticesCount(VerticesCount);
 			DatasmithMesh.SetUVChannelsCount(UVChannelCount);
 			DatasmithMesh.SetUVCount(DefaultUVChannel, VerticesCount);
 			DatasmithMesh.SetFacesCount(FacesCount);
-
+			
 			int VertexIndexOffset = 0;
 			int FaceIndexOffset = 0;
 			for (int FragmentIndex = 0; FragmentIndex < GeometryFragments.Count; ++FragmentIndex)
 			{
 				Geometry Geometry = GeometryFragments[FragmentIndex];
-
+				
+				if (Geometry == null)
+				{
+					continue;
+				}
+				
 				for (int VertexIndex = 0; VertexIndex < Geometry.VertexCount; VertexIndex++)
 				{
 					GeometryVertex Vertex = new GeometryVertex(Geometry, VertexIndex);
-
+					
 					Vector3d Position = Vertex.Position + PositionOffset;
 					DatasmithMesh.SetVertex(VertexIndexOffset + VertexIndex, (float)Position.X, (float)Position.Y, (float)Position.Z);
-
+					
 					Vector3d Normal = Vertex.Normal;
 					DatasmithMesh.SetNormal(VertexIndexOffset + VertexIndex, (float)Normal.X, (float)Normal.Y, (float)Normal.Z);
-
+					
 					Vector2d Uv = Vertex.UV;
 					DatasmithMesh.SetUV(DefaultUVChannel, VertexIndexOffset + VertexIndex, (float)Uv.X, (float)Uv.Y);
 				}
-
+				
 				uint TriangleCount = Geometry.TriangleCount;
 				int MaterialId = SlotRemap[FragmentIndex];
 				for (int TriangleIndex = 0; TriangleIndex < TriangleCount; TriangleIndex++)
@@ -2580,88 +2666,52 @@ namespace DatasmithNavisworks
 					int Vertex1 = (int)Geometry.Indices[TriangleIndex * 3] + VertexIndexOffset;
 					int Vertex2 = (int)Geometry.Indices[TriangleIndex * 3 + 1] + VertexIndexOffset;
 					int Vertex3 = (int)Geometry.Indices[TriangleIndex * 3 + 2] + VertexIndexOffset;
-
+					
 					DatasmithMesh.SetFace(FaceIndex, Vertex1, Vertex2, Vertex3, MaterialId);
 					DatasmithMesh.SetFaceUV(FaceIndex, DefaultUVChannel, Vertex1, Vertex2, Vertex3);
 				}
-
+				
 				VertexIndexOffset += (int)Geometry.VertexCount;
 				FaceIndexOffset += (int)Geometry.TriangleCount;
 			}
 		}
 
-		private unsafe FDatasmithFacadeMesh  CreateDatasmithMeshForBakedFragmentGeometry(SceneContext SceneContext, GeometryInstance Instance, List<int> SlotRemap, out Vector3d OutTranslation)
+		private static unsafe void FillDatasmithMeshFromGeometryFragmentsWithTransform(FDatasmithFacadeMesh DatasmithMesh, List<int> SlotRemap, 
+			Vector3d PositionOffset, List<Geometry> GeometryFragments, TransformMatrix[] Transforms,
+			int VerticesCount, int FacesCount)
 		{
-			FDatasmithFacadeMesh DatasmithMesh = CreateDatasmithMeshForItem(SceneContext, Instance.SceneItem, 0);
 			const int UVChannelCount = 1;
 			const int DefaultUVChannel = 0;
-			int VerticesCount = 0;
-			int FacesCount = 0;
-			List<Geometry> GeometryFragments = new List<Geometry>(Instance.SharedInstancedForInstancesRef.InstancedGeometry.FragmentGeometryCount);
-
-			for (int FragmentIndex = 0; FragmentIndex < Instance.SharedInstancedForInstancesRef.InstancedGeometry.FragmentGeometryCount; ++FragmentIndex)
-			{
-				FragmentGeometry FragmentGeometry = Instance.SharedInstancedForInstancesRef.InstancedGeometry.GetFragmentGeometry(FragmentIndex);
-				OptimizeSourceGeometry(FragmentGeometry);
-
-				if (FragmentGeometry.GetGeometry(out Geometry Geometry))
-				{
-					GeometryFragments.Add(Geometry);
-					VerticesCount += (int)Geometry.VertexCount;
-					FacesCount += (int)Geometry.TriangleCount;
-				}
-				else
-				{
-					// Add an empty entry to make sure the indices match SlotRemap.
-					GeometryFragments.Add(null);
-				}
-			}
-
-			// Compute bounding box of consolidated geometry to repivot merged geometry to the bbox center
-			GeometryUtil.BoundingBox BoundingBox = new GeometryUtil.BoundingBox();
-			for (int FragmentIndex = 0; FragmentIndex < Instance.SharedInstancedForInstancesRef.InstancedGeometry.FragmentGeometryCount; FragmentIndex++)
-			{
-				FragmentGeometry FragmentGeometry = Instance.SharedInstancedForInstancesRef.InstancedGeometry.GetFragmentGeometry(FragmentIndex);
-				OptimizeSourceGeometry(FragmentGeometry);
-				if (FragmentGeometry.GetGeometry(out Geometry Geometry))
-				{
-					BoundingBox.Extend(Geometry, SceneContext.NavisworksToUnrealTransform.Multiply(Instance.GetTransform(FragmentIndex)));
-				}
-			}
-			// Center of the bounding box of vertices already transformed to export(Unreal) coordinates
-			Vector3d PivotTranslation = BoundingBox.BottomCenter;
-			OutTranslation = PivotTranslation; 
-
 			DatasmithMesh.SetVerticesCount(VerticesCount);
 			DatasmithMesh.SetUVChannelsCount(UVChannelCount);
 			DatasmithMesh.SetUVCount(DefaultUVChannel, VerticesCount);
 			DatasmithMesh.SetFacesCount(FacesCount);
-
+			
 			// For Fragments that don't all share single transform - bake Fragment transform into its vertices
 			int VertexIndexOffset = 0;
 			int FaceIndexOffset = 0;
 			for (int FragmentIndex = 0; FragmentIndex < GeometryFragments.Count; ++FragmentIndex)
 			{
 				Geometry Geometry = GeometryFragments[FragmentIndex];
-
+				
 				if (Geometry != null)
 				{
-					TransformMatrix Transform = SceneContext.NavisworksToUnrealTransform.Multiply(Instance.GetTransform(FragmentIndex));
-
+					TransformMatrix Transform = Transforms[FragmentIndex];
+					
 					for (int VertexIndex = 0; VertexIndex < Geometry.VertexCount; VertexIndex++)
 					{
 						GeometryVertex Vertex = new GeometryVertex(Geometry, VertexIndex);
-
-						Vector3d CoordTransformed = Transform.TransformPosition(Vertex.Position) - PivotTranslation;
-						DatasmithMesh.SetVertex(VertexIndexOffset + VertexIndex, (float)CoordTransformed.X, (float)CoordTransformed.Y, (float)CoordTransformed.Z);
-
-						Vector3 NormalTransformed = new Vector3(Transform.TransformNormal(Vertex.Normal));
-						DatasmithMesh.SetNormal(VertexIndexOffset + VertexIndex, (float)NormalTransformed.X, (float)NormalTransformed.Y, (float)NormalTransformed.Z);
-
+						
+						Vector3d Position = Transform.TransformPosition(Vertex.Position) + PositionOffset;
+						DatasmithMesh.SetVertex(VertexIndexOffset + VertexIndex, (float)Position.X, (float)Position.Y, (float)Position.Z);
+							
+						Vector3 Normal = new Vector3(Transform.TransformNormal(Vertex.Normal));
+						DatasmithMesh.SetNormal(VertexIndexOffset + VertexIndex, (float)Normal.X, (float)Normal.Y, (float)Normal.Z);
+						
 						Vector2d Uv = Vertex.UV;
 						DatasmithMesh.SetUV(DefaultUVChannel, VertexIndexOffset + VertexIndex, (float)Uv.X, (float)Uv.Y);
 					}
-
+					
 					uint TriangleCount = Geometry.TriangleCount;
 					int MaterialId = SlotRemap[FragmentIndex];
 					for (int TriangleIndex = 0; TriangleIndex < TriangleCount; TriangleIndex++)
@@ -2670,7 +2720,7 @@ namespace DatasmithNavisworks
 						int Vertex1 = (int)Geometry.Indices[TriangleIndex * 3] + VertexIndexOffset;
 						int Vertex2 = (int)Geometry.Indices[TriangleIndex * 3 + 1] + VertexIndexOffset;
 						int Vertex3 = (int)Geometry.Indices[TriangleIndex * 3 + 2] + VertexIndexOffset;
-
+						
 						DatasmithMesh.SetFace(FaceIndex, Vertex1, Vertex2, Vertex3, MaterialId);
 						DatasmithMesh.SetFaceUV(FaceIndex, DefaultUVChannel, Vertex1, Vertex2, Vertex3);
 					}
@@ -2678,11 +2728,53 @@ namespace DatasmithNavisworks
 					FaceIndexOffset += (int)Geometry.TriangleCount;
 				}
 			}
-
-			Instance.BoundingBox = BoundingBox;
-			return DatasmithMesh;
 		}
-
+		
+		/// <summary>
+		/// Collect geometry fragments from the geometry
+		/// starting at FragmentStart until threshold of MaxTriangleCount is reached
+		/// Also, computes and returns VerticesCount and FacesCount of the collected geometry
+		/// </summary>
+		/// <returns> last FragmentIndex that wasn't added to the output list, InstancedGeometry.FragmentGeometryCount if all fragments exhausted </returns>
+		private static int GetGeometryFragmentsToFillDatasmithMesh(SceneItemInstancedGeometry InstancedGeometry, 
+			int MaxTriangleCount, int FragmentStart, 
+			out List<Geometry> GeometryFragments, out int VerticesCount, out int FacesCount)
+		{
+			VerticesCount = 0;
+			FacesCount = 0;
+			GeometryFragments = new List<Geometry>(InstancedGeometry.FragmentGeometryCount);
+			
+			int FragmentEnd = InstancedGeometry.FragmentGeometryCount;
+			
+			for (int FragmentIndex = FragmentStart; FragmentIndex < InstancedGeometry.FragmentGeometryCount; ++FragmentIndex)
+			{
+				FragmentGeometry FragmentGeometry = InstancedGeometry.GetFragmentGeometry(FragmentIndex);
+				if (FacesCount + FragmentGeometry.TriangleCount > MaxTriangleCount)
+				{
+					if (FacesCount > 0)
+					{
+						// Stop merging geometry fragments into single mesh when mesh is too big
+						// Next parts of the geometry will fill another datasmith mesh
+						FragmentEnd = FragmentIndex;
+						break;
+					}
+				}
+				if (FragmentGeometry.GetGeometry(out Geometry Geometry))
+				{
+					GeometryFragments.Add(Geometry);
+					VerticesCount += (int)Geometry.VertexCount;
+					FacesCount += (int)Geometry.TriangleCount;
+				}
+				else
+				{
+					// Add an empty entry to make sure the indices match SlotRemap.
+					GeometryFragments.Add(null);
+				}
+			}
+			
+			return FragmentEnd;
+		}
+		
 		[MethodImpl(MethodImplOptions.NoInlining)]
 		private void OptimizeSourceGeometry(Geometry ReadGeometry)
 		{
