Removed DirectLink tests from build
Fixed issue with strsafe.h when compiling DatasmithSDK sample solution

#rnx

diff --git a/Engine/Source/Programs/Enterprise/Datasmith/DatasmithSDKSample/pch.h b/Engine/Source/Programs/Enterprise/Datasmith/DatasmithSDKSample/pch.h
--- a/Engine/Source/Programs/Enterprise/Datasmith/DatasmithSDKSample/pch.h
+++ b/Engine/Source/Programs/Enterprise/Datasmith/DatasmithSDKSample/pch.h
@@ -3,6 +3,7 @@
 #ifndef PCH_H
 #define PCH_H
 
+#include <strsafe.h>
 
 // DatasmithSDK project can be built from the main solution, we check that datasmith headers are accessible
 #if !__has_include("DatasmithCore.h")
