Fixed compilation error introduced by upgrade on compiler

#rnx

diff --git a/Engine/Source/Programs/Enterprise/Datasmith/DatasmithARCHICADExporter/Private/Utils/APIEnvir.h b/Engine/Source/Programs/Enterprise/Datasmith/DatasmithARCHICADExporter/Private/Utils/APIEnvir.h
--- a/Engine/Source/Programs/Enterprise/Datasmith/DatasmithARCHICADExporter/Private/Utils/APIEnvir.h
+++ b/Engine/Source/Programs/Enterprise/Datasmith/DatasmithARCHICADExporter/Private/Utils/APIEnvir.h
@@ -3,6 +3,8 @@
 #pragma once
 
 #include <stddef.h>
+#include <strsafe.h>
+#include <string>
 
 #include "IDatasmithSceneElements.h"
 #include "DatasmithSceneFactory.h"
