Fixed compilation error of SketchUp plugins on Mac

#rnx

diff --git a/Engine/Source/Programs/Enterprise/Datasmith/DatasmithSketchUpRubyExporter/Build/build_mac.py b/Engine/Source/Programs/Enterprise/Datasmith/DatasmithSketchUpRubyExporter/Build/build_mac.py
--- a/Engine/Source/Programs/Enterprise/Datasmith/DatasmithSketchUpRubyExporter/Build/build_mac.py
+++ b/Engine/Source/Programs/Enterprise/Datasmith/DatasmithSketchUpRubyExporter/Build/build_mac.py
@@ -124,6 +124,7 @@ UE_Project = f'{UE_Runtime}/Projects/Public'
 
 UE_ThirdParty = f'{UE_Engine}/Source/ThirdParty'
 UE_Imath = f'{UE_ThirdParty}/Imath/Deploy/Imath-3.1.9'
+UE_GSL = f'{UE_ThirdParty}/GuidelinesSupportLibrary/GSL-1144'
 
 UE_INCLUDES_PATH = f'''"{UE_Core}" "{UE_Core}/Internationalization" "{UE_Core}/Async" "{UE_Core}/Containers" "{UE_Core}/Delegates" "{UE_Core}/GenericPlatform" "{UE_Core}/HAL" "{UE_Core}/Logging" "{UE_Core}/Math" "{UE_Core}/Misc" "{UE_Core}/Modules" "{UE_Core}/Modules/Boilerplate" "{UE_Core}/ProfilingDebugging" "{UE_Core}/Serialization" "{UE_Core}/Serialization/Csv" "{UE_Core}/Stats" "{UE_Core}/Templates" "{UE_Core}/UObject" "{UE_Runtime}/CoreUObject/Public" "{UE_Project}/Interfaces" "{UE_Runtime}/TraceLog/Public"  "{UE_Runtime}/Datasmith/DatasmithCore/Public" "{UE_Runtime}/Datasmith/DirectLink/Public" "{UE_Runtime}/Messaging/Public" "{UE_Runtime}/Launch/Resources" "{UE_Engine}/Source/Developer/Datasmith/DatasmithExporter/Public" "{UE_Engine}/Source/Developer/Datasmith/DatasmithExporterUI/Public"'''
 
@@ -136,7 +137,7 @@ def create_compiler_for_arch(arch):
     CXX = ['xcrun', 'clang', 
     '-x', 'objective-c++', 
     '-arch', arch,
-    '-std=c++17',
+    '-std=c++20',
     '-stdlib=libc++',
     #    '-isysroot', '/Applications/Xcode.app/Contents/Developer/Platforms/MacOSX.platform/Developer/SDKs/MacOSX11.1.sdk',
     '-target', arch_target_dict[arch],
@@ -163,6 +164,7 @@ def create_compiler_for_arch(arch):
     include_paths = []
     include_paths += parse_paths(UE_INCLUDES_PATH)
     include_paths.append(f'{UE_Imath}/include')
+    include_paths.append(f'{UE_GSL}/include')
     include_paths.append(f'{SU_SDK_PATH}/samples/common/ThirdParty/ruby/include/mac')
     include_paths.append(f'{SU_SDK_PATH}/samples/common/ThirdParty/ruby/include/mac/{arch}-darwin') # todo: might not be needed - sample uses x86_64-darwin17 which doesn't exist
 
diff --git a/Engine/Source/Programs/Enterprise/Datasmith/DatasmithSketchUpRubyExporter/Private/DatasmithSketchUpSDKBegins.h b/Engine/Source/Programs/Enterprise/Datasmith/DatasmithSketchUpRubyExporter/Private/DatasmithSketchUpSDKBegins.h
--- a/Engine/Source/Programs/Enterprise/Datasmith/DatasmithSketchUpRubyExporter/Private/DatasmithSketchUpSDKBegins.h
+++ b/Engine/Source/Programs/Enterprise/Datasmith/DatasmithSketchUpRubyExporter/Private/DatasmithSketchUpSDKBegins.h
@@ -3,11 +3,19 @@
 #include "HAL/Platform.h"
 
 #if PLATFORM_WINDOWS
-// Begin Datasmith platform include gard.
+
+// Temporary solution for static analysis of DatasmithSketchUpSDKCeases.h
+#define SKETCHUP_WINDOWS_PLATFORM_TYPES_GUARD
+
+// Begin Datasmith platform include guard.
 #include "Windows/AllowWindowsPlatformTypes.h"
 
 // Turn off some compiler warnings.
 #pragma warning(push)
 #pragma warning(disable: 4100) // C4100: unreferenced formal parameter
 
+#elif defined(PLATFORM_MAC) && PLATFORM_MAC
+
+#define volatile
+
 #endif
diff --git a/Engine/Source/Programs/Enterprise/Datasmith/DatasmithSketchUpRubyExporter/Private/DatasmithSketchUpSDKCeases.h b/Engine/Source/Programs/Enterprise/Datasmith/DatasmithSketchUpRubyExporter/Private/DatasmithSketchUpSDKCeases.h
--- a/Engine/Source/Programs/Enterprise/Datasmith/DatasmithSketchUpRubyExporter/Private/DatasmithSketchUpSDKCeases.h
+++ b/Engine/Source/Programs/Enterprise/Datasmith/DatasmithSketchUpRubyExporter/Private/DatasmithSketchUpSDKCeases.h
@@ -2,6 +2,9 @@
 
 #if PLATFORM_WINDOWS
 
+// Temporary solution for static analysis of this include
+#ifdef SKETCHUP_WINDOWS_PLATFORM_TYPES_GUARD
+
 // Restore back the compiler warnings.
 #pragma warning(pop)
 
@@ -9,3 +12,9 @@
 #include "Windows/HideWindowsPlatformTypes.h"
 
 #endif
+
+#elif defined(PLATFORM_MAC) && PLATFORM_MAC
+
+#undef volatile
+
+#endif
