Adjusting pathing to xattr to be absolute for Datasmith building

OS may have different versions installed

#rnx

diff --git a/Engine/Source/Programs/Enterprise/Datasmith/DatasmithARCHICADExporter/Build/Config/setupdylibs.sh b/Engine/Source/Programs/Enterprise/Datasmith/DatasmithARCHICADExporter/Build/Config/setupdylibs.sh
--- a/Engine/Source/Programs/Enterprise/Datasmith/DatasmithARCHICADExporter/Build/Config/setupdylibs.sh
+++ b/Engine/Source/Programs/Enterprise/Datasmith/DatasmithARCHICADExporter/Build/Config/setupdylibs.sh
@@ -25,19 +25,19 @@ fi
 # Remove ArchiCAD resource tool from quarantine
 pushd "$UE_SDKS_ROOT/HostMac/Mac/Archicad"
 	chmod 777 23.1/Support/Tools/OSX/ResConv
-	xattr -r -d com.apple.quarantine 23.1/Support/Tools/OSX/ResConv
+	/usr/bin/xattr -r -d com.apple.quarantine 23.1/Support/Tools/OSX/ResConv
 
 	chmod 777 24/Support/Tools/OSX/ResConv
-	xattr -r -d com.apple.quarantine 24/Support/Tools/OSX/ResConv
+	/usr/bin/xattr -r -d com.apple.quarantine 24/Support/Tools/OSX/ResConv
 
 	chmod 777 25/Support/Tools/OSX/ResConv
-	xattr -r -d com.apple.quarantine 25/Support/Tools/OSX/ResConv
+	/usr/bin/xattr -r -d com.apple.quarantine 25/Support/Tools/OSX/ResConv
 
 	chmod 777 26/Support/Tools/OSX/ResConv
-	xattr -r -d com.apple.quarantine 26/Support/Tools/OSX/ResConv
+	/usr/bin/xattr -r -d com.apple.quarantine 26/Support/Tools/OSX/ResConv
 
 	chmod 777 27/Support/Tools/OSX/ResConv
-	xattr -r -d com.apple.quarantine 27/Support/Tools/OSX/ResConv
+	/usr/bin/xattr -r -d com.apple.quarantine 27/Support/Tools/OSX/ResConv
 popd
 
 OurDylibFolder=$projectPath/Dylibs
