Fixed compilation errors of ARCHICAD exporter on Mac
Linkage errors only happen in Horde build

#rnx

diff --git a/Engine/Source/Programs/Enterprise/Datasmith/DatasmithARCHICADExporter/Build/Config/UE_AC.xcconfig b/Engine/Source/Programs/Enterprise/Datasmith/DatasmithARCHICADExporter/Build/Config/UE_AC.xcconfig
--- a/Engine/Source/Programs/Enterprise/Datasmith/DatasmithARCHICADExporter/Build/Config/UE_AC.xcconfig
+++ b/Engine/Source/Programs/Enterprise/Datasmith/DatasmithARCHICADExporter/Build/Config/UE_AC.xcconfig
@@ -9,9 +9,10 @@ AC_INCLUDES_PATH = "$(AC_SDK_PATH)/Support/Inc" "$(AC_SDK_PATH)/Support/Modules/
 UE_Runtime = $(UE_Engine)/Source/Runtime
 UE_Core = $(UE_Runtime)/Core/Public
 UE_Project = $(UE_Runtime)/Projects/Public
+UE_ThirdParty = $(UE_Engine)/Source/ThirdParty
 UE_ArchiCADExporter = $(UE_Engine)/Source/Programs/Enterprise/Datasmith/DatasmithARCHICADExporter
 
-UE_INCLUDES_PATH = "$(UE_Core)" "$(UE_Core)/Internationalization" "$(UE_Core)/Async" "$(UE_Core)/Concurrency" "$(UE_Core)/Containers" "$(UE_Core)/Delegates" "$(UE_Core)/GenericPlatform" "$(UE_Core)/HAL" "$(UE_Core)/Logging" "$(UE_Core)/Math" "$(UE_Core)/Misc" "$(UE_Core)/Modules" "$(UE_Core)/Modules\Boilerplate" "$(UE_Core)/ProfilingDebugging" "$(UE_Core)/Serialization" "$(UE_Core)/Serialization\Csv" "$(UE_Core)/Stats" "$(UE_Core)/Templates" "$(UE_Core)/UObject" "$(UE_Project)/Public" "$(UE_Project)/Interfaces" "$(UE_Runtime)/TraceLog/Public"  "$(UE_Runtime)/Datasmith/DatasmithCore/Public" "$(UE_Runtime)/Datasmith/DirectLink/Public" "$(UE_Runtime)/Launch/Resources" "$(UE_Engine)/Source/Developer/Datasmith/DatasmithExporter/Public" "$(UE_Engine)/Source/Developer/Datasmith/DatasmithExporterUI/Public" "$(UE_Engine)/Source/Runtime/Messaging/Public" "$(UE_Engine)/Source/Runtime/CoreUObject/Public/" "$(UE_ArchiCADExporter)/DatasmithValidator/Public"
+UE_INCLUDES_PATH = "$(UE_Core)" "$(UE_Core)/Internationalization" "$(UE_Core)/Async" "$(UE_Core)/Concurrency" "$(UE_Core)/Containers" "$(UE_Core)/Delegates" "$(UE_Core)/GenericPlatform" "$(UE_Core)/HAL" "$(UE_Core)/Logging" "$(UE_Core)/Math" "$(UE_Core)/Misc" "$(UE_Core)/Modules" "$(UE_Core)/Modules\Boilerplate" "$(UE_Core)/ProfilingDebugging" "$(UE_Core)/Serialization" "$(UE_Core)/Serialization\Csv" "$(UE_Core)/Stats" "$(UE_Core)/Templates" "$(UE_Core)/UObject" "$(UE_Project)/Public" "$(UE_Project)/Interfaces" "$(UE_Runtime)/TraceLog/Public"  "$(UE_Runtime)/Datasmith/DatasmithCore/Public" "$(UE_Runtime)/Datasmith/DirectLink/Public" "$(UE_Runtime)/Launch/Resources" "$(UE_Engine)/Source/Developer/Datasmith/DatasmithExporter/Public" "$(UE_Engine)/Source/Developer/Datasmith/DatasmithExporterUI/Public" "$(UE_Engine)/Source/Runtime/Messaging/Public" "$(UE_Engine)/Source/Runtime/CoreUObject/Public/" "$(UE_ArchiCADExporter)/DatasmithValidator/Public" "$(UE_ThirdParty)/GuidelinesSupportLibrary/GSL-1144/include"
 
 USER_HEADER_SEARCH_PATHS = "." $(AC_INCLUDES_PATH) $(UE_INCLUDES_PATH)
 
diff --git a/Engine/Source/Programs/Enterprise/Datasmith/DatasmithARCHICADExporter/Build/Config/UE_AC_ARM.xcconfig b/Engine/Source/Programs/Enterprise/Datasmith/DatasmithARCHICADExporter/Build/Config/UE_AC_ARM.xcconfig
--- a/Engine/Source/Programs/Enterprise/Datasmith/DatasmithARCHICADExporter/Build/Config/UE_AC_ARM.xcconfig
+++ b/Engine/Source/Programs/Enterprise/Datasmith/DatasmithARCHICADExporter/Build/Config/UE_AC_ARM.xcconfig
@@ -9,9 +9,10 @@ AC_INCLUDES_PATH = "$(AC_SDK_PATH)/Support/Inc" "$(AC_SDK_PATH)/Support/Modules/
 UE_Runtime = $(UE_Engine)/Source/Runtime
 UE_Core = $(UE_Runtime)/Core/Public
 UE_Project = $(UE_Runtime)/Projects/Public
+UE_ThirdParty = $(UE_Engine)/Source/ThirdParty
 UE_ArchiCADExporter = $(UE_Engine)/Source/Programs/Enterprise/Datasmith/DatasmithARCHICADExporter
 
-UE_INCLUDES_PATH = "$(UE_Core)" "$(UE_Core)/Internationalization" "$(UE_Core)/Async" "$(UE_Core)/Concurrency" "$(UE_Core)/Containers" "$(UE_Core)/Delegates" "$(UE_Core)/GenericPlatform" "$(UE_Core)/HAL" "$(UE_Core)/Logging" "$(UE_Core)/Math" "$(UE_Core)/Misc" "$(UE_Core)/Modules" "$(UE_Core)/Modules\Boilerplate" "$(UE_Core)/ProfilingDebugging" "$(UE_Core)/Serialization" "$(UE_Core)/Serialization\Csv" "$(UE_Core)/Stats" "$(UE_Core)/Templates" "$(UE_Core)/UObject" "$(UE_Project)/Public" "$(UE_Project)/Interfaces" "$(UE_Runtime)/TraceLog/Public"  "$(UE_Runtime)/Datasmith/DatasmithCore/Public" "$(UE_Runtime)/Datasmith/DirectLink/Public" "$(UE_Runtime)/Launch/Resources" "$(UE_Engine)/Source/Developer/Datasmith/DatasmithExporter/Public" "$(UE_Engine)/Source/Developer/Datasmith/DatasmithExporterUI/Public" "$(UE_Engine)/Source/Runtime/Messaging/Public" "$(UE_Engine)/Source/Runtime/CoreUObject/Public/" "$(UE_ArchiCADExporter)/DatasmithValidator/Public"
+UE_INCLUDES_PATH = "$(UE_Core)" "$(UE_Core)/Internationalization" "$(UE_Core)/Async" "$(UE_Core)/Concurrency" "$(UE_Core)/Containers" "$(UE_Core)/Delegates" "$(UE_Core)/GenericPlatform" "$(UE_Core)/HAL" "$(UE_Core)/Logging" "$(UE_Core)/Math" "$(UE_Core)/Misc" "$(UE_Core)/Modules" "$(UE_Core)/Modules\Boilerplate" "$(UE_Core)/ProfilingDebugging" "$(UE_Core)/Serialization" "$(UE_Core)/Serialization\Csv" "$(UE_Core)/Stats" "$(UE_Core)/Templates" "$(UE_Core)/UObject" "$(UE_Project)/Public" "$(UE_Project)/Interfaces" "$(UE_Runtime)/TraceLog/Public"  "$(UE_Runtime)/Datasmith/DatasmithCore/Public" "$(UE_Runtime)/Datasmith/DirectLink/Public" "$(UE_Runtime)/Launch/Resources" "$(UE_Engine)/Source/Developer/Datasmith/DatasmithExporter/Public" "$(UE_Engine)/Source/Developer/Datasmith/DatasmithExporterUI/Public" "$(UE_Engine)/Source/Runtime/Messaging/Public" "$(UE_Engine)/Source/Runtime/CoreUObject/Public/" "$(UE_ArchiCADExporter)/DatasmithValidator/Public" "$(UE_ThirdParty)/GuidelinesSupportLibrary/GSL-1144/include"
 
 USER_HEADER_SEARCH_PATHS = "." $(AC_INCLUDES_PATH) $(UE_INCLUDES_PATH)
 
diff --git a/Engine/Source/Programs/Enterprise/Datasmith/DatasmithARCHICADExporter/Private/Utils/APIEnvir.h b/Engine/Source/Programs/Enterprise/Datasmith/DatasmithARCHICADExporter/Private/Utils/APIEnvir.h
--- a/Engine/Source/Programs/Enterprise/Datasmith/DatasmithARCHICADExporter/Private/Utils/APIEnvir.h
+++ b/Engine/Source/Programs/Enterprise/Datasmith/DatasmithARCHICADExporter/Private/Utils/APIEnvir.h
@@ -3,7 +3,9 @@
 #pragma once
 
 #include <stddef.h>
+#if PLATFORM_WINDOWS
 #include <strsafe.h>
+#endif
 #include <string>
 
 #include "IDatasmithSceneElements.h"
@@ -20,6 +22,11 @@ THIRD_PARTY_INCLUDES_START
 #pragma clang diagnostic push
 #pragma clang diagnostic ignored "-Wcomma"
 #pragma clang diagnostic ignored "-Wdefaulted-function-deleted"
+#pragma clang diagnostic ignored "-Wdeprecated-volatile"
+#pragma clang diagnostic ignored "-Wdeprecated-declarations"
+#pragma clang diagnostic ignored "-Wenum-compare"
+#pragma clang diagnostic ignored "-Wenum-compare-conditional"
+#pragma clang diagnostic ignored "-Wdeprecated-anon-enum-enum-conversion"
 
 #if (__clang_major__ > 12) || (__clang_major__ == 12 && __clang_minor__ == 0 && __clang_patchlevel__ > 4)
 #pragma clang diagnostic ignored "-Wnon-c-typedef-for-linkage"
diff --git a/Engine/Source/Programs/Enterprise/Datasmith/DatasmithARCHICADExporter/Private/Utils/ShellOpenDocument.cpp b/Engine/Source/Programs/Enterprise/Datasmith/DatasmithARCHICADExporter/Private/Utils/ShellOpenDocument.cpp
--- a/Engine/Source/Programs/Enterprise/Datasmith/DatasmithARCHICADExporter/Private/Utils/ShellOpenDocument.cpp
+++ b/Engine/Source/Programs/Enterprise/Datasmith/DatasmithARCHICADExporter/Private/Utils/ShellOpenDocument.cpp
@@ -4,6 +4,7 @@
 	#include <windows.h>
 #else
 	#include <unistd.h>
+	#include <stdio.h>
 	#include <stdlib.h>
 #endif
 
diff --git a/Engine/Source/Programs/Enterprise/Datasmith/DatasmithARCHICADExporter/Private/Utils/TaskMgr.h b/Engine/Source/Programs/Enterprise/Datasmith/DatasmithARCHICADExporter/Private/Utils/TaskMgr.h
--- a/Engine/Source/Programs/Enterprise/Datasmith/DatasmithARCHICADExporter/Private/Utils/TaskMgr.h
+++ b/Engine/Source/Programs/Enterprise/Datasmith/DatasmithARCHICADExporter/Private/Utils/TaskMgr.h
@@ -78,7 +78,7 @@ class FTaskMgr
 	volatile bool bTerminate;
 
 	// Number of thread waiting
-	volatile int NbTaskWaiting;
+	int NbTaskWaiting;
 
 	// False: tasks to run in main threads, true: Task run in thread
 	bool bThreadingEnabled;
