Preemptive fixes for post-compilation scripts
A rare successful compilation of the ARCHICAD exporter on Mac, https://horde.devtools.epicgames.com/job/67128402a1204f7d4332c16b, showed that the post-compilation scripts are broken.
The changes are to fix those failures.

#rnx

diff --git a/Engine/Source/Programs/Enterprise/Datasmith/DatasmithARCHICADExporter/Build/Config/app_to_bundle.sh b/Engine/Source/Programs/Enterprise/Datasmith/DatasmithARCHICADExporter/Build/Config/app_to_bundle.sh
--- a/Engine/Source/Programs/Enterprise/Datasmith/DatasmithARCHICADExporter/Build/Config/app_to_bundle.sh
+++ b/Engine/Source/Programs/Enterprise/Datasmith/DatasmithARCHICADExporter/Build/Config/app_to_bundle.sh
@@ -1,10 +1,11 @@
 #!/bin/sh
 
 ConfigPath=`dirname "$0"`
-projectPath=$ConfigPath/..
-
-RelativeEnginePath=$projectPath/../../../../../../../Engine
-EnginePath=`python -c "import os; print(os.path.realpath('$RelativeEnginePath'))"`
+pushd $ConfigPath/..
+projectPath=`pwd`
+cd ../../../../../..
+EnginePath=`pwd`
+popd
 
 echo "Renaming *.app to *.bundle"
 pushd $EnginePath/Binaries/Mac/DatasmithARCHICADExporter
diff --git a/Engine/Source/Programs/Enterprise/Datasmith/DatasmithARCHICADExporter/Build/Config/bundle_to_app.sh b/Engine/Source/Programs/Enterprise/Datasmith/DatasmithARCHICADExporter/Build/Config/bundle_to_app.sh
--- a/Engine/Source/Programs/Enterprise/Datasmith/DatasmithARCHICADExporter/Build/Config/bundle_to_app.sh
+++ b/Engine/Source/Programs/Enterprise/Datasmith/DatasmithARCHICADExporter/Build/Config/bundle_to_app.sh
@@ -1,10 +1,12 @@
 #!/bin/sh
 
-ConfigPath=`dirname "$0"`
-projectPath=$ConfigPath/..
 
-RelativeEnginePath=$projectPath/../../../../../../../Engine
-EnginePath=`python -c "import os; print(os.path.realpath('$RelativeEnginePath'))"`
+ConfigPath=`dirname "$0"`
+pushd $ConfigPath/..
+projectPath=`pwd`
+cd ../../../../../..
+EnginePath=`pwd`
+popd
 
 echo "Renaming *.bundle to *.app"
 pushd $EnginePath/Binaries/Mac/DatasmithARCHICADExporter
