Datasmith Archicad plugin : Add support of Archicad 28
#rnx

diff --git a/Engine/Source/Programs/Enterprise/Datasmith/DatasmithARCHICADExporter/Build/Config/Notes.txt b/Engine/Source/Programs/Enterprise/Datasmith/DatasmithARCHICADExporter/Build/Config/Notes.txt
--- a/Engine/Source/Programs/Enterprise/Datasmith/DatasmithARCHICADExporter/Build/Config/Notes.txt
+++ b/Engine/Source/Programs/Enterprise/Datasmith/DatasmithARCHICADExporter/Build/Config/Notes.txt
@@ -11,6 +11,7 @@ To build you need to define some user environment variables
 			UE_SDKs_Root/HostWin64/Win64/Archicad/25/Support/...
 			UE_SDKs_Root/HostWin64/Win64/Archicad/26/Support/...
 			UE_SDKs_Root/HostWin64/Win64/Archicad/27/Support/...
+			UE_SDKs_Root/HostWin64/Win64/Archicad/28/Support/...
 	UE_DatasmithArchicadDeveloperId='Organization (UserId)'
 		Released by Apple "Developer ID Certification Authority"
 			Must be valid for Application and for Installer
diff --git a/Engine/Source/Programs/Enterprise/Datasmith/DatasmithARCHICADExporter/Build/Config/UE_AC.xcconfig b/Engine/Source/Programs/Enterprise/Datasmith/DatasmithARCHICADExporter/Build/Config/UE_AC.xcconfig
--- a/Engine/Source/Programs/Enterprise/Datasmith/DatasmithARCHICADExporter/Build/Config/UE_AC.xcconfig
+++ b/Engine/Source/Programs/Enterprise/Datasmith/DatasmithARCHICADExporter/Build/Config/UE_AC.xcconfig
@@ -24,7 +24,7 @@ VRevision = 0
 VBetaOrRevision = b // b for beta or . for revision
 VAddOn = $(VMajor).$(VMinor)$(VBetaOrRevision)$(VRevision)
 
-CurrentYear = 2022
+CurrentYear = 2024
 
 FRAMEWORK_SEARCH_PATHS = $(FRAMEWORK_SEARCH_PATHS) $(AC_SDK_PATH)/Support/Frameworks
 
diff --git a/Engine/Source/Programs/Enterprise/Datasmith/DatasmithARCHICADExporter/Build/Config/UE_AC_ARM.xcconfig b/Engine/Source/Programs/Enterprise/Datasmith/DatasmithARCHICADExporter/Build/Config/UE_AC_ARM.xcconfig
--- a/Engine/Source/Programs/Enterprise/Datasmith/DatasmithARCHICADExporter/Build/Config/UE_AC_ARM.xcconfig
+++ b/Engine/Source/Programs/Enterprise/Datasmith/DatasmithARCHICADExporter/Build/Config/UE_AC_ARM.xcconfig
@@ -24,11 +24,11 @@ VRevision = 0
 VBetaOrRevision = b // b for beta or . for revision
 VAddOn = $(VMajor).$(VMinor)$(VBetaOrRevision)$(VRevision)
 
-CurrentYear = 2022
+CurrentYear = 2024
 
 FRAMEWORK_SEARCH_PATHS = $(FRAMEWORK_SEARCH_PATHS) $(AC_SDK_PATH)/Support/Frameworks
 
-LIBRARY_SEARCH_PATHS = "Dylibs" "$(AC_SDK_PATH)/Support/Lib/Mactel" "$(AC_SDK_PATH)/Support/Lib"
+LIBRARY_SEARCH_PATHS = "Dylibs" "$(AC_SDK_PATH)/Support/Lib"
 
 REZ_SEARCH_PATHS = RO "$(AC_SDK_PATH)/Support/Inc" "$(AC_SDK_PATH)/Support/Modules/DGLib"
 
diff --git a/Engine/Source/Programs/Enterprise/Datasmith/DatasmithARCHICADExporter/Build/Config/app_to_bundle.sh b/Engine/Source/Programs/Enterprise/Datasmith/DatasmithARCHICADExporter/Build/Config/app_to_bundle.sh
--- a/Engine/Source/Programs/Enterprise/Datasmith/DatasmithARCHICADExporter/Build/Config/app_to_bundle.sh
+++ b/Engine/Source/Programs/Enterprise/Datasmith/DatasmithARCHICADExporter/Build/Config/app_to_bundle.sh
@@ -14,4 +14,5 @@ mv DatasmithARCHICAD24Exporter.app DatasmithARCHICAD24Exporter.bundle
 mv DatasmithARCHICAD25Exporter.app DatasmithARCHICAD25Exporter.bundle
 mv DatasmithARCHICAD26Exporter.app DatasmithARCHICAD26Exporter.bundle
 mv DatasmithARCHICAD27Exporter.app DatasmithARCHICAD27Exporter.bundle
+mv DatasmithARCHICAD28Exporter.app DatasmithARCHICAD28Exporter.bundle
 popd
diff --git a/Engine/Source/Programs/Enterprise/Datasmith/DatasmithARCHICADExporter/Build/Config/bundle_to_app.sh b/Engine/Source/Programs/Enterprise/Datasmith/DatasmithARCHICADExporter/Build/Config/bundle_to_app.sh
--- a/Engine/Source/Programs/Enterprise/Datasmith/DatasmithARCHICADExporter/Build/Config/bundle_to_app.sh
+++ b/Engine/Source/Programs/Enterprise/Datasmith/DatasmithARCHICADExporter/Build/Config/bundle_to_app.sh
@@ -15,4 +15,5 @@ mv DatasmithARCHICAD24Exporter.bundle DatasmithARCHICAD24Exporter.app
 mv DatasmithARCHICAD25Exporter.bundle DatasmithARCHICAD25Exporter.app
 mv DatasmithARCHICAD26Exporter.bundle DatasmithARCHICAD26Exporter.app
 mv DatasmithARCHICAD27Exporter.bundle DatasmithARCHICAD27Exporter.app
+mv DatasmithARCHICAD28Exporter.bundle DatasmithARCHICAD28Exporter.app
 popd
diff --git a/Engine/Source/Programs/Enterprise/Datasmith/DatasmithARCHICADExporter/Build/Config/setupdylibs.sh b/Engine/Source/Programs/Enterprise/Datasmith/DatasmithARCHICADExporter/Build/Config/setupdylibs.sh
--- a/Engine/Source/Programs/Enterprise/Datasmith/DatasmithARCHICADExporter/Build/Config/setupdylibs.sh
+++ b/Engine/Source/Programs/Enterprise/Datasmith/DatasmithARCHICADExporter/Build/Config/setupdylibs.sh
@@ -38,6 +38,9 @@ pushd "$UE_SDKS_ROOT/HostMac/Mac/Archicad"
 
 	chmod 777 27/Support/Tools/OSX/ResConv
 	/usr/bin/xattr -r -d com.apple.quarantine 27/Support/Tools/OSX/ResConv
+
+	chmod 777 28/Support/Tools/OSX/ResConv
+	/usr/bin/xattr -r -d com.apple.quarantine 28/Support/Tools/OSX/ResConv
 popd
 
 OurDylibFolder=$projectPath/Dylibs
diff --git a/Engine/Source/Programs/Enterprise/Datasmith/DatasmithARCHICADExporter/Installer/CustomActionManaged/CustomAction.cs b/Engine/Source/Programs/Enterprise/Datasmith/DatasmithARCHICADExporter/Installer/CustomActionManaged/CustomAction.cs
--- a/Engine/Source/Programs/Enterprise/Datasmith/DatasmithARCHICADExporter/Installer/CustomActionManaged/CustomAction.cs
+++ b/Engine/Source/Programs/Enterprise/Datasmith/DatasmithARCHICADExporter/Installer/CustomActionManaged/CustomAction.cs
@@ -122,7 +122,7 @@ namespace CustomActionManaged
 		{
 			try
 			{
-				string[] MajorVersions = new string[] { "27", "26", "25", "24", "23", };
+				string[] MajorVersions = new string[] { "28", "27", "26", "25", "24", "23", };
 
 				string[] MinorVersions = new string[] { "0", };
 				
diff --git a/Engine/Source/Programs/Enterprise/Datasmith/DatasmithARCHICADExporter/Private/AddOn.cpp b/Engine/Source/Programs/Enterprise/Datasmith/DatasmithARCHICADExporter/Private/AddOn.cpp
--- a/Engine/Source/Programs/Enterprise/Datasmith/DatasmithARCHICADExporter/Private/AddOn.cpp
+++ b/Engine/Source/Programs/Enterprise/Datasmith/DatasmithARCHICADExporter/Private/AddOn.cpp
@@ -30,6 +30,10 @@
 
 using namespace UE_AC;
 
+#if !defined (__ACENV_CALL) && AC_VERSION > 27
+    #define __ACENV_CALL
+#endif
+
 // =============================================================================
 // Required functions
 // =============================================================================
diff --git a/Engine/Source/Programs/Enterprise/Datasmith/DatasmithARCHICADExporter/Private/ElementID.cpp b/Engine/Source/Programs/Enterprise/Datasmith/DatasmithARCHICADExporter/Private/ElementID.cpp
--- a/Engine/Source/Programs/Enterprise/Datasmith/DatasmithARCHICADExporter/Private/ElementID.cpp
+++ b/Engine/Source/Programs/Enterprise/Datasmith/DatasmithARCHICADExporter/Private/ElementID.cpp
@@ -58,8 +58,12 @@ FAssValueName::SAssValueName TAssEnumName< ModelerAPI::Element::Type >::AssEnumN
 	EnumName(ModelerAPI::Element, BalusterSetElement),
 #endif
 #if AC_VERSION >= 23
+
+	#if AC_VERSION < 28
 	EnumName(ModelerAPI::Element, AnalyticalSupportElement),
 	EnumName(ModelerAPI::Element, AnalyticalLinkElement),
+	#endif
+
 	EnumName(ModelerAPI::Element, Opening),
 	EnumName(ModelerAPI::Element, Openingframeinfill),
 	EnumName(ModelerAPI::Element, Openingpatchinfill),
diff --git a/Engine/Source/Programs/Enterprise/Datasmith/DatasmithARCHICADExporter/Private/Export.cpp b/Engine/Source/Programs/Enterprise/Datasmith/DatasmithARCHICADExporter/Private/Export.cpp
--- a/Engine/Source/Programs/Enterprise/Datasmith/DatasmithARCHICADExporter/Private/Export.cpp
+++ b/Engine/Source/Programs/Enterprise/Datasmith/DatasmithARCHICADExporter/Private/Export.cpp
@@ -17,7 +17,11 @@ enum : GSType
 
 const utf8_t* StrFileExtension = "udatasmith";
 
+#if AC_VERSION < 28
 static GSErrCode __ACENV_CALL SaveToDatasmithFile(const API_IOParams* IOParams, Modeler::SightPtr sight)
+#else
+static GSErrCode SaveToDatasmithFile(const API_IOParams* IOParams, Modeler::SightPtr sight)
+#endif
 {
 	GSErrCode GSErr = TryFunctionCatchAndAlert("FExport::SaveDatasmithFile", [&IOParams, &sight]() -> GSErrCode {
 		return FExport::SaveDatasmithFile(*IOParams, *sight.GetPtr());
diff --git a/Engine/Source/Programs/Enterprise/Datasmith/DatasmithARCHICADExporter/Private/Menus.cpp b/Engine/Source/Programs/Enterprise/Datasmith/DatasmithARCHICADExporter/Private/Menus.cpp
--- a/Engine/Source/Programs/Enterprise/Datasmith/DatasmithARCHICADExporter/Private/Menus.cpp
+++ b/Engine/Source/Programs/Enterprise/Datasmith/DatasmithARCHICADExporter/Private/Menus.cpp
@@ -52,15 +52,23 @@ GSErrCode FMenus::Initialize()
 	for (short IndexMenu = kStrListMenuItemSnapshot; IndexMenu <= kStrListMenuItemAbout && GSErr == NoError;
 		 IndexMenu++)
 	{
+		#if AC_VERSION < 28
 		GSErr = ACAPI_Install_MenuHandler(LocalizeResId(IndexMenu), MenuCommandHandler);
+		#else
+		GSErr = ACAPI_MenuItem_InstallMenuHandler(LocalizeResId(IndexMenu), MenuCommandHandler);
+		#endif
 	}
 	if (GSErr != NoError)
 	{
 		UE_AC_DebugF("FMenus::Initialize - ACAPI_Install_MenuHandler error=%s\n", GetErrorName(GSErr));
 	}
 #endif
-
+	
+	#if AC_VERSION < 28
 	GSErr = ACAPI_Install_MenuHandler(LocalizeResId(kStrListMenuDatasmith), MenuCommandHandler);
+	#else
+	GSErr = ACAPI_MenuItem_InstallMenuHandler(LocalizeResId(kStrListMenuDatasmith), MenuCommandHandler);
+	#endif
 	if (GSErr != NoError)
 	{
 		UE_AC_DebugF("FMenus::Initialize - ACAPI_Install_MenuHandler error=%s\n", GetErrorName(GSErr));
@@ -140,7 +148,11 @@ void FMenus::AutoSyncChanged()
 }
 
 // Menu command handler
+#if AC_VERSION < 28
 GSErrCode __ACENV_CALL FMenus::MenuCommandHandler(const API_MenuParams* MenuParams) noexcept
+#else
+GSErrCode FMenus::MenuCommandHandler(const API_MenuParams* MenuParams) noexcept
+#endif
 {
 	return TryFunctionCatchAndAlert("FMenus::DoMenuCommand",
 									[&MenuParams]() -> GSErrCode { return FMenus::DoMenuCommand(*MenuParams); });
diff --git a/Engine/Source/Programs/Enterprise/Datasmith/DatasmithARCHICADExporter/Private/Menus.h b/Engine/Source/Programs/Enterprise/Datasmith/DatasmithARCHICADExporter/Private/Menus.h
--- a/Engine/Source/Programs/Enterprise/Datasmith/DatasmithARCHICADExporter/Private/Menus.h
+++ b/Engine/Source/Programs/Enterprise/Datasmith/DatasmithARCHICADExporter/Private/Menus.h
@@ -28,7 +28,11 @@ class FMenus
 
   private:
 	// Menu command handler
+	#if AC_VERSION < 28
 	static GSErrCode __ACENV_CALL MenuCommandHandler(const API_MenuParams* MenuParams) noexcept;
+	#else
+	static GSErrCode MenuCommandHandler(const API_MenuParams* MenuParams) noexcept;
+	#endif
 
 	// Process menu command
 	static GSErrCode DoMenuCommand(const API_MenuParams& InMenuParams);
diff --git a/Engine/Source/Programs/Enterprise/Datasmith/DatasmithARCHICADExporter/Private/MetaData.cpp b/Engine/Source/Programs/Enterprise/Datasmith/DatasmithARCHICADExporter/Private/MetaData.cpp
--- a/Engine/Source/Programs/Enterprise/Datasmith/DatasmithARCHICADExporter/Private/MetaData.cpp
+++ b/Engine/Source/Programs/Enterprise/Datasmith/DatasmithARCHICADExporter/Private/MetaData.cpp
@@ -2,9 +2,15 @@
 
 #include "MetaData.h"
 #include "Utils/ElementTools.h"
+#include "Utils/Element2String.h"
 #include "Version.h"
 #include "SyncData.h"
 
+#if AC_VERSION > 27
+#include "ACAPI/IFCObjectAccessor.hpp"
+#include "ACAPI/IFCPropertyAccessor.hpp"
+#endif
+
 BEGIN_NAMESPACE_UE_AC
 
 static const GS::UniString StringTrue("True");
@@ -127,8 +133,13 @@ void FMetaData::AddMetaDataProperty(API_VariantType variantType, const GS::UniSt
 {
 	static const GS::UniString StringUndefinedLocalized(GetGSName(kName_Undefined));
 
+	#if AC_VERSION > 27
+	if (PropertyValue.IsEmpty() || PropertyValue.IsEqual(StringUndefined, GS::CaseInsensitive) ||
+		PropertyValue.IsEqual(StringUndefinedLocalized, GS::CaseInsensitive))
+	#else
 	if (PropertyValue.IsEmpty() || PropertyValue.IsEqual(StringUndefined, GS::UniString::CaseInsensitive) ||
 		PropertyValue.IsEqual(StringUndefinedLocalized, GS::UniString::CaseInsensitive))
+	#endif
 		return;
 
 	EDatasmithKeyValuePropertyType PropertyValueType; // = EDatasmithKeyValuePropertyType::String;
@@ -232,6 +243,7 @@ void FMetaData::ExportCategories(const API_Guid& ElementId)
 
 void FMetaData::ExportIFCType(const API_Guid& ElementId)
 {
+	#if AC_VERSION < 28
 	GS::UniString IfcType;
 	GS::UniString TypeObjectIfcType;
 	GSErrCode	  GSErr = ACAPI_Element_GetIFCType(ElementId, &IfcType, &TypeObjectIfcType);
@@ -252,6 +264,35 @@ void FMetaData::ExportIFCType(const API_Guid& ElementId)
 	{
 		UE_AC_DebugF("FMetaData::ExportIFCType - ACAPI_Element_GetIFCType returned error %s", GetErrorName(GSErr));
 	}
+	#else
+	API_Elem_Head Head{};
+	Head.guid = ElementId;
+	auto ElementObjectID = IFCAPI::GetObjectAccessor().CreateElementObjectID(Head);
+	if (ElementObjectID.IsOk())
+	{
+		auto IfcType = IFCAPI::GetObjectAccessor().GetIFCType(*ElementObjectID);
+		if (IfcType.IsOk())
+		{
+			static const GS::UniString KeyIfcType("IFC_Type");
+			AddMetaDataProperty(API_PropertyStringValueType, KeyIfcType, *IfcType);
+		}
+		else
+		{
+			UE_AC_DebugF("FMetaData::ExportIFCType - IFCAPI::GetObjectAccessor().GetIFCType returned error %s", GetErrorName(IfcType.UnwrapErr().kind));
+		}
+
+		auto TypeObjectIfcType = IFCAPI::GetObjectAccessor().GetTypeObjectIFCType(*ElementObjectID);
+		if (TypeObjectIfcType.IsOk())
+		{
+			static const GS::UniString KeyIfcType("IFC_ObjectType");
+			AddMetaDataProperty(API_PropertyStringValueType, KeyIfcType, *TypeObjectIfcType);
+		}
+		else
+		{
+			UE_AC_DebugF("FMetaData::ExportIFCType - IFCAPI::GetObjectAccessor().GetTypeObjectIFCType returned error %s", GetErrorName(TypeObjectIfcType.UnwrapErr().kind));
+		}
+	}
+	#endif
 }
 
 void FMetaData::ExportIFCProperties(const API_Guid& ElementId)
@@ -259,8 +300,11 @@ void FMetaData::ExportIFCProperties(const API_Guid& ElementId)
 	static const GS::UniString StringIFC("IFC_");
 	static const GS::UniString StringLower("_lower");
 	static const GS::UniString StringUpper("_upper");
+	static const GS::UniString SetPoint("_setpoint");
 
 	GS::Array< API_IFCProperty > IFCProperties;
+
+	#if AC_VERSION < 28
 	GSErrCode					 GSErr = ACAPI_Element_GetIFCProperties(ElementId, false, &IFCProperties);
 	if (GSErr == NoError)
 	{
@@ -326,12 +370,96 @@ void FMetaData::ExportIFCProperties(const API_Guid& ElementId)
 						 GetErrorName(GSErr));
 		}
 	}
+	#else
+	API_Elem_Head Head{};
+	Head.guid = ElementId;
+	const auto ElementObjectID = IFCAPI::GetObjectAccessor().CreateElementObjectID(Head);
+	if (ElementObjectID.IsOk())
+	{
+		const auto LocalProperties = IFCAPI::PropertyAccessor(*ElementObjectID).GetLocalProperties();
+		if (LocalProperties.IsOk())
+		{
+			for (const auto& Property : *LocalProperties)
+			{
+				auto PropertyVariant = Property.GetTyped();
+				GS::UniString KeyName = StringIFC + Property.GetName();
+				GS::UniString ValueName;
+
+				if (std::holds_alternative<IFCAPI::v1::PropertySingleValue>(PropertyVariant))
+				{
+					const auto& PropertyCast = std::get<IFCAPI::v1::PropertySingleValue>(PropertyVariant);
+					ValueName = FElement2String::PropertyValueToString(PropertyCast.GetNominalValue());
+					AddMetaDataProperty(API_PropertyStringValueType, KeyName, ValueName);
+					
+				}
+				else if (std::holds_alternative<IFCAPI::v1::PropertyListValue>(PropertyVariant))
+				{
+					const auto& PropertyCast = std::get<IFCAPI::v1::PropertyListValue>(PropertyVariant);
+					Int32 i = 0;
+					for (const auto& value : PropertyCast.GetListValues())
+					{
+						ValueName = FElement2String::PropertyValueToString(value);
+						AddMetaDataProperty(API_PropertyStringValueType, KeyName + '_' + GS::ValueToUniString(++i), ValueName);
+					}
+				}
+				else if (std::holds_alternative<IFCAPI::v1::PropertyBoundedValue>(PropertyVariant))
+				{
+					const auto& PropertyCast = std::get<IFCAPI::v1::PropertyBoundedValue>(PropertyVariant);
+
+					ValueName = FElement2String::PropertyValueToString(PropertyCast.GetLowerBoundValue());
+					AddMetaDataProperty(API_PropertyStringValueType, KeyName + "_" + StringLower, ValueName);
+
+					ValueName = FElement2String::PropertyValueToString(PropertyCast.GetUpperBoundValue());
+					AddMetaDataProperty(API_PropertyStringValueType, KeyName + "_" + StringUpper, ValueName);
+
+					ValueName = FElement2String::PropertyValueToString(PropertyCast.GetSetPointValue());
+					AddMetaDataProperty(API_PropertyStringValueType, KeyName + "_" + SetPoint, ValueName);
+
+				}
+				else if (std::holds_alternative<IFCAPI::v1::PropertyEnumeratedValue>(PropertyVariant))
+				{
+					const auto& PropertyCast = std::get<IFCAPI::v1::PropertyEnumeratedValue>(PropertyVariant);
+					Int32 i = 0;
+					for (const auto& value : PropertyCast.GetEnumerationValues())
+					{
+						ValueName = FElement2String::PropertyValueToString(value);
+						AddMetaDataProperty(API_PropertyStringValueType, KeyName + '_' + GS::ValueToUniString(++i), ValueName);
+					}
+					
+				}
+				else
+				{
+					const auto& PropertyCast = std::get<IFCAPI::v1::PropertyTableValue>(PropertyVariant);
+					auto DefiningValues = PropertyCast.GetDefiningValues();
+					auto DefinedValues = PropertyCast.GetDefinedValues();
+
+					for (size_t i = 0; i < DefiningValues.size(); ++i)
+					{
+						ValueName =  FElement2String::PropertyValueToString(DefinedValues[i]);
+						AddMetaDataProperty(API_PropertyStringValueType, KeyName + '_' + FElement2String::PropertyValueToString(DefiningValues[i]), ValueName);
+					}
+				}
+			}
+		}
+		else if (LocalProperties.UnwrapErr().kind != APIERR_BADPARS)
+        {
+            UE_AC_DebugF("FElement2String::GetIFCPropertiesAsString - IFCAPI::PropertyAccessor().GetLocalProperties returned error %s\n",
+                GetErrorName(LocalProperties.UnwrapErr().kind));
+        }
+	}
+	else if (ElementObjectID.UnwrapErr().kind != APIERR_BADPARS)
+    {
+        UE_AC_DebugF("FElement2String::GetIFCPropertiesAsString - IFCAPI::GetObjectAccessor().CreateElementObjectID returned error %s\n",
+            GetErrorName(ElementObjectID.UnwrapErr().kind));
+    }
+	#endif
 }
 
 void FMetaData::ExportIFCAttributes(const API_Guid& ElementId)
 {
 	static const GS::UniString StringIFCAttribute("IFC_Attribute_");
 
+	#if AC_VERSION < 28
 	GS::Array< API_IFCAttribute > IFCAttributes;
 	GSErrCode					  GSErr = ACAPI_Element_GetIFCAttributes(ElementId, false, &IFCAttributes);
 	if (GSErr == NoError)
@@ -350,6 +478,36 @@ void FMetaData::ExportIFCAttributes(const API_Guid& ElementId)
 						 GetErrorName(GSErr));
 		}
 	}
+	#else
+	API_Elem_Head Head{};
+	Head.guid = ElementId;
+	auto ElementObjectID = IFCAPI::GetObjectAccessor().CreateElementObjectID(Head);
+	if (ElementObjectID.IsOk())
+	{
+		auto IFCAttributes = IFCAPI::PropertyAccessor(*ElementObjectID).GetAttributes();
+		if (IFCAttributes.IsOk())
+		{
+			for (const auto& IfcAttribute : *IFCAttributes)
+			{
+				if (IfcAttribute.GetValue().has_value())
+				{
+					AddMetaDataProperty(API_PropertyStringValueType, StringIFCAttribute + IfcAttribute.GetName(),
+						IfcAttribute.GetValue().value());
+				} 
+			}
+		}
+		else if (IFCAttributes.UnwrapErr().kind != APIERR_BADPARS)
+        {
+            UE_AC_DebugF("FMetaData::ExportIFCAttributes -  IFCAPI::PropertyAccessor().GetAttributes() returned error %s\n",
+                GetErrorName(IFCAttributes.UnwrapErr().kind));
+        }
+	}
+	else if (ElementObjectID.UnwrapErr().kind != APIERR_BADPARS)
+    {
+        UE_AC_DebugF("FMetaData::ExportIFCAttributes - IFCAPI::GetObjectAccessor().CreateElementObjectID returned error %s\n",
+            GetErrorName(ElementObjectID.UnwrapErr().kind));
+    }
+	#endif
 }
 
 void FMetaData::ExportProperties(const API_Guid& ElementId)
diff --git a/Engine/Source/Programs/Enterprise/Datasmith/DatasmithARCHICADExporter/Private/Palette.cpp b/Engine/Source/Programs/Enterprise/Datasmith/DatasmithARCHICADExporter/Private/Palette.cpp
--- a/Engine/Source/Programs/Enterprise/Datasmith/DatasmithARCHICADExporter/Private/Palette.cpp
+++ b/Engine/Source/Programs/Enterprise/Datasmith/DatasmithARCHICADExporter/Private/Palette.cpp
@@ -417,8 +417,13 @@ FAssValueName::SAssValueName TAssEnumName< API_PaletteMessageID >::AssEnumName[]
 };
 // clang-format on
 
+#if AC_VERSION < 28
 GSErrCode __ACENV_CALL FPalette::APIPaletteControlCallBack(Int32 ReferenceID, API_PaletteMessageID MessageID,
 														   GS::IntPtr Param)
+#else
+GSErrCode FPalette::APIPaletteControlCallBack(Int32 ReferenceID, API_PaletteMessageID MessageID,
+	GS::IntPtr Param)
+#endif
 {
 	GSErrCode GSErr = APIERR_GENERAL;
 	try
diff --git a/Engine/Source/Programs/Enterprise/Datasmith/DatasmithARCHICADExporter/Private/Palette.h b/Engine/Source/Programs/Enterprise/Datasmith/DatasmithARCHICADExporter/Private/Palette.h
--- a/Engine/Source/Programs/Enterprise/Datasmith/DatasmithARCHICADExporter/Private/Palette.h
+++ b/Engine/Source/Programs/Enterprise/Datasmith/DatasmithARCHICADExporter/Private/Palette.h
@@ -63,8 +63,13 @@ class FPalette
 
 	static void SetPaletteMenuTexts(bool PaletteIsOn, bool PaletteIsVisible);
 
+#	if AC_VERSION < 28
 	static GSErrCode __ACENV_CALL APIPaletteControlCallBack(Int32 referenceID, API_PaletteMessageID messageID,
-															GS::IntPtr /*param*/);
+		GS::IntPtr /*param*/);
+	#else
+	static GSErrCode APIPaletteControlCallBack(Int32 referenceID, API_PaletteMessageID messageID,
+		GS::IntPtr /*param*/);
+	#endif
 
 	static short DGCALLBACK CntlDlgCallBack(short message, short dialID, short item, DGUserData userData,
 											DGMessageData MsgData);
diff --git a/Engine/Source/Programs/Enterprise/Datasmith/DatasmithARCHICADExporter/Private/SyncDatabase.cpp b/Engine/Source/Programs/Enterprise/Datasmith/DatasmithARCHICADExporter/Private/SyncDatabase.cpp
--- a/Engine/Source/Programs/Enterprise/Datasmith/DatasmithARCHICADExporter/Private/SyncDatabase.cpp
+++ b/Engine/Source/Programs/Enterprise/Datasmith/DatasmithARCHICADExporter/Private/SyncDatabase.cpp
@@ -217,7 +217,15 @@ void FMeshCacheIndexor::ReadFromFile()
 				MeshName.resize(MeshNameSize);
 				NoErrorCall(GSErr, Reader.ReadBin(const_cast< char* >(MeshName.c_str()), MeshNameSize));
 			}
-			NoErrorCall(GSErr, strlen(MeshName.c_str()) == MeshNameSize ? NoError : ErrIO);
+			
+			if (strlen(MeshName.c_str()) == MeshNameSize)
+			{
+				NoErrorCall(GSErr, NoError);
+			}
+			else
+			{
+				NoErrorCall(GSErr, ErrIO);
+			}
 
 			TUniquePtr< FMeshDimensions > Dimensions = MakeUnique< FMeshDimensions >();
 			NoErrorCall(GSErr, Reader.Read(Dimensions->Area));
@@ -226,7 +234,15 @@ void FMeshCacheIndexor::ReadFromFile()
 			NoErrorCall(GSErr, Reader.Read(Dimensions->Width));
 			int32 Mark = 0;
 			NoErrorCall(GSErr, Reader.Read(Mark));
-			NoErrorCall(GSErr, Mark == 0x600DF00D ? NoError : ErrIO);
+
+			if (Mark == 0x600DF00D)
+			{
+				NoErrorCall(GSErr, NoError);
+			}
+			else
+			{
+				NoErrorCall(GSErr, ErrIO);
+			}
 
 			if (GSErr == NoError)
 			{
@@ -594,7 +610,11 @@ void FSyncDatabase::ResetMeshClasses()
 {
 	for (auto& IterMeshClass : MeshClasses)
 	{
+		#if AC_VERSION < 28
 		FMeshClass& MeshClass = **IterMeshClass.value;
+		#else
+		FMeshClass& MeshClass = *IterMeshClass.value.Get();
+		#endif
 		MeshClass.InstancesCount = 0;
 		MeshClass.TransformCount = 0;
 	}
@@ -605,7 +625,11 @@ void FSyncDatabase::CleanMeshClasses(const FSyncContext& InSyncContext)
 	int MeshClassesForgetCount = 0;
 	for (auto& IterMeshClass : MeshClasses)
 	{
+		#if AC_VERSION < 28
 		FMeshClass& MeshClass = **IterMeshClass.value;
+		#else
+		FMeshClass& MeshClass = *IterMeshClass.value.Get();
+		#endif
 		if (MeshClass.InstancesCount == 0 && MeshClass.ElementType != ModelerAPI::Element::Type::UndefinedElement)
 		{
 			MeshClass.ElementType = ModelerAPI::Element::Type::UndefinedElement;
diff --git a/Engine/Source/Programs/Enterprise/Datasmith/DatasmithARCHICADExporter/Private/Synchronizer.cpp b/Engine/Source/Programs/Enterprise/Datasmith/DatasmithARCHICADExporter/Private/Synchronizer.cpp
--- a/Engine/Source/Programs/Enterprise/Datasmith/DatasmithARCHICADExporter/Private/Synchronizer.cpp
+++ b/Engine/Source/Programs/Enterprise/Datasmith/DatasmithARCHICADExporter/Private/Synchronizer.cpp
@@ -133,8 +133,13 @@ GSErrCode FSynchronizer::Initialize()
 }
 
 // Intra add-ons command handler
+#if AC_VERSION < 28
 GSErrCode __ACENV_CALL FSynchronizer::SyncCommandHandler(GSHandle ParHdl, GSPtr /* ResultData */,
 														 bool /* SilentMod */) noexcept
+#else
+GSErrCode FSynchronizer::SyncCommandHandler(GSHandle ParHdl, GSPtr /* ResultData */,
+	bool /* SilentMod */) noexcept
+#endif
 {
 	return TryFunctionCatchAndAlert("FSynchronizer::DoSyncCommand",
 									[ParHdl]() -> GSErrCode { return FSynchronizer::DoSyncCommand(ParHdl); });
@@ -174,7 +179,11 @@ GSErrCode FSynchronizer::DoSyncCommand(GSHandle ParHdl)
 		UE_AC_DebugF("FSynchronizer::DoSyncCommand - APIAny_GetMDCLParameterID 1 error %s\n", GetErrorName(GSErr));
 		return GSErr;
 	}
+	#if AC_VERSION > 27
+	if (CHCompareCStrings(Param.name, "Reason", GS::CaseSensitive) != 0 || Param.type != MDCLPar_string)
+	#else
 	if (CHCompareCStrings(Param.name, "Reason", CS_CaseSensitive) != 0 || Param.type != MDCLPar_string)
+	#endif
 	{
 		UE_AC_DebugF("FSynchronizer::DoSyncCommand - Invalid parameters (type=%d) %s\n", Param.type, Param.name);
 		return APIERR_BADPARS;
diff --git a/Engine/Source/Programs/Enterprise/Datasmith/DatasmithARCHICADExporter/Private/Synchronizer.h b/Engine/Source/Programs/Enterprise/Datasmith/DatasmithARCHICADExporter/Private/Synchronizer.h
--- a/Engine/Source/Programs/Enterprise/Datasmith/DatasmithARCHICADExporter/Private/Synchronizer.h
+++ b/Engine/Source/Programs/Enterprise/Datasmith/DatasmithARCHICADExporter/Private/Synchronizer.h
@@ -55,7 +55,11 @@ class FSynchronizer
 	static GSErrCode Initialize();
 
 	// Intra add-ons command handler
+	#if AC_VERSION < 28
 	static GSErrCode __ACENV_CALL SyncCommandHandler(GSHandle ParHdl, GSPtr ResultData, bool SilentMode) noexcept;
+	#else
+	static GSErrCode SyncCommandHandler(GSHandle ParHdl, GSPtr ResultData, bool SilentMode) noexcept;
+	#endif
 
 	// Process intra add-ons command
 	static GSErrCode DoSyncCommand(GSHandle ParHdl);
diff --git a/Engine/Source/Programs/Enterprise/Datasmith/DatasmithARCHICADExporter/Private/Utils/AddonTools.cpp b/Engine/Source/Programs/Enterprise/Datasmith/DatasmithARCHICADExporter/Private/Utils/AddonTools.cpp
--- a/Engine/Source/Programs/Enterprise/Datasmith/DatasmithARCHICADExporter/Private/Utils/AddonTools.cpp
+++ b/Engine/Source/Programs/Enterprise/Datasmith/DatasmithARCHICADExporter/Private/Utils/AddonTools.cpp
@@ -10,17 +10,30 @@
 #include <cstdarg>
 #include <time.h>
 
-#if AC_VERSION > 26
+#if AC_VERSION > 26 
 class API_PenType;
+
+#if AC_VERSION < 28
 GSErrCode ACAPI_ProjectOperation_Quit() { return ACAPI_ProjectOperation_Quit(0); }
-GSErrCode ACAPI_GraphicalOverride_GetOverrideRule (API_OverrideRule& rule) { return GS::NoError; }
-GSErrCode ACAPI_GraphicalOverride_CreateOverrideRule (API_OverrideRule& rule) { const API_Guid _Guid{};  return ACAPI_GraphicalOverride_CreateOverrideRule(rule, _Guid); }
-GSErrCode ACAPI_Preferences_SetOldVersion (Int32 version, GSSize nByte, const void* data, unsigned short platformSign, API_FTypeID oldPlanFileID) { return GS::NoError; }
+GSErrCode ACAPI_GraphicalOverride_GetOverrideRule(API_OverrideRule& rule) { return GS::NoError; }
+GSErrCode ACAPI_GraphicalOverride_CreateOverrideRule(API_OverrideRule& rule) { const API_Guid _Guid{};  return ACAPI_GraphicalOverride_CreateOverrideRule(rule, _Guid); }
+GSErrCode ACAPI_Preferences_SetOldVersion(Int32 version, GSSize nByte, const void* data, unsigned short platformSign, API_FTypeID oldPlanFileID) { return GS::NoError; }
+#endif
+
+#ifdef __clang__
+#pragma clang diagnostic push
+#pragma clang diagnostic ignored "-Wdeprecated-enum-compare-conditional"
+#endif //__clang__
 
 // Definition of most of the methods from previous versions of the ArchiCAD SDK
 // Include once to avoid linkage errors for multiple implementations.
 #include "ACAPI_MigrationHeader.hpp"
-#endif
+
+#ifdef __clang__
+#pragma clang diagnostic pop
+#endif //__clang__
+
+#endif // AC_VERSION > 26 
 
 #if PLATFORM_WINDOWS && (__cplusplus < 201103L)
 	#define va_copy(destination, source) ((destination) = (source))
diff --git a/Engine/Source/Programs/Enterprise/Datasmith/DatasmithARCHICADExporter/Private/Utils/Element2String.cpp b/Engine/Source/Programs/Enterprise/Datasmith/DatasmithARCHICADExporter/Private/Utils/Element2String.cpp
--- a/Engine/Source/Programs/Enterprise/Datasmith/DatasmithARCHICADExporter/Private/Utils/Element2String.cpp
+++ b/Engine/Source/Programs/Enterprise/Datasmith/DatasmithARCHICADExporter/Private/Utils/Element2String.cpp
@@ -5,6 +5,11 @@
 #include "TAssValueName.h"
 #include "LibPartInfo.h"
 
+#if AC_VERSION > 27
+#include "ACAPI/IFCObjectAccessor.hpp"
+#include "ACAPI/IFCPropertyAccessor.hpp"
+#endif
+
 BEGIN_NAMESPACE_UE_AC
 
 utf8_string GetHotLinkInfo(const API_Element& InElement)
@@ -398,6 +403,7 @@ utf8_string FElement2String::GetIFCAttributesAsString(const API_Guid& InElementG
 {
 	utf8_string IFCAttributesString;
 
+	#if AC_VERSION < 28
 	GS::Array< API_IFCAttribute > Attributes;
 	bool						  bStoredOnly = false;
 	GSErrCode					  GSErr = ACAPI_Element_GetIFCAttributes(InElementGuid, bStoredOnly, &Attributes);
@@ -423,15 +429,176 @@ utf8_string FElement2String::GetIFCAttributesAsString(const API_Guid& InElementG
 			UE_AC_DebugF("FElement2String::GetIFCAttributesAsString - Error=%d getting IFC attributes\n", GSErr);
 		}
 	}
+	#else
+	API_Elem_Head Head{};
+	Head.guid = InElementGuid;
+	auto ElementObjectID = IFCAPI::GetObjectAccessor().CreateElementObjectID(Head);
+	if (ElementObjectID.IsOk())
+	{
+		auto IFCAttributes = IFCAPI::PropertyAccessor(*ElementObjectID).GetAttributes();
+		if (IFCAttributes.IsOk())
+		{
+			IFCAttributesString += Utf8StringFormat("\t\tIFCAtributes Nb=%zu\n", IFCAttributes->size());
+			for (const auto& IfcAttribute : *IFCAttributes)
+			{
+				if (IfcAttribute.GetValue().has_value())
+				{
+					IFCAttributesString += Utf8StringFormat(
+						"\t\t\tName=%s, Type=%s, %s=%s\n", IfcAttribute.GetName().ToUtf8(), IfcAttribute.GetValueType().ToUtf8(),
+						IfcAttribute.GetValue().has_value() ? "Value" : "Default", IfcAttribute.GetValue().value().ToUtf8());
+				}
+			}
+		}
+		else
+		{
+			UE_AC_DebugF("FElement2String::GetIFCAttributesAsString - Error=%s IFCAPI::PropertyAccessor(*ElementObjectID).GetAttributes()\n", GetErrorName(IFCAttributes.UnwrapErr().kind));
+		}
+	}
+	else
+	{
+		UE_AC_DebugF("FElement2String::GetIFCAttributesAsString - Error=%s IFCAPI::GetObjectAccessor().CreateElementObjectID\n", GetErrorName(ElementObjectID.UnwrapErr().kind));
+	}
+	#endif
 
 	return IFCAttributesString;
 }
 
+const GS::UniString TrueString("True");
+const GS::UniString FalseString("False");
+const GS::UniString UnknownString("Unknown");
+const GS::UniString SeparatorString(", ");
+
+#if AC_VERSION > 27
+GS::UniString FElement2String::GetValueTypePostfix(const GS::UniString& InValueType)
+{
+	return GS::UniString::Printf(" (%T)", InValueType.ToPrintf());
+}
+
+GS::UniString FElement2String::PropertyValueToString(const IFCAPI::v1::Value& InValue)
+{
+	GS::UniString result;
+
+	auto anyValue = InValue.GetAnyValue();
+	if (anyValue.has_value()) 
+	{
+		if (std::holds_alternative<Int64>(*anyValue)) 
+		{
+			result.Append(GS::ValueToUniString(std::get<Int64>(*anyValue)));
+		}
+		else if (std::holds_alternative<double>(*anyValue)) 
+		{
+			result.Append(GS::ValueToUniString(std::get<double>(*anyValue)));
+		}
+		else if (std::holds_alternative<bool>(*anyValue)) 
+		{
+			result.Append(std::get<bool>(*anyValue) ? TrueString : FalseString);
+		}
+		else if (std::holds_alternative<IFCAPI::IfcLogical>(*anyValue)) 
+		{
+			switch (std::get<IFCAPI::IfcLogical>(*anyValue)) 
+			{
+			case IFCAPI::IfcLogical::Unknown: result.Append(UnknownString); break;
+			case IFCAPI::IfcLogical::False: result.Append(FalseString); break;
+			case IFCAPI::IfcLogical::True: result.Append(TrueString); break;
+			}
+		}
+		else if (std::holds_alternative<GS::UniString>(*anyValue)) 
+		{
+			result.Append(std::get<GS::UniString>(*anyValue));
+		}
+	}
+
+	if (!result.IsEmpty())
+	{
+		result.Append(GetValueTypePostfix(InValue.GetType()));
+	}
+
+	return result;
+}
+
+GS::UniString FElement2String::Concatenate(const GS::Array<GS::UniString>& InStrings, const GS::UniString& InSeparator)
+{
+	GS::UniString result;
+
+	for (const auto& item : InStrings) 
+	{
+		if (!result.IsEmpty())
+		{
+			result.Append(InSeparator);
+		}
+
+		result.Append(item);
+	}
+
+	return result;
+}
+
+GS::UniString FElement2String::ConcatenatePropertyValues(const std::vector<IFCAPI::Value>& InValues, const GS::UniString& InSeparator)
+{
+	GS::Array<GS::UniString> valueStrings;
+
+	for (const auto& value : InValues)
+	{
+		valueStrings.Push(PropertyValueToString(value));
+	}
+
+	return Concatenate(valueStrings, InSeparator);
+}
+
+GS::UniString FElement2String::PropertyBoundedValueToString(const IFCAPI::PropertyBoundedValue& InBoundedValue)
+{
+	GS::UniString result;
+	GS::UniString lowerBoundValue = PropertyValueToString(InBoundedValue.GetLowerBoundValue());
+	GS::UniString upperBoundValue = PropertyValueToString(InBoundedValue.GetUpperBoundValue());
+	GS::UniString setPointValue = PropertyValueToString(InBoundedValue.GetSetPointValue());
+
+	if (!lowerBoundValue.IsEmpty() || !upperBoundValue.IsEmpty()) 
+	{
+		result.Append("[");
+		result.Append(lowerBoundValue);
+		result.Append(" ... ");
+		result.Append(upperBoundValue);
+		result.Append("]");
+	}
+
+	if (!setPointValue.IsEmpty()) 
+	{
+		if (!result.IsEmpty())
+		{
+			result.Append(SeparatorString);
+		}
+		result.Append(setPointValue);
+	}
+
+	return result;
+}
+
+GS::UniString FElement2String::PropertyTableValueToString(const IFCAPI::PropertyTableValue& InTableValue)
+{
+	GS::Array<GS::UniString> names;
+	auto definingValues = InTableValue.GetDefiningValues();
+	auto definedValues = InTableValue.GetDefinedValues();
+
+	if (!DBVERIFY(definingValues.size() == definedValues.size()))
+	{
+		return "";
+	}
+
+	for (size_t i = 0; i < definingValues.size(); ++i)
+	{
+		names.Push(PropertyValueToString(definingValues[i]) + ": " + PropertyValueToString(definedValues[i]));
+	}
+
+	return Concatenate(names, SeparatorString);
+}
+#endif //AC_VERSION > 27
+
 // Tool:Return IFC properties as a string
 utf8_string FElement2String::GetIFCPropertiesAsString(const API_Guid& InElementGuid)
 {
 	utf8_string IFCPropertiesString;
 
+	#if AC_VERSION < 28
 	GS::Array< API_IFCProperty > Properties;
 	bool						 bStoredOnly = false;
 	GSErrCode					 GSErr = ACAPI_Element_GetIFCProperties(InElementGuid, bStoredOnly, &Properties);
@@ -543,6 +710,81 @@ utf8_string FElement2String::GetIFCPropertiesAsString(const API_Guid& InElementG
 			UE_AC_DebugF("FElement2String::GetIFCPropertiesAsString - Error=%d getting IFC attributes\n", GSErr);
 		}
 	}
+	#else
+	API_Elem_Head Head{};
+	Head.guid = InElementGuid;
+	const auto ElementObjectID = IFCAPI::GetObjectAccessor().CreateElementObjectID(Head);
+	if (ElementObjectID.IsOk())
+	{
+		const auto LocalProperties = IFCAPI::PropertyAccessor(*ElementObjectID).GetLocalProperties();
+		if (LocalProperties.IsOk())
+		{
+			IFCPropertiesString += Utf8StringFormat("\t\tIFCProperties Nb=%zu\n", LocalProperties->size());
+			for (const auto& Property : *LocalProperties)
+			{
+				IFCPropertiesString += Utf8StringFormat("\t\t\tSet=%s, Name=%s",
+					Property.GetPropertySetName().ToUtf8(),
+					Property.GetName().ToUtf8());
+
+				auto propertyVariant = Property.GetTyped();
+
+				if (std::holds_alternative<IFCAPI::v1::PropertySingleValue>(propertyVariant)) 
+				{
+					const auto& propertyCast = std::get<IFCAPI::v1::PropertySingleValue>(propertyVariant);
+					IFCPropertiesString += PropertyValueToString(propertyCast.GetNominalValue()).ToUtf8();
+					IFCPropertiesString += "\n";
+				}
+				else if (std::holds_alternative<IFCAPI::v1::PropertyListValue>(propertyVariant)) 
+				{
+					const auto& propertyCast = std::get<IFCAPI::v1::PropertyListValue>(propertyVariant);
+					IFCPropertiesString += "{";
+					IFCPropertiesString += ConcatenatePropertyValues(propertyCast.GetListValues(), SeparatorString).ToUtf8();
+					IFCPropertiesString += "}\n";
+				}
+				else if (std::holds_alternative<IFCAPI::v1::PropertyBoundedValue>(propertyVariant)) 
+				{
+					const auto& propertyCast = std::get<IFCAPI::v1::PropertyBoundedValue>(propertyVariant);
+					IFCPropertiesString += PropertyBoundedValueToString(propertyCast).ToUtf8();
+					IFCPropertiesString += "\n";
+
+				}
+				else if (std::holds_alternative<IFCAPI::v1::PropertyEnumeratedValue>(propertyVariant)) 
+				{
+					const auto& propertyCast = std::get<IFCAPI::v1::PropertyEnumeratedValue>(propertyVariant);
+					GS::Array<GS::UniString> names;
+					for (const auto& value : propertyCast.GetEnumerationValues())
+					{
+						names.Push(PropertyValueToString(value));
+					}
+					IFCPropertiesString += "(";
+					IFCPropertiesString += Concatenate(names, SeparatorString).ToUtf8();
+					IFCPropertiesString += ")\n";
+				}
+				else 
+				{
+					const auto& propertyCast = std::get<IFCAPI::v1::PropertyTableValue>(propertyVariant);
+					IFCPropertiesString += PropertyTableValueToString(propertyCast).ToUtf8();
+				}
+			}
+		}
+		else
+		{
+			if (LocalProperties.UnwrapErr().kind != APIERR_BADPARS)
+			{
+				UE_AC_DebugF("FElement2String::GetIFCPropertiesAsString - IFCAPI::PropertyAccessor().GetLocalProperties returned error %s\n",
+					GetErrorName(LocalProperties.UnwrapErr().kind));
+			}
+		}
+	}
+	else
+	{
+		if (ElementObjectID.UnwrapErr().kind != APIERR_BADPARS)
+		{
+			UE_AC_DebugF("FElement2String::GetIFCPropertiesAsString - IFCAPI::GetObjectAccessor().CreateElementObjectID returned error %s\n",
+				GetErrorName(ElementObjectID.UnwrapErr().kind));
+		}
+	}
+	#endif
 
 	return IFCPropertiesString;
 }
diff --git a/Engine/Source/Programs/Enterprise/Datasmith/DatasmithARCHICADExporter/Private/Utils/Element2String.h b/Engine/Source/Programs/Enterprise/Datasmith/DatasmithARCHICADExporter/Private/Utils/Element2String.h
--- a/Engine/Source/Programs/Enterprise/Datasmith/DatasmithARCHICADExporter/Private/Utils/Element2String.h
+++ b/Engine/Source/Programs/Enterprise/Datasmith/DatasmithARCHICADExporter/Private/Utils/Element2String.h
@@ -3,6 +3,10 @@
 #pragma once
 
 #include "AddonTools.h"
+#if AC_VERSION > 27
+#include "ACAPI/IFCValue.hpp"
+#include "ACAPI/IFCProperty.hpp"
+#endif
 
 BEGIN_NAMESPACE_UE_AC
 
@@ -49,6 +53,20 @@ class FElement2String
 	// Tool:Return property objects as a string
 	static utf8_string GetPropertyObjectsAsString(const API_Elem_Head& InElementHeader);
 
+#if AC_VERSION > 27
+	static GS::UniString GetValueTypePostfix(const GS::UniString& InValueType);
+
+	static GS::UniString PropertyValueToString(const IFCAPI::v1::Value& InValue);
+
+	static GS::UniString Concatenate(const GS::Array<GS::UniString>& InStrings, const GS::UniString& InSeparator);
+
+	static GS::UniString ConcatenatePropertyValues(const std::vector<IFCAPI::Value>& InValues, const GS::UniString& InSeparator);
+
+	static GS::UniString PropertyBoundedValueToString(const IFCAPI::v1::PropertyBoundedValue& InBoundedValue);
+
+	static GS::UniString PropertyTableValueToString(const IFCAPI::PropertyTableValue& InTableValue);
+#endif //AC_VERSION > 27
+
   private:
 	// Tool:Return parameters as a string
 	static utf8_string GetParametersAsString(const API_AddParType* const* InParamsHandle);
diff --git a/Engine/Source/Programs/Enterprise/Datasmith/DatasmithARCHICADExporter/Private/Utils/ElementTools.cpp b/Engine/Source/Programs/Enterprise/Datasmith/DatasmithARCHICADExporter/Private/Utils/ElementTools.cpp
--- a/Engine/Source/Programs/Enterprise/Datasmith/DatasmithARCHICADExporter/Private/Utils/ElementTools.cpp
+++ b/Engine/Source/Programs/Enterprise/Datasmith/DatasmithARCHICADExporter/Private/Utils/ElementTools.cpp
@@ -133,9 +133,11 @@ FAssValueName::SAssValueName TAssEnumName< API_ElemVariationID >::AssEnumName[]
 	ValueName(APIVarId_LabelCWJunction),
 #endif
 
+#if AC_VERSION < 28
 	{ APIVarId_SymbStair, 		"SymbStair" },
+#endif
+
 	{ APIVarId_WallEnd, 		"WallEnd" },
-	{ APIVarId_SymbStair, 		"SymbStair" },
 	{ APIVarId_Door, 			"Door" },
 	{ APIVarId_Object, 			"Object" },
 	{ APIVarId_GridElement, 	"GridElement" },
diff --git a/Engine/Source/Programs/Enterprise/Datasmith/DatasmithARCHICADExporter/Private/Utils/Error.cpp b/Engine/Source/Programs/Enterprise/Datasmith/DatasmithARCHICADExporter/Private/Utils/Error.cpp
--- a/Engine/Source/Programs/Enterprise/Datasmith/DatasmithARCHICADExporter/Private/Utils/Error.cpp
+++ b/Engine/Source/Programs/Enterprise/Datasmith/DatasmithARCHICADExporter/Private/Utils/Error.cpp
@@ -39,8 +39,13 @@ void ShowAlert(const UE_AC_Error& InException, const utf8_t* InFct)
 
 void ShowAlert(const GS::GSException& inGSException, const utf8_t* InFct)
 {
-	UE_AC_DebugF("Caught a GS exception (%d) \"%s\" in %s\n", (int)inGSException.GetID(), inGSException.GetName(),
+	#if AC_VERSION > 27
+	UE_AC_DebugF("Caught a GS exception (%d) \"%s\" in %s\n", (int)inGSException.GetErrCode(), inGSException.GetName(),
 				 InFct);
+	#else
+	UE_AC_DebugF("Caught a GS exception (%d) \"%s\" in %s\n", (int)inGSException.GetID(), inGSException.GetName(),
+		InFct);
+	#endif
 
 	DGResAlert(ACAPI_GetOwnResModule(), LocalizeResId(kAlertACDBError));
 }
diff --git a/Engine/Source/Programs/Enterprise/Datasmith/DatasmithARCHICADExporter/Private/Utils/Error.h b/Engine/Source/Programs/Enterprise/Datasmith/DatasmithARCHICADExporter/Private/Utils/Error.h
--- a/Engine/Source/Programs/Enterprise/Datasmith/DatasmithARCHICADExporter/Private/Utils/Error.h
+++ b/Engine/Source/Programs/Enterprise/Datasmith/DatasmithARCHICADExporter/Private/Utils/Error.h
@@ -95,13 +95,24 @@ template < typename Functor > GSErrCode TryFunctionCatchAndLog(const utf8_t* InF
 	{
 		if (gs.GetFileName() != nullptr)
 		{
-			UE_AC_DebugF("%s - GS Exception %s - %d : %s, %s, %d\n", InFctName, gs.GetName(), gs.GetID(),
+			#if AC_VERSION > 27
+			UE_AC_DebugF("%s - GS Exception %s - %d : %s, %s, %d\n", InFctName, gs.GetName(), gs.GetErrCode(),
 						 gs.GetMessage().ToUtf8(), gs.GetFileName(), gs.GetLineNumber());
+			#else
+			UE_AC_DebugF("%s - GS Exception %s - %d : %s, %s, %d\n", InFctName, gs.GetName(), gs.GetID(),
+				gs.GetMessage().ToUtf8(), gs.GetFileName(), gs.GetLineNumber());
+			#endif
 		}
 		else
 		{
-			UE_AC_DebugF("%s - GS Exception %s - %d : %s\n", InFctName, gs.GetName(), gs.GetID(),
+			#if AC_VERSION > 27
+			UE_AC_DebugF("%s - GS Exception %s - %d : %s\n", InFctName, gs.GetName(), gs.GetErrCode(),
 						 gs.GetMessage().ToUtf8());
+			#else
+			UE_AC_DebugF("%s - GS Exception %s - %d : %s\n", InFctName, gs.GetName(), gs.GetID(),
+				gs.GetMessage().ToUtf8());
+			#endif
+
 		}
 	}
 	catch (...)
diff --git a/Engine/Source/Programs/Enterprise/Datasmith/DatasmithARCHICADExporter/Private/Utils/TaskCalledFromEventLoop.cpp b/Engine/Source/Programs/Enterprise/Datasmith/DatasmithARCHICADExporter/Private/Utils/TaskCalledFromEventLoop.cpp
--- a/Engine/Source/Programs/Enterprise/Datasmith/DatasmithARCHICADExporter/Private/Utils/TaskCalledFromEventLoop.cpp
+++ b/Engine/Source/Programs/Enterprise/Datasmith/DatasmithARCHICADExporter/Private/Utils/TaskCalledFromEventLoop.cpp
@@ -49,14 +49,22 @@ GSErrCode FTaskCalledFromEventLoop::DoTasks(GSHandle ParamHandle)
 			return GSErr;
 		}
 		TSharedPtr< FTaskCalledFromEventLoop > TaskPtr;
+		#if AC_VERSION > 27
+		if (CHCompareCStrings(Param.name, "TaskSharedRef", GS::CaseSensitive) == 0 && Param.type == MDCLPar_pointer)
+		#else
 		if (CHCompareCStrings(Param.name, "TaskSharedRef", CS_CaseSensitive) == 0 && Param.type == MDCLPar_pointer)
+		#endif
 		{
 			TSharedRef< FTaskCalledFromEventLoop >* TaskRefPtr =
 				reinterpret_cast< TSharedRef< FTaskCalledFromEventLoop >* >(Param.ptr_par);
 			TaskPtr = *TaskRefPtr;
 			delete TaskRefPtr;
 		}
+		#if AC_VERSION > 27
+		else if (CHCompareCStrings(Param.name, "TaskWeakPtr", GS::CaseSensitive) == 0 && Param.type == MDCLPar_pointer)
+		#else
 		else if (CHCompareCStrings(Param.name, "TaskWeakPtr", CS_CaseSensitive) == 0 && Param.type == MDCLPar_pointer)
+		#endif
 		{
 			TWeakPtr< FTaskCalledFromEventLoop >* TaskWeakPtrPtr =
 				reinterpret_cast< TWeakPtr< FTaskCalledFromEventLoop >* >(Param.ptr_par);
