SketchUp Datasmith Exporter : Fix NaN in Transforms due to double -> float -> double conversion that is now useless.  Culprit was creating an Identiy FQuat from a Imath::Quat<float>.

diff --git a/Engine/Source/Programs/Enterprise/Datasmith/DatasmithSketchUpRubyExporter/Private/DatasmithSketchUpUtils.cpp b/Engine/Source/Programs/Enterprise/Datasmith/DatasmithSketchUpRubyExporter/Private/DatasmithSketchUpUtils.cpp
--- a/Engine/Source/Programs/Enterprise/Datasmith/DatasmithSketchUpRubyExporter/Private/DatasmithSketchUpUtils.cpp
+++ b/Engine/Source/Programs/Enterprise/Datasmith/DatasmithSketchUpRubyExporter/Private/DatasmithSketchUpUtils.cpp
@@ -241,32 +241,34 @@ bool DatasmithSketchUpUtils::DecomposeTransform(
 	// We use Imath::extractAndRemoveScalingAndShear() because FMatrix::ExtractScaling() is deemed unreliable.
 
 	// Set up a scaling and rotation matrix.
-	auto& SMatrix = InWorldTransform.values;
-	Imath::Matrix44<float> Matrix(float(SMatrix[0]), float(SMatrix[1]), float(SMatrix[2]), 0.0,
-		float(SMatrix[4]), float(SMatrix[5]), float(SMatrix[6]), 0.0,
-		float(SMatrix[8]), float(SMatrix[9]), float(SMatrix[10]), 0.0,
-		0.0, 0.0, 0.0, 1.0);
+	Imath::Matrix44<double> Matrix
+	(
+		InWorldTransform.values[0], InWorldTransform.values[1], InWorldTransform.values[2], 0.0,
+		InWorldTransform.values[4], InWorldTransform.values[5], InWorldTransform.values[6], 0.0,
+		InWorldTransform.values[8], InWorldTransform.values[9], InWorldTransform.values[10], 0.0,
+		0.0, 0.0, 0.0, 1.0
+	);
 
 	// Remove any scaling from the matrix and get the scale vector that was initially present.
-	Imath::Vec3<float> Scale;
-	Imath::Vec3<float> Shear;
-	if (!Imath::extractAndRemoveScalingAndShear<float>(Matrix, Scale, Shear, false))
+	Imath::Vec3<double> Scale;
+	Imath::Vec3<double> Shear;
+	if (!Imath::extractAndRemoveScalingAndShear<double>(Matrix, Scale, Shear, false))
 	{
 		// Scaling is zero
 		return false;
 	}
 
-	if (SMatrix[15] != 1.0)
+	if (InWorldTransform.values[15] != 1.0)
 	{
 		// Apply the extra SketchUp uniform scaling factor.
-		Scale *= SMatrix[15];
+		Scale *= InWorldTransform.values[15];
 	}
 
 	OutScale = FVector(Scale.x, Scale.y, Scale.z);
 	OutShear = FVector(Shear.x, Shear.y, Shear.z);
 
 	// Initialize a rotation quaternion with the rotation matrix.
-	Imath::Quat<float> Quaternion = Imath::extractQuat<float>(Matrix);
+	Imath::Quat<double> Quaternion = Imath::extractQuat<double>(Matrix);
 
 	// Convert the SketchUp right-handed Z-up coordinate rotation into an Unreal left-handed Z-up coordinate rotation.
 	// This is done by inverting the X and Z components of the quaternion to mirror the quaternion on the XZ-plane.
@@ -275,14 +277,14 @@ bool DatasmithSketchUpUtils::DecomposeTransform(
 	Quaternion.normalize();
 
 	// Make sure Unreal will be able to handle the rotation quaternion.
-	float              Angle = Quaternion.angle();
-	Imath::Vec3<float> Axis = Quaternion.axis();
+	double              Angle = Quaternion.angle();
+	Imath::Vec3<double> Axis = Quaternion.axis();
 	OutRotation = FQuat(FVector(Axis.x, Axis.y, Axis.z), Angle);
 
 	// Convert the SketchUp right-handed Z-up coordinate translation into an Unreal left-handed Z-up coordinate translation.
 	// To avoid perturbating X, which is forward in Unreal, the handedness conversion is done by flipping the side vector Y.
 	// SketchUp uses inches as internal system unit for all 3D coordinates in the model while Unreal uses centimeters.
-	OutTranslation = FromSketchUp::ConvertPosition(SMatrix[12], SMatrix[13], SMatrix[14]);
+	OutTranslation = FromSketchUp::ConvertPosition(InWorldTransform.values[12], InWorldTransform.values[13], InWorldTransform.values[14]);
 
 	return true;
 }
