[Datasmith] [ArchiCAD] fixed plugin build after improved Printf validation

diff --git a/Engine/Source/Programs/Enterprise/Datasmith/DatasmithARCHICADExporter/DatasmithValidator/Private/SceneValidator.cpp b/Engine/Source/Programs/Enterprise/Datasmith/DatasmithARCHICADExporter/DatasmithValidator/Private/SceneValidator.cpp
--- a/Engine/Source/Programs/Enterprise/Datasmith/DatasmithARCHICADExporter/DatasmithValidator/Private/SceneValidator.cpp
+++ b/Engine/Source/Programs/Enterprise/Datasmith/DatasmithARCHICADExporter/DatasmithValidator/Private/SceneValidator.cpp
@@ -294,13 +294,13 @@ void FSceneValidator::CheckTexturesFiles()
                 }
                 else
                 {
-                    AddMessage(kBug, TEXT("Texture Name=\"%s\", Label=\"%s\", Path=\"%s\" - Relative path whitout reference path"), Texture->GetName(), Texture->GetLabel(), *TextureFile);
+                    AddMessage(kBug, TEXT("Texture Name=\"%s\", Label=\"%s\", Path=\"%s\" - Relative path without reference path"), Texture->GetName(), Texture->GetLabel(), *TextureFile);
                     Reported = true;
                 }
             }
             if (!Reported && !FPaths::FileExists(TextureFile))
             {
-                AddMessage(kBug, TEXT("Texture Name=\"%s\", Label=\"%s\", Path=\"%s\" - File not found"), *TextureFile);
+                AddMessage(kBug, TEXT("Texture Name=\"%s\", Label=\"%s\", Path=\"%s\" - File not found"), Texture->GetName(), Texture->GetLabel(), *TextureFile);
             }
         }
     }
diff --git a/Engine/Source/Programs/Enterprise/Datasmith/DatasmithARCHICADExporter/DatasmithValidator/Private/SceneValidator.h b/Engine/Source/Programs/Enterprise/Datasmith/DatasmithARCHICADExporter/DatasmithValidator/Private/SceneValidator.h
--- a/Engine/Source/Programs/Enterprise/Datasmith/DatasmithARCHICADExporter/DatasmithValidator/Private/SceneValidator.h
+++ b/Engine/Source/Programs/Enterprise/Datasmith/DatasmithARCHICADExporter/DatasmithValidator/Private/SceneValidator.h
@@ -81,8 +81,8 @@ class FSceneValidator : public ISceneValidator
 
 	void AddElements(const IDatasmithElement& InElement, FMapNameToUsage* IOMap = nullptr);
 
-	template < typename FmtType, typename... Types >
-	void AddMessage(TInfoLevel Level, const FmtType& Fmt, Types... Args)
+	template < typename... Types >
+	void AddMessage(TInfoLevel Level, UE::Core::TCheckedFormatString<FString::FmtCharType, Types...> Fmt, Types... Args)
 	{
 		AddMessageImpl(Level, FString::Printf(Fmt, Args...));
 	}
