Fixed compilation errors due to the creation of AutoRTFM module

#rnx

diff --git a/Engine/Source/Programs/Enterprise/Datasmith/DatasmithARCHICADExporter/Build/Config/UE_AC.xcconfig b/Engine/Source/Programs/Enterprise/Datasmith/DatasmithARCHICADExporter/Build/Config/UE_AC.xcconfig
--- a/Engine/Source/Programs/Enterprise/Datasmith/DatasmithARCHICADExporter/Build/Config/UE_AC.xcconfig
+++ b/Engine/Source/Programs/Enterprise/Datasmith/DatasmithARCHICADExporter/Build/Config/UE_AC.xcconfig
@@ -12,7 +12,7 @@ UE_Project = $(UE_Runtime)/Projects/Public
 UE_ThirdParty = $(UE_Engine)/Source/ThirdParty
 UE_ArchiCADExporter = $(UE_Engine)/Source/Programs/Enterprise/Datasmith/DatasmithARCHICADExporter
 
-UE_INCLUDES_PATH = "$(UE_Core)" "$(UE_Core)/Internationalization" "$(UE_Core)/Async" "$(UE_Core)/Concurrency" "$(UE_Core)/Containers" "$(UE_Core)/Delegates" "$(UE_Core)/GenericPlatform" "$(UE_Core)/HAL" "$(UE_Core)/Logging" "$(UE_Core)/Math" "$(UE_Core)/Misc" "$(UE_Core)/Modules" "$(UE_Core)/Modules\Boilerplate" "$(UE_Core)/ProfilingDebugging" "$(UE_Core)/Serialization" "$(UE_Core)/Serialization\Csv" "$(UE_Core)/Stats" "$(UE_Core)/Templates" "$(UE_Core)/UObject" "$(UE_Project)/Public" "$(UE_Project)/Interfaces" "$(UE_Runtime)/TraceLog/Public"  "$(UE_Runtime)/Datasmith/DatasmithCore/Public" "$(UE_Runtime)/Datasmith/DirectLink/Public" "$(UE_Runtime)/Launch/Resources" "$(UE_Engine)/Source/Developer/Datasmith/DatasmithExporter/Public" "$(UE_Engine)/Source/Developer/Datasmith/DatasmithExporterUI/Public" "$(UE_Engine)/Source/Runtime/Messaging/Public" "$(UE_Engine)/Source/Runtime/CoreUObject/Public/" "$(UE_ArchiCADExporter)/DatasmithValidator/Public" "$(UE_ThirdParty)/GuidelinesSupportLibrary/GSL-1144/include"
+UE_INCLUDES_PATH = "$(UE_Core)" "$(UE_Core)/Internationalization" "$(UE_Core)/Async" "$(UE_Core)/Concurrency" "$(UE_Core)/Containers" "$(UE_Core)/Delegates" "$(UE_Core)/GenericPlatform" "$(UE_Core)/HAL" "$(UE_Core)/Logging" "$(UE_Core)/Math" "$(UE_Core)/Misc" "$(UE_Core)/Modules" "$(UE_Core)/Modules\Boilerplate" "$(UE_Core)/ProfilingDebugging" "$(UE_Core)/Serialization" "$(UE_Core)/Serialization\Csv" "$(UE_Core)/Stats" "$(UE_Core)/Templates" "$(UE_Core)/UObject" "$(UE_Runtime)/AutoRTFM/Public" "$(UE_Runtime)/AutoRTFM/Public/AutoRTFM" "$(UE_Project)/Public" "$(UE_Project)/Interfaces" "$(UE_Runtime)/TraceLog/Public"  "$(UE_Runtime)/Datasmith/DatasmithCore/Public" "$(UE_Runtime)/Datasmith/DirectLink/Public" "$(UE_Runtime)/Launch/Resources" "$(UE_Engine)/Source/Developer/Datasmith/DatasmithExporter/Public" "$(UE_Engine)/Source/Developer/Datasmith/DatasmithExporterUI/Public" "$(UE_Engine)/Source/Runtime/Messaging/Public" "$(UE_Engine)/Source/Runtime/CoreUObject/Public/" "$(UE_ArchiCADExporter)/DatasmithValidator/Public" "$(UE_ThirdParty)/GuidelinesSupportLibrary/GSL-1144/include"
 
 USER_HEADER_SEARCH_PATHS = "." $(AC_INCLUDES_PATH) $(UE_INCLUDES_PATH)
 
diff --git a/Engine/Source/Programs/Enterprise/Datasmith/DatasmithARCHICADExporter/Build/Config/UE_AC_ARM.xcconfig b/Engine/Source/Programs/Enterprise/Datasmith/DatasmithARCHICADExporter/Build/Config/UE_AC_ARM.xcconfig
--- a/Engine/Source/Programs/Enterprise/Datasmith/DatasmithARCHICADExporter/Build/Config/UE_AC_ARM.xcconfig
+++ b/Engine/Source/Programs/Enterprise/Datasmith/DatasmithARCHICADExporter/Build/Config/UE_AC_ARM.xcconfig
@@ -12,7 +12,7 @@ UE_Project = $(UE_Runtime)/Projects/Public
 UE_ThirdParty = $(UE_Engine)/Source/ThirdParty
 UE_ArchiCADExporter = $(UE_Engine)/Source/Programs/Enterprise/Datasmith/DatasmithARCHICADExporter
 
-UE_INCLUDES_PATH = "$(UE_Core)" "$(UE_Core)/Internationalization" "$(UE_Core)/Async" "$(UE_Core)/Concurrency" "$(UE_Core)/Containers" "$(UE_Core)/Delegates" "$(UE_Core)/GenericPlatform" "$(UE_Core)/HAL" "$(UE_Core)/Logging" "$(UE_Core)/Math" "$(UE_Core)/Misc" "$(UE_Core)/Modules" "$(UE_Core)/Modules\Boilerplate" "$(UE_Core)/ProfilingDebugging" "$(UE_Core)/Serialization" "$(UE_Core)/Serialization\Csv" "$(UE_Core)/Stats" "$(UE_Core)/Templates" "$(UE_Core)/UObject" "$(UE_Project)/Public" "$(UE_Project)/Interfaces" "$(UE_Runtime)/TraceLog/Public"  "$(UE_Runtime)/Datasmith/DatasmithCore/Public" "$(UE_Runtime)/Datasmith/DirectLink/Public" "$(UE_Runtime)/Launch/Resources" "$(UE_Engine)/Source/Developer/Datasmith/DatasmithExporter/Public" "$(UE_Engine)/Source/Developer/Datasmith/DatasmithExporterUI/Public" "$(UE_Engine)/Source/Runtime/Messaging/Public" "$(UE_Engine)/Source/Runtime/CoreUObject/Public/" "$(UE_ArchiCADExporter)/DatasmithValidator/Public" "$(UE_ThirdParty)/GuidelinesSupportLibrary/GSL-1144/include"
+UE_INCLUDES_PATH = "$(UE_Core)" "$(UE_Core)/Internationalization" "$(UE_Core)/Async" "$(UE_Core)/Concurrency" "$(UE_Core)/Containers" "$(UE_Core)/Delegates" "$(UE_Core)/GenericPlatform" "$(UE_Core)/HAL" "$(UE_Core)/Logging" "$(UE_Core)/Math" "$(UE_Core)/Misc" "$(UE_Core)/Modules" "$(UE_Core)/Modules\Boilerplate" "$(UE_Core)/ProfilingDebugging" "$(UE_Core)/Serialization" "$(UE_Core)/Serialization\Csv" "$(UE_Core)/Stats" "$(UE_Core)/Templates" "$(UE_Core)/UObject" "$(UE_Runtime)/AutoRTFM/Public" "$(UE_Runtime)/AutoRTFM/Public/AutoRTFM" "$(UE_Project)/Public" "$(UE_Project)/Interfaces" "$(UE_Runtime)/TraceLog/Public"  "$(UE_Runtime)/Datasmith/DatasmithCore/Public" "$(UE_Runtime)/Datasmith/DirectLink/Public" "$(UE_Runtime)/Launch/Resources" "$(UE_Engine)/Source/Developer/Datasmith/DatasmithExporter/Public" "$(UE_Engine)/Source/Developer/Datasmith/DatasmithExporterUI/Public" "$(UE_Engine)/Source/Runtime/Messaging/Public" "$(UE_Engine)/Source/Runtime/CoreUObject/Public/" "$(UE_ArchiCADExporter)/DatasmithValidator/Public" "$(UE_ThirdParty)/GuidelinesSupportLibrary/GSL-1144/include"
 
 USER_HEADER_SEARCH_PATHS = "." $(AC_INCLUDES_PATH) $(UE_INCLUDES_PATH)
 
diff --git a/Engine/Source/Programs/Enterprise/Datasmith/DatasmithSDK/DatasmithSDK.Build.cs b/Engine/Source/Programs/Enterprise/Datasmith/DatasmithSDK/DatasmithSDK.Build.cs
--- a/Engine/Source/Programs/Enterprise/Datasmith/DatasmithSDK/DatasmithSDK.Build.cs
+++ b/Engine/Source/Programs/Enterprise/Datasmith/DatasmithSDK/DatasmithSDK.Build.cs
@@ -8,6 +8,8 @@ namespace UnrealBuildTool.Rules
 		public DatasmithSDK(ReadOnlyTargetRules Target)
 			: base(Target)
 		{
+			CppStandard = CppStandardVersion.Cpp20;
+
 			PublicIncludePathModuleNames.Add("Launch");
 
 			PrivateDependencyModuleNames.AddRange(
diff --git a/Engine/Source/Programs/Enterprise/Datasmith/DatasmithSDK/DatasmithSDK.Target.cs b/Engine/Source/Programs/Enterprise/Datasmith/DatasmithSDK/DatasmithSDK.Target.cs
--- a/Engine/Source/Programs/Enterprise/Datasmith/DatasmithSDK/DatasmithSDK.Target.cs
+++ b/Engine/Source/Programs/Enterprise/Datasmith/DatasmithSDK/DatasmithSDK.Target.cs
@@ -154,6 +154,7 @@ public class DatasmithSDKTarget : TargetRules
 			(@"$(EngineDir)\Source\Runtime\TraceLog\Public\", Path.Combine(DstDatasmithSDK,"Private\\"), "*.inl"),
 			(@"$(EngineDir)\Source\Runtime\Messaging\Public\", Path.Combine(DstDatasmithSDK,"Private\\"), "*.h"),
 			(@"$(EngineDir)\Source\Runtime\Messaging\Public\", Path.Combine(DstDatasmithSDK,"Private\\"), "*.inl"),
+			(@"$(EngineDir)\Source\Runtime\AutoRTFM\Public\", Path.Combine(DstDatasmithSDK,"Private\\"), "*.h"),
 			(@"$(EngineDir)\Source\Runtime\Core\Public\", Path.Combine(DstDatasmithSDK,"Private\\"), "*.h"),
 			(@"$(EngineDir)\Source\Runtime\Core\Public\", Path.Combine(DstDatasmithSDK,"Private\\"), "*.inl"),
 			(@"$(EngineDir)\Source\Runtime\CoreUObject\Public\", Path.Combine(DstDatasmithSDK,"Private\\"), "*.h"),
diff --git a/Engine/Source/Programs/Enterprise/Datasmith/DatasmithSketchUpRubyExporter/Build/build_mac.py b/Engine/Source/Programs/Enterprise/Datasmith/DatasmithSketchUpRubyExporter/Build/build_mac.py
--- a/Engine/Source/Programs/Enterprise/Datasmith/DatasmithSketchUpRubyExporter/Build/build_mac.py
+++ b/Engine/Source/Programs/Enterprise/Datasmith/DatasmithSketchUpRubyExporter/Build/build_mac.py
@@ -126,7 +126,7 @@ UE_ThirdParty = f'{UE_Engine}/Source/ThirdParty'
 UE_Imath = f'{UE_ThirdParty}/Imath/Deploy/Imath-3.1.12'
 UE_GSL = f'{UE_ThirdParty}/GuidelinesSupportLibrary/GSL-1144'
 
-UE_INCLUDES_PATH = f'''"{UE_Core}" "{UE_Core}/Internationalization" "{UE_Core}/Async" "{UE_Core}/Containers" "{UE_Core}/Delegates" "{UE_Core}/GenericPlatform" "{UE_Core}/HAL" "{UE_Core}/Logging" "{UE_Core}/Math" "{UE_Core}/Misc" "{UE_Core}/Modules" "{UE_Core}/Modules/Boilerplate" "{UE_Core}/ProfilingDebugging" "{UE_Core}/Serialization" "{UE_Core}/Serialization/Csv" "{UE_Core}/Stats" "{UE_Core}/Templates" "{UE_Core}/UObject" "{UE_Runtime}/CoreUObject/Public" "{UE_Project}/Interfaces" "{UE_Runtime}/TraceLog/Public"  "{UE_Runtime}/Datasmith/DatasmithCore/Public" "{UE_Runtime}/Datasmith/DirectLink/Public" "{UE_Runtime}/Messaging/Public" "{UE_Runtime}/Launch/Resources" "{UE_Engine}/Source/Developer/Datasmith/DatasmithExporter/Public" "{UE_Engine}/Source/Developer/Datasmith/DatasmithExporterUI/Public"'''
+UE_INCLUDES_PATH = f'''"{UE_Core}" "{UE_Core}/Internationalization" "{UE_Core}/Async" "{UE_Core}/Containers" "{UE_Core}/Delegates" "{UE_Core}/GenericPlatform" "{UE_Core}/HAL" "{UE_Core}/Logging" "{UE_Core}/Math" "{UE_Core}/Misc" "{UE_Core}/Modules" "{UE_Core}/Modules/Boilerplate" "{UE_Core}/ProfilingDebugging" "{UE_Core}/Serialization" "{UE_Core}/Serialization/Csv" "{UE_Core}/Stats" "{UE_Core}/Templates" "{UE_Core}/UObject" "{UE_Runtime}/CoreUObject/Public" "{UE_Runtime}/AutoRTFM/Public" "{UE_Runtime}/AutoRTFM/Public/AutoRTFM" "{UE_Project}/Interfaces" "{UE_Runtime}/TraceLog/Public"  "{UE_Runtime}/Datasmith/DatasmithCore/Public" "{UE_Runtime}/Datasmith/DirectLink/Public" "{UE_Runtime}/Messaging/Public" "{UE_Runtime}/Launch/Resources" "{UE_Engine}/Source/Developer/Datasmith/DatasmithExporter/Public" "{UE_Engine}/Source/Developer/Datasmith/DatasmithExporterUI/Public"'''
 
 # Framework search path needed to lookup SketchUpApi and Ruby
 framework_search_paths = []
diff --git a/Engine/Source/Programs/Enterprise/Datasmith/DatasmithSketchUpRubyExporter/DatasmithSketchUpRuby2020.Build.cs b/Engine/Source/Programs/Enterprise/Datasmith/DatasmithSketchUpRubyExporter/DatasmithSketchUpRuby2020.Build.cs
--- a/Engine/Source/Programs/Enterprise/Datasmith/DatasmithSketchUpRubyExporter/DatasmithSketchUpRuby2020.Build.cs
+++ b/Engine/Source/Programs/Enterprise/Datasmith/DatasmithSketchUpRubyExporter/DatasmithSketchUpRuby2020.Build.cs
@@ -19,6 +19,7 @@ namespace UnrealBuildTool.Rules
 			PrivateDependencyModuleNames.AddRange(
 				new string[]
 				{
+					"AutoRTFM",
 					"DatasmithExporter",
 					"DatasmithExporterUI",
 
