DatasmithMax: Fix some c++20 compile errors

#rnx

diff --git a/Engine/Source/Programs/Enterprise/Datasmith/DatasmithMaxExporter/Private/DatasmithMaxArnold.cpp b/Engine/Source/Programs/Enterprise/Datasmith/DatasmithMaxExporter/Private/DatasmithMaxArnold.cpp
--- a/Engine/Source/Programs/Enterprise/Datasmith/DatasmithMaxExporter/Private/DatasmithMaxArnold.cpp
+++ b/Engine/Source/Programs/Enterprise/Datasmith/DatasmithMaxExporter/Private/DatasmithMaxArnold.cpp
@@ -151,7 +151,7 @@ void FDatasmithMaxMatWriter::ExportPhysicalMaterialCoat(TSharedRef< IDatasmithSc
 	}
 
 	// There is no coating effect
-	if (CoatAmount <= 0.0f || ((CoatColorMap == false || bCoatColorMapOn == false) && CoatColor.r == 0.0f && CoatColor.g == 0.0f && CoatColor.b == 0))
+	if (CoatAmount <= 0.0f || ((CoatColorMap == nullptr || bCoatColorMapOn == false) && CoatColor.r == 0.0f && CoatColor.g == 0.0f && CoatColor.b == 0))
 	{
 		return;
 	}
diff --git a/Engine/Source/Programs/Enterprise/Datasmith/DatasmithMaxExporter/Private/DatasmithMaxAttributes.cpp b/Engine/Source/Programs/Enterprise/Datasmith/DatasmithMaxExporter/Private/DatasmithMaxAttributes.cpp
--- a/Engine/Source/Programs/Enterprise/Datasmith/DatasmithMaxExporter/Private/DatasmithMaxAttributes.cpp
+++ b/Engine/Source/Programs/Enterprise/Datasmith/DatasmithMaxExporter/Private/DatasmithMaxAttributes.cpp
@@ -5,6 +5,8 @@
 #include "DatasmithMaxExporterDefines.h"
 #include "DatasmithMaxClassIDs.h"
 
+#define NOMINMAX
+
 MAX_INCLUDES_START
 	#include "inode.h"
 	#include "iparamb2.h"
diff --git a/Engine/Source/Programs/Enterprise/Datasmith/DatasmithMaxExporter/Private/DatasmithMaxCameraExporter.cpp b/Engine/Source/Programs/Enterprise/Datasmith/DatasmithMaxExporter/Private/DatasmithMaxCameraExporter.cpp
--- a/Engine/Source/Programs/Enterprise/Datasmith/DatasmithMaxExporter/Private/DatasmithMaxCameraExporter.cpp
+++ b/Engine/Source/Programs/Enterprise/Datasmith/DatasmithMaxExporter/Private/DatasmithMaxCameraExporter.cpp
@@ -397,7 +397,7 @@ bool FDatasmithMaxCameraExporter::ExportCamera(TimeValue CurrentTime, INode& Nod
 {
 	ObjectState ObjState = Node.EvalWorldState(0);
 	CameraObject* CamObj = (CameraObject*)ObjState.obj;
-	if (CamObj == false)
+	if (CamObj == nullptr)
 	{
 		return false;
 	}
diff --git a/Engine/Source/Programs/Enterprise/Datasmith/DatasmithMaxExporter/Private/DatasmithMaxExporterWithDirectLink.cpp b/Engine/Source/Programs/Enterprise/Datasmith/DatasmithMaxExporter/Private/DatasmithMaxExporterWithDirectLink.cpp
--- a/Engine/Source/Programs/Enterprise/Datasmith/DatasmithMaxExporter/Private/DatasmithMaxExporterWithDirectLink.cpp
+++ b/Engine/Source/Programs/Enterprise/Datasmith/DatasmithMaxExporter/Private/DatasmithMaxExporterWithDirectLink.cpp
@@ -1251,7 +1251,7 @@ public:
 				if (TSharedPtr<IDatasmithScene> ScenePtr = ExportedScene.GetDatasmithScene())
 				{
 					IDatasmithScene& Scene = *ScenePtr;
-					return FString::Printf(TEXT("Actors: %d; Meshes: %d, Materials: %d"), 
+					return FString::Printf(TEXT("Actors: %d; Meshes: %d, Materials: %d Textures: %d"), 
 					                       Scene.GetActorsCount(),
 					                       Scene.GetMeshesCount(),
 					                       Scene.GetMaterialsCount(),
@@ -3919,7 +3919,7 @@ bool Export(const TCHAR* Name, const TCHAR* OutputPath, bool bQuiet, bool bSelec
 		IDatasmithScene& Scene = *ExportedScene.GetDatasmithScene();
 		ExportedScene.GetSceneExporter().Export(ExportedScene.GetDatasmithScene().ToSharedRef(), false);
 
-		PROGRESS_STAGE_RESULT(FString::Printf(TEXT("Actors: %d; Meshes: %d, Materials: %d"), 
+		PROGRESS_STAGE_RESULT(FString::Printf(TEXT("Actors: %d; Meshes: %d, Materials: %d Textures: %d"), 
 			Scene.GetActorsCount(),
 			Scene.GetMeshesCount(),
 			Scene.GetMaterialsCount(),
diff --git a/Engine/Source/Programs/Enterprise/Datasmith/DatasmithMaxExporter/Private/MaxMaterialsToUEPbr/DatasmithMaxPhysicalMaterialToUEPbr.cpp b/Engine/Source/Programs/Enterprise/Datasmith/DatasmithMaxExporter/Private/MaxMaterialsToUEPbr/DatasmithMaxPhysicalMaterialToUEPbr.cpp
--- a/Engine/Source/Programs/Enterprise/Datasmith/DatasmithMaxExporter/Private/MaxMaterialsToUEPbr/DatasmithMaxPhysicalMaterialToUEPbr.cpp
+++ b/Engine/Source/Programs/Enterprise/Datasmith/DatasmithMaxExporter/Private/MaxMaterialsToUEPbr/DatasmithMaxPhysicalMaterialToUEPbr.cpp
@@ -586,7 +586,7 @@ struct FMaxPhysicalMaterialClearCoat
 	bool HasClearCoat()
 	{
 		// There is no coating effect
-		if (CoatAmount <= 0.0f || ((CoatColorMap.Map == false || CoatColorMap.bEnabled == false) && CoatColor.r == 0.0f && CoatColor.g == 0.0f && CoatColor.b == 0))
+		if (CoatAmount <= 0.0f || ((CoatColorMap.Map == nullptr || CoatColorMap.bEnabled == false) && CoatColor.r == 0.0f && CoatColor.g == 0.0f && CoatColor.b == 0))
 		{
 			return false;
 		}
diff --git a/Engine/Source/Programs/Enterprise/Datasmith/DatasmithMaxExporter/Private/MaxMaterialsToUEPbr/DatasmithMaxTexmapToUEPbr.cpp b/Engine/Source/Programs/Enterprise/Datasmith/DatasmithMaxExporter/Private/MaxMaterialsToUEPbr/DatasmithMaxTexmapToUEPbr.cpp
--- a/Engine/Source/Programs/Enterprise/Datasmith/DatasmithMaxExporter/Private/MaxMaterialsToUEPbr/DatasmithMaxTexmapToUEPbr.cpp
+++ b/Engine/Source/Programs/Enterprise/Datasmith/DatasmithMaxExporter/Private/MaxMaterialsToUEPbr/DatasmithMaxTexmapToUEPbr.cpp
@@ -1948,7 +1948,7 @@ IDatasmithMaterialExpression* FDatasmithMaxBakeableToUEPbr::Convert( FDatasmithM
 {
 	if ( !InTexmap )
 	{
-		return false;
+		return nullptr;
 	}
 
 	MSTR ClassName;
