Disabled crash reporting for Datasmith exporters

#rnx

diff --git a/Engine/Source/Programs/Enterprise/Datasmith/DatasmithARCHICADExporter/Build/Config/UE_AC.xcconfig b/Engine/Source/Programs/Enterprise/Datasmith/DatasmithARCHICADExporter/Build/Config/UE_AC.xcconfig
--- a/Engine/Source/Programs/Enterprise/Datasmith/DatasmithARCHICADExporter/Build/Config/UE_AC.xcconfig
+++ b/Engine/Source/Programs/Enterprise/Datasmith/DatasmithARCHICADExporter/Build/Config/UE_AC.xcconfig
@@ -63,7 +63,7 @@ ADDON_PREPROCESSOR_DEFINITIONS = ADDON_VERSION=$(VAddOn) VMAJOR=$(VMajor) VMINOR
 
 AC_PREPROCESSOR_DEFINITIONS = AC_VERSION=$(AC_VERSION) AC_SDK=$(AC_SDK_BUILD_NUMBER) macintosh=1 ACExtension
 
-UE_PREPROCESSOR_DEFINITIONS = UE_BUILD_DEVELOPMENT=1 UE_BUILD_MINIMAL=1 WITH_EDITOR=0 WITH_EDITORONLY_DATA=0 WITH_SERVER_CODE=1 WITH_ENGINE=0 WITH_UNREAL_DEVELOPER_TOOLS=0 WITH_PLUGIN_SUPPORT=0 IS_MONOLITHIC=1 IS_PROGRAM=1 PLATFORM_MAC=1 PLATFORM_APPLE=1 UE_BUILD_DEVELOPMENT_WITH_DEBUGGAME=0 UBT_COMPILED_PLATFORM=Mac CORE_API=DLLIMPORT DATASMITHEXPORTER_API=DLLIMPORT DATASMITHCORE_API=DLLIMPORT DIRECTLINK_API=DLLIMPORT DATASMITHEXPORTERUI_API=DLLIMPORT DATASMITHVALIDATOR_API=DLLIMPORT COREUOBJECT_API=DLLIMPORT
+UE_PREPROCESSOR_DEFINITIONS = NOINITCRASHREPORTER=1 UE_BUILD_DEVELOPMENT=1 UE_BUILD_MINIMAL=1 WITH_EDITOR=0 WITH_EDITORONLY_DATA=0 WITH_SERVER_CODE=1 WITH_ENGINE=0 WITH_UNREAL_DEVELOPER_TOOLS=0 WITH_PLUGIN_SUPPORT=0 IS_MONOLITHIC=1 IS_PROGRAM=1 PLATFORM_MAC=1 PLATFORM_APPLE=1 UE_BUILD_DEVELOPMENT_WITH_DEBUGGAME=0 UBT_COMPILED_PLATFORM=Mac CORE_API=DLLIMPORT DATASMITHEXPORTER_API=DLLIMPORT DATASMITHCORE_API=DLLIMPORT DIRECTLINK_API=DLLIMPORT DATASMITHEXPORTERUI_API=DLLIMPORT DATASMITHVALIDATOR_API=DLLIMPORT COREUOBJECT_API=DLLIMPORT
 
 //:configuration = Debug
 GCC_PREPROCESSOR_DEFINITIONS = $(ADDON_PREPROCESSOR_DEFINITIONS) $(AC_PREPROCESSOR_DEFINITIONS) $(UE_PREPROCESSOR_DEFINITIONS)
diff --git a/Engine/Source/Programs/Enterprise/Datasmith/DatasmithARCHICADExporter/Build/Config/UE_AC_ARM.xcconfig b/Engine/Source/Programs/Enterprise/Datasmith/DatasmithARCHICADExporter/Build/Config/UE_AC_ARM.xcconfig
--- a/Engine/Source/Programs/Enterprise/Datasmith/DatasmithARCHICADExporter/Build/Config/UE_AC_ARM.xcconfig
+++ b/Engine/Source/Programs/Enterprise/Datasmith/DatasmithARCHICADExporter/Build/Config/UE_AC_ARM.xcconfig
@@ -63,7 +63,7 @@ ADDON_PREPROCESSOR_DEFINITIONS = ADDON_VERSION=$(VAddOn) VMAJOR=$(VMajor) VMINOR
 
 AC_PREPROCESSOR_DEFINITIONS = AC_VERSION=$(AC_VERSION) AC_SDK=$(AC_SDK_BUILD_NUMBER) macintosh=1 ACExtension
 
-UE_PREPROCESSOR_DEFINITIONS = UE_BUILD_DEVELOPMENT=1 UE_BUILD_MINIMAL=1 WITH_EDITOR=0 WITH_EDITORONLY_DATA=0 WITH_SERVER_CODE=1 WITH_ENGINE=0 WITH_UNREAL_DEVELOPER_TOOLS=0 WITH_PLUGIN_SUPPORT=0 IS_MONOLITHIC=1 IS_PROGRAM=1 PLATFORM_MAC=1 PLATFORM_APPLE=1 UE_BUILD_DEVELOPMENT_WITH_DEBUGGAME=0 UBT_COMPILED_PLATFORM=Mac CORE_API=DLLIMPORT DATASMITHEXPORTER_API=DLLIMPORT DATASMITHCORE_API=DLLIMPORT DIRECTLINK_API=DLLIMPORT DATASMITHEXPORTERUI_API=DLLIMPORT DATASMITHVALIDATOR_API=DLLIMPORT COREUOBJECT_API=DLLIMPORT
+UE_PREPROCESSOR_DEFINITIONS = NOINITCRASHREPORTER=1 UE_BUILD_DEVELOPMENT=1 UE_BUILD_MINIMAL=1 WITH_EDITOR=0 WITH_EDITORONLY_DATA=0 WITH_SERVER_CODE=1 WITH_ENGINE=0 WITH_UNREAL_DEVELOPER_TOOLS=0 WITH_PLUGIN_SUPPORT=0 IS_MONOLITHIC=1 IS_PROGRAM=1 PLATFORM_MAC=1 PLATFORM_APPLE=1 UE_BUILD_DEVELOPMENT_WITH_DEBUGGAME=0 UBT_COMPILED_PLATFORM=Mac CORE_API=DLLIMPORT DATASMITHEXPORTER_API=DLLIMPORT DATASMITHCORE_API=DLLIMPORT DIRECTLINK_API=DLLIMPORT DATASMITHEXPORTERUI_API=DLLIMPORT DATASMITHVALIDATOR_API=DLLIMPORT COREUOBJECT_API=DLLIMPORT
 
 //:configuration = Debug
 GCC_PREPROCESSOR_DEFINITIONS = $(ADDON_PREPROCESSOR_DEFINITIONS) $(AC_PREPROCESSOR_DEFINITIONS) $(UE_PREPROCESSOR_DEFINITIONS)
diff --git a/Engine/Source/Programs/Enterprise/Datasmith/DatasmithARCHICADExporter/DatsmithUE4ArchiCAD/DatasmithUE4ArchiCAD.Target.cs b/Engine/Source/Programs/Enterprise/Datasmith/DatasmithARCHICADExporter/DatsmithUE4ArchiCAD/DatasmithUE4ArchiCAD.Target.cs
--- a/Engine/Source/Programs/Enterprise/Datasmith/DatasmithARCHICADExporter/DatsmithUE4ArchiCAD/DatasmithUE4ArchiCAD.Target.cs
+++ b/Engine/Source/Programs/Enterprise/Datasmith/DatasmithARCHICADExporter/DatsmithUE4ArchiCAD/DatasmithUE4ArchiCAD.Target.cs
@@ -36,5 +36,8 @@ public class DatasmithUE4ArchiCADTarget : TargetRules
 		bHasExports = true;
 
 		GlobalDefinitions.Add("UE_EXTERNAL_PROFILING_ENABLED=0"); // For DirectLinkUI (see FDatasmithExporterManager::FInitOptions)
+
+		// Disable CrashReporter
+		GlobalDefinitions.Add("NOINITCRASHREPORTER=1");
 	}
 }
diff --git a/Engine/Source/Programs/Enterprise/Datasmith/DatasmithFacadeCSharp/DatasmithFacadeCSharp.Target.cs b/Engine/Source/Programs/Enterprise/Datasmith/DatasmithFacadeCSharp/DatasmithFacadeCSharp.Target.cs
--- a/Engine/Source/Programs/Enterprise/Datasmith/DatasmithFacadeCSharp/DatasmithFacadeCSharp.Target.cs
+++ b/Engine/Source/Programs/Enterprise/Datasmith/DatasmithFacadeCSharp/DatasmithFacadeCSharp.Target.cs
@@ -37,6 +37,9 @@ public class DatasmithFacadeCSharpTarget : TargetRules
 		 */
 		GlobalDefinitions.Add("UE_EXTERNAL_PROFILING_ENABLED=0");
 
+		// Disable CrashReporter
+		GlobalDefinitions.Add("NOINITCRASHREPORTER=1");
+
 		AddPreBuildSteps(Target.Platform);
 		AddPostBuildSteps(Target.Platform);
 	}
diff --git a/Engine/Source/Programs/Enterprise/Datasmith/DatasmithMaxExporter/DatasmithMax2017.Target.cs b/Engine/Source/Programs/Enterprise/Datasmith/DatasmithMaxExporter/DatasmithMax2017.Target.cs
--- a/Engine/Source/Programs/Enterprise/Datasmith/DatasmithMaxExporter/DatasmithMax2017.Target.cs
+++ b/Engine/Source/Programs/Enterprise/Datasmith/DatasmithMaxExporter/DatasmithMax2017.Target.cs
@@ -34,7 +34,11 @@ public abstract class DatasmithMaxBaseTarget : TargetRules
 		bHasExports = true;
 		bForceEnableExceptions = true;
 
-		GlobalDefinitions.Add("UE_EXTERNAL_PROFILING_ENABLED=0"); // For DirectLinkUI (see FDatasmithExporterManager::FInitOptions)
+		// For DirectLinkUI (see FDatasmithExporterManager::FInitOptions)
+		GlobalDefinitions.Add("UE_EXTERNAL_PROFILING_ENABLED=0");
+		
+		// Disable CrashReporter
+		GlobalDefinitions.Add("NOINITCRASHREPORTER=1");
 
 		// todo: remove?
 		// bSupportEditAndContinue = true;
diff --git a/Engine/Source/Programs/Enterprise/Datasmith/DatasmithSketchUpRubyExporter/DatasmithSketchUpRuby2020.Target.cs b/Engine/Source/Programs/Enterprise/Datasmith/DatasmithSketchUpRubyExporter/DatasmithSketchUpRuby2020.Target.cs
--- a/Engine/Source/Programs/Enterprise/Datasmith/DatasmithSketchUpRubyExporter/DatasmithSketchUpRuby2020.Target.cs
+++ b/Engine/Source/Programs/Enterprise/Datasmith/DatasmithSketchUpRubyExporter/DatasmithSketchUpRuby2020.Target.cs
@@ -29,6 +29,9 @@ public abstract class DatasmithSketchUpRubyBaseTarget : TargetRules
 		bForceEnableExceptions = true;
 
 		GlobalDefinitions.Add("UE_EXTERNAL_PROFILING_ENABLED=0"); // For DirectLinkUI (see FDatasmithExporterManager::FInitOptions)
+
+		// Disable CrashReporter
+		GlobalDefinitions.Add("NOINITCRASHREPORTER=1");
 	}
 
 	protected void AddCopyPostBuildStep(TargetInfo Target)
