[Sketchup]
Datasmith plugin : support Sketchup 2025.

diff --git a/Engine/Source/Programs/Enterprise/Datasmith/DatasmithSketchUpRubyExporter/Build/build_mac.sh b/Engine/Source/Programs/Enterprise/Datasmith/DatasmithSketchUpRubyExporter/Build/build_mac.sh
--- a/Engine/Source/Programs/Enterprise/Datasmith/DatasmithSketchUpRubyExporter/Build/build_mac.sh
+++ b/Engine/Source/Programs/Enterprise/Datasmith/DatasmithSketchUpRubyExporter/Build/build_mac.sh
@@ -94,6 +94,7 @@ BuildSketchUpPlugin 2021 SDK_Mac_2021-0-338 "x86_64" "$Intermediatex86_64Path"
 BuildSketchUpPlugin 2022 SDK_Mac_2022-0-353 "x86_64 arm64" "$IntermediateUniversalPath"
 BuildSketchUpPlugin 2023 SDK_Mac_2023-0-366 "x86_64 arm64" "$IntermediateUniversalPath"
 BuildSketchUpPlugin 2024 SDK_Mac_2024-0-483 "x86_64 arm64" "$IntermediateUniversalPath"
+BuildSketchUpPlugin 2025 SDK_Mac_2025-0-574 "x86_64 arm64" "$IntermediateUniversalPath"
 
 
 # install_name_tool -change @rpath/DatasmithSDK.dylib @loader_path/Dylibs/DatasmithSDK.dylib DatasmithSketchUp.bundle 
diff --git a/Engine/Source/Programs/Enterprise/Datasmith/DatasmithSketchUpRubyExporter/DatasmithSketchUpRuby2025.Build.cs b/Engine/Source/Programs/Enterprise/Datasmith/DatasmithSketchUpRubyExporter/DatasmithSketchUpRuby2025.Build.cs
new file mode 100644
--- /dev/null
+++ b/Engine/Source/Programs/Enterprise/Datasmith/DatasmithSketchUpRubyExporter/DatasmithSketchUpRuby2025.Build.cs
@@ -0,0 +1,31 @@
+// Copyright Epic Games, Inc. All Rights Reserved.
+
+using System.IO;
+
+namespace UnrealBuildTool.Rules
+{
+
+	[SupportedPlatforms("Win64", "Mac")]
+	public class DatasmithSketchUpRuby2025 : DatasmithSketchUpRubyBase
+	{
+		public DatasmithSketchUpRuby2025(ReadOnlyTargetRules Target)
+			: base(Target)
+		{
+			PrivateDefinitions.Add("SKP_SDK_2025");
+		}
+
+		public override string GetSketchUpSDKFolder()
+		{
+			return "SDK_WIN_x64_2025-0-575";
+		}
+
+		public override string GetSketchUpEnvVar()
+		{
+			return "SKP_SDK_2025";
+		}
+		public override string GetRubyLibName()
+		{
+			return "x64-ucrt-ruby320.lib";
+		}
+	}
+}
diff --git a/Engine/Source/Programs/Enterprise/Datasmith/DatasmithSketchUpRubyExporter/DatasmithSketchUpRuby2025.Target.cs b/Engine/Source/Programs/Enterprise/Datasmith/DatasmithSketchUpRubyExporter/DatasmithSketchUpRuby2025.Target.cs
new file mode 100644
--- /dev/null
+++ b/Engine/Source/Programs/Enterprise/Datasmith/DatasmithSketchUpRubyExporter/DatasmithSketchUpRuby2025.Target.cs
@@ -0,0 +1,16 @@
+// Copyright Epic Games, Inc. All Rights Reserved.
+
+using UnrealBuildTool;
+
+[SupportedPlatforms("Win64")]
+public class DatasmithSketchUpRuby2025Target : DatasmithSketchUpRubyBaseTarget
+{
+	public DatasmithSketchUpRuby2025Target(TargetInfo Target)
+		: base(Target)
+	{
+		LaunchModuleName = "DatasmithSketchUpRuby2025";
+		ExeBinariesSubFolder = @"SketchUpRuby/2025";
+
+		AddCopyPostBuildStep(Target);
+	}
+}
