This fixes the issue with minimal code change. Alternatively I could use the material color in the Hash but it'd require to change more code around the plugin to inject those values.

diff --git a/Engine/Source/Programs/Enterprise/Datasmith/DatasmithSketchUpRubyExporter/Private/DatasmithSketchUpTexture.cpp b/Engine/Source/Programs/Enterprise/Datasmith/DatasmithSketchUpRubyExporter/Private/DatasmithSketchUpTexture.cpp
--- a/Engine/Source/Programs/Enterprise/Datasmith/DatasmithSketchUpRubyExporter/Private/DatasmithSketchUpTexture.cpp
+++ b/Engine/Source/Programs/Enterprise/Datasmith/DatasmithSketchUpRubyExporter/Private/DatasmithSketchUpTexture.cpp
@@ -216,6 +216,14 @@ void FTextureCollection::AcquireImage(FTexture& Texture)
 	MD5.Update(reinterpret_cast<const uint8*>(&Texture.bColorized), sizeof(Texture.bColorized)); // Colorized flag affects resulting texture image
 	MD5.Update(Data.GetData(), Data.Num());
 
+	// colorized textures need an additional differentiation since the texture should be baked
+	// once per different tint it has applied in the scene
+	// Here we're using MaterialName, which works fine, but we could also use the tint color
+	if(Texture.bColorized)
+	{
+		MD5.Update(reinterpret_cast<const uint8*>(&Texture.MaterialName), sizeof(Texture.MaterialName));
+	}	
+
 	FMD5Hash TextureHash;
 	TextureHash.Set(MD5);
 
