Fix: when using DirectLink, if a cut opening component changed, the related parent components didn't

diff --git a/Engine/Source/Programs/Enterprise/Datasmith/DatasmithSketchUpRubyExporter/Private/DatasmithSketchUpComponent.cpp b/Engine/Source/Programs/Enterprise/Datasmith/DatasmithSketchUpRubyExporter/Private/DatasmithSketchUpComponent.cpp
--- a/Engine/Source/Programs/Enterprise/Datasmith/DatasmithSketchUpRubyExporter/Private/DatasmithSketchUpComponent.cpp
+++ b/Engine/Source/Programs/Enterprise/Datasmith/DatasmithSketchUpRubyExporter/Private/DatasmithSketchUpComponent.cpp
@@ -467,6 +467,7 @@ void FComponentDefinition::Parse(FExportContext& Context)
 
 	// Get whether or not the source SketchUp component behaves like a billboard.
 	bSketchupSourceFaceCamera = SComponentBehavior.component_always_face_camera;
+	bIsCutOpening = SComponentBehavior.component_cuts_opening;
 }
 
 void FComponentInstance::SetupActor(FExportContext& Context, FNodeOccurence& Node)
@@ -840,10 +841,20 @@ void FComponentInstance::UpdateOccurrenceTransformation(FExportContext& Context,
 
 	bool bTransformChanged = !DatasmithSketchUpUtils::CompareSUTransformations(WorldTransform, Node.WorldTransformSource);
 
+	// if the component's transform has changed, in some situations
+	// we might need to invalidate the parent component geometry as well
+	if(bTransformChanged)
+	{
+		FComponentDefinition* CompDef = static_cast<FComponentDefinition*>(GetDefinition());
+		if(CompDef && CompDef->bIsCutOpening)
+		{
+			Parent->InvalidateDefinitionGeometry();
+		}
+	}
+
 	Node.WorldTransformSource = WorldTransform;
 	Node.WorldTransform = WorldTransform;
 
-
 	auto GetNodePath = [&Node]()
 	{
 		TArray<FString> NamePath;
diff --git a/Engine/Source/Programs/Enterprise/Datasmith/DatasmithSketchUpRubyExporter/Private/DatasmithSketchUpComponent.h b/Engine/Source/Programs/Enterprise/Datasmith/DatasmithSketchUpRubyExporter/Private/DatasmithSketchUpComponent.h
--- a/Engine/Source/Programs/Enterprise/Datasmith/DatasmithSketchUpRubyExporter/Private/DatasmithSketchUpComponent.h
+++ b/Engine/Source/Programs/Enterprise/Datasmith/DatasmithSketchUpRubyExporter/Private/DatasmithSketchUpComponent.h
@@ -215,6 +215,9 @@ namespace DatasmithSketchUp
 		// Whether or not the source SketchUp component behaves like a billboard, always presenting a 2D surface perpendicular to the direction of camera.
 		bool bSketchupSourceFaceCamera = false;
 
+		// cut opening is used to create non destructive boolean. 
+		bool bIsCutOpening = false;
+
 	private:
 		SUComponentDefinitionRef ComponentDefinitionRef;
 
