#TM-18008

Fix: material tiling is inconsistent for geom face and node inherited materials

diff --git a/Engine/Source/Programs/Enterprise/Datasmith/DatasmithSketchUpRubyExporter/Private/DatasmithSketchUpExportContext.cpp b/Engine/Source/Programs/Enterprise/Datasmith/DatasmithSketchUpRubyExporter/Private/DatasmithSketchUpExportContext.cpp
--- a/Engine/Source/Programs/Enterprise/Datasmith/DatasmithSketchUpRubyExporter/Private/DatasmithSketchUpExportContext.cpp
+++ b/Engine/Source/Programs/Enterprise/Datasmith/DatasmithSketchUpRubyExporter/Private/DatasmithSketchUpExportContext.cpp
@@ -835,6 +835,17 @@ TSharedPtr<FMaterial> FMaterialCollection::CreateMaterial(SUMaterialRef Material
 	return Material;
 }
 
+TSharedPtr<FMaterial> FMaterialCollection::FindMaterialWithId(FEntityIDType MaterialId)
+{
+	for (DatasmithSketchUp::FMaterial* Material : Materials)
+	{
+		if(Material->GetMaterialID().EntityID == MaterialId.EntityID)
+			return MakeShareable(Material);
+	}
+
+	return nullptr;
+}
+
 void FMaterialCollection::SetMeshActorOverrideMaterial(FNodeOccurence& Node, DatasmithSketchUp::FEntitiesGeometry& EntitiesGeometry, const TSharedPtr<IDatasmithMeshActorElement>& MeshActor)
 {
 	// SketchUp has 'material override' only for single material - 'Default' material or material from default('Layer0') layer in ColorByTag mode. 
@@ -948,7 +959,6 @@ TSharedPtr<FMaterial> FLayerMaterials::FindOrCreateMaterialForLayer(FLayerIDType
 		if (SULayerGetMaterial(LayerRef, &MaterialRef) == SU_ERROR_NONE)
 		{
 			TSharedPtr<FMaterial> Material = Context.Materials.CreateMaterial(MaterialRef);
-			Material->bGeometryHasScalingBakedIntoUvs = false;
 			MaterialForLayerId.Add(LayerID, Material);
 			LayerIdForMaterial.Add(Material.Get(), LayerID);
 			return Material;
diff --git a/Engine/Source/Programs/Enterprise/Datasmith/DatasmithSketchUpRubyExporter/Private/DatasmithSketchUpExportContext.h b/Engine/Source/Programs/Enterprise/Datasmith/DatasmithSketchUpRubyExporter/Private/DatasmithSketchUpExportContext.h
--- a/Engine/Source/Programs/Enterprise/Datasmith/DatasmithSketchUpRubyExporter/Private/DatasmithSketchUpExportContext.h
+++ b/Engine/Source/Programs/Enterprise/Datasmith/DatasmithSketchUpRubyExporter/Private/DatasmithSketchUpExportContext.h
@@ -287,6 +287,7 @@ namespace DatasmithSketchUp
 		void Update();
 
 		TSharedPtr<FMaterial> CreateMaterial(SUMaterialRef SMaterialDefinitionRef);
+		TSharedPtr<FMaterial> FindMaterialWithId(FEntityIDType MaterialId);
 		void InvalidateMaterial(SUMaterialRef SMaterialDefinitionRef);
 		bool RemoveMaterial(FMaterial* Material);
 
diff --git a/Engine/Source/Programs/Enterprise/Datasmith/DatasmithSketchUpRubyExporter/Private/DatasmithSketchUpMaterial.cpp b/Engine/Source/Programs/Enterprise/Datasmith/DatasmithSketchUpRubyExporter/Private/DatasmithSketchUpMaterial.cpp
--- a/Engine/Source/Programs/Enterprise/Datasmith/DatasmithSketchUpRubyExporter/Private/DatasmithSketchUpMaterial.cpp
+++ b/Engine/Source/Programs/Enterprise/Datasmith/DatasmithSketchUpRubyExporter/Private/DatasmithSketchUpMaterial.cpp
@@ -8,7 +8,6 @@
 #include "DatasmithSketchUpMetadata.h"
 #include "DatasmithSketchUpMesh.h"
 #include "DatasmithSketchUpString.h"
-#include "DatasmithSketchUpSummary.h"
 #include "DatasmithSketchUpTexture.h"
 #include "DatasmithSketchUpUtils.h"
 
@@ -22,7 +21,6 @@
 #include "SketchUpAPI/model/texture.h"
 #include "SketchUpAPI/model/typed_value.h"
 #include "SketchUpAPI/model/rendering_options.h"
-#include "DatasmithSketchUpSDKCeases.h"
 
 // Datasmith SDK.
 
@@ -137,8 +135,7 @@ namespace DatasmithSketchUp
 		FExportContext& Context,
 		FExtractedMaterial& InMaterial,
 		TCHAR const* InMaterialName,
-		FTexture* Texture,
-		bool bInScaleTexture
+		FTexture* Texture
 	)
 	{
 		// Create a Datasmith material element for the material definition.
@@ -163,29 +160,20 @@ namespace DatasmithSketchUp
 			ExpressionTexture->SetTexturePathName(Texture->GetDatasmithElementName());
 
 			// Apply texture scaling
-			if (bInScaleTexture && !Texture->TextureScale.Equals(FVector2D::UnitVector))
-			{
-				IDatasmithMaterialExpressionFunctionCall* UVEditExpression = DatasmithMaterialElementPtr->AddMaterialExpression< IDatasmithMaterialExpressionFunctionCall >();
-				UVEditExpression->SetFunctionPathName(TEXT("/DatasmithContent/Materials/UVEdit.UVEdit"));
+			IDatasmithMaterialExpressionFunctionCall* UVEditExpression = DatasmithMaterialElementPtr->AddMaterialExpression< IDatasmithMaterialExpressionFunctionCall >();
+			UVEditExpression->SetFunctionPathName(TEXT("/DatasmithContent/Materials/UVEdit.UVEdit"));
 
-				UVEditExpression->ConnectExpression(ExpressionTexture->GetInputCoordinate());
+			UVEditExpression->ConnectExpression(ExpressionTexture->GetInputCoordinate());
 
-				// Tiling
-				IDatasmithMaterialExpressionColor* TilingValue = DatasmithMaterialElementPtr->AddMaterialExpression< IDatasmithMaterialExpressionColor >();
-				TilingValue->SetName(TEXT("UV Tiling"));
-				TilingValue->GetColor() = FLinearColor(Texture->TextureScale.X, Texture->TextureScale.Y, 0.f);
+			// Tiling
+			IDatasmithMaterialExpressionColor* TilingValue = DatasmithMaterialElementPtr->AddMaterialExpression< IDatasmithMaterialExpressionColor >();
+			TilingValue->SetName(TEXT("UV Tiling"));
+			TilingValue->GetColor() = FLinearColor(Texture->TextureScale.X, Texture->TextureScale.Y, 0.f);
+			TilingValue->ConnectExpression(*UVEditExpression->GetInput(2));
 
-				TilingValue->ConnectExpression(*UVEditExpression->GetInput(2));
-
-				//IDatasmithMaterialExpressionColor* OffsetValue = MaterialElement->AddMaterialExpression< IDatasmithMaterialExpressionColor >();
-				//OffsetValue->SetName(TEXT("UV Offset"));
-				//OffsetValue->GetColor() = FLinearColor(0.f, 0.f, 0.f);
-				//OffsetValue->ConnectExpression(*UVEditExpression->GetInput(7));
-
-				IDatasmithMaterialExpressionTextureCoordinate* TextureCoordinateExpression = DatasmithMaterialElementPtr->AddMaterialExpression< IDatasmithMaterialExpressionTextureCoordinate >();
-				TextureCoordinateExpression->SetCoordinateIndex(0);
-				TextureCoordinateExpression->ConnectExpression(*UVEditExpression->GetInput(0));
-			}
+			IDatasmithMaterialExpressionTextureCoordinate* TextureCoordinateExpression = DatasmithMaterialElementPtr->AddMaterialExpression< IDatasmithMaterialExpressionTextureCoordinate >();
+			TextureCoordinateExpression->SetCoordinateIndex(0);
+			TextureCoordinateExpression->ConnectExpression(*UVEditExpression->GetInput(0));
 
 			ExpressionTexture->ConnectExpression(DatasmithMaterialElementPtr->GetBaseColor());
 
@@ -418,13 +406,13 @@ void FMaterial::Update(FExportContext& Context)
 	MaterialDirectlyAppliedToMeshes.RemoveDatasmithElement(Context);
 	if (MaterialDirectlyAppliedToMeshes.HasUsers())
 	{
-		MaterialDirectlyAppliedToMeshes.DatasmithElement = CreateMaterialElement(Context, ExtractedMaterial, *ExtractedMaterial.LocalizedMaterialName, Texture, !bGeometryHasScalingBakedIntoUvs);
+		MaterialDirectlyAppliedToMeshes.DatasmithElement = CreateMaterialElement(Context, ExtractedMaterial, *ExtractedMaterial.LocalizedMaterialName, Texture);
 	}
 
 	MaterialInheritedByNodes.RemoveDatasmithElement(Context);
 	if (MaterialInheritedByNodes.HasUsers())
 	{
-		MaterialInheritedByNodes.DatasmithElement = CreateMaterialElement(Context, ExtractedMaterial, *ExtractedMaterial.InheritedMaterialName, Texture, true);
+		MaterialInheritedByNodes.DatasmithElement = CreateMaterialElement(Context, ExtractedMaterial, *ExtractedMaterial.InheritedMaterialName, Texture);
 	}
 
 	bInvalidated = false;
diff --git a/Engine/Source/Programs/Enterprise/Datasmith/DatasmithSketchUpRubyExporter/Private/DatasmithSketchUpMaterial.h b/Engine/Source/Programs/Enterprise/Datasmith/DatasmithSketchUpRubyExporter/Private/DatasmithSketchUpMaterial.h
--- a/Engine/Source/Programs/Enterprise/Datasmith/DatasmithSketchUpRubyExporter/Private/DatasmithSketchUpMaterial.h
+++ b/Engine/Source/Programs/Enterprise/Datasmith/DatasmithSketchUpRubyExporter/Private/DatasmithSketchUpMaterial.h
@@ -9,6 +9,7 @@
 #include "DatasmithSketchUpSDKCeases.h"
 
 // Datasmith SDK.
+#include "DatasmithSketchUpUtils.h"
 #include "Containers/Map.h"
 #include "Containers/UnrealString.h"
 #include "Misc/SecureHash.h"
@@ -126,18 +127,18 @@ namespace DatasmithSketchUp
 
 		FMD5Hash ComputeHash(FExportContext& Context);
 
+		FMaterialIDType GetMaterialID() const
+		{
+			return DatasmithSketchUpUtils::GetMaterialID(MaterialRef);
+		}
+		
 		// Material can be directly applied to a face in SketchUp
 		FMaterialOccurrence MaterialDirectlyAppliedToMeshes;
 		// In case face has Default material assigned it inherits material set to it's parent(in general - first non-Default material in ancestors chain)
 		FMaterialOccurrence MaterialInheritedByNodes;
-
-		// Whether UVs of geometry have texture scaling baked, when applying this material
-		// Scaling is baked for regular materials, applied directly to faces so regular materials assigned to static meshes need to take this into account.
-		bool bGeometryHasScalingBakedIntoUvs = true;
 	private:
 
 		SUMaterialRef MaterialRef;
-
 		int32 EntityId;
 
 		FTexture* Texture = nullptr;
diff --git a/Engine/Source/Programs/Enterprise/Datasmith/DatasmithSketchUpRubyExporter/Private/DatasmithSketchUpMesh.cpp b/Engine/Source/Programs/Enterprise/Datasmith/DatasmithSketchUpRubyExporter/Private/DatasmithSketchUpMesh.cpp
--- a/Engine/Source/Programs/Enterprise/Datasmith/DatasmithSketchUpRubyExporter/Private/DatasmithSketchUpMesh.cpp
+++ b/Engine/Source/Programs/Enterprise/Datasmith/DatasmithSketchUpRubyExporter/Private/DatasmithSketchUpMesh.cpp
@@ -31,10 +31,13 @@
 
 // Datasmith SDK.
 
+#include <SketchUpAPI/model/texture.h>
+
 #include "DatasmithMesh.h"
 #include "DatasmithMeshExporter.h"
 #include "DatasmithSceneExporter.h"
 #include "DatasmithSceneFactory.h"
+#include "DatasmithSketchUpTexture.h"
 #include "DatasmithUtils.h"
 #include "Misc/SecureHash.h"
 
@@ -212,6 +215,28 @@ namespace DatasmithSketchUp
 		MeshVertexUVQs.Reserve(MeshVertexUVQs.Num() + VertexPointCount);
 		MeshTriangleIndices.Reserve(MeshTriangleIndices.Num() + TriangleVertexIndexCount);
 
+
+		// we use TextureScale to modify the scale of the mesh UV
+		// while we also modify texture tiling in the material, this is required
+		// to have consistent of materials across "node inherited" and "face applied materials"
+		// for node materials, tiling and uv scale is 1, so no changes
+		// for face applied materials, the tiling "cancels" between uv material tiling/scaling , so there is no
+		// visual change, but this allows to have a material in Datasmith with tiling that if we apply to other surfaces
+		// will keep a consistent visual
+		FVector2D TextureScale = FVector2D::One();
+		SUTextureRef TextureRef = SU_INVALID;
+		SUMaterialGetTexture(bUseFrontMaterial ? FrontMaterialRef : BackMaterialRef, &TextureRef);
+
+		if(SUIsValid(TextureRef))
+		{
+			size_t TextureWidth = 0;
+			size_t TextureHeight = 0;
+			double TextureSScale = 1.0;
+			double TextureTScale = 1.0;
+			SUTextureGetDimensions(TextureRef, &TextureWidth, &TextureHeight, &TextureSScale, &TextureTScale); // we can ignore the returned SU_RESULT
+			TextureScale = FVector2D(TextureSScale, TextureTScale);
+		}
+
 		if (bUseFrontMaterial)
 		{
 			// Combine the mesh vertex normals.
@@ -223,6 +248,8 @@ namespace DatasmithSketchUp
 				// Retrieve the front UVQ texture coordinate of the mesh vertex.
 				SUUVQ SUVQ;
 				SUUVHelperGetFrontUVQ(UVHelperRef, &VertexPoints[VertexPointIndex], &SUVQ); // we can ignore the returned SU_RESULT
+				SUVQ.u /= TextureScale.X;
+				SUVQ.v /= TextureScale.Y;
 
 				MeshVertexUVQs.Add(SUVQ);
 			}
@@ -259,7 +286,9 @@ namespace DatasmithSketchUp
 				// Retrieve the back UVQ texture coordinate of the mesh vertex.
 				SUUVQ UVQ;
 				SUUVHelperGetBackUVQ(UVHelperRef, &VertexPoints[VertexPointIndex], &UVQ); // we can ignore the returned SU_RESULT
-
+				UVQ.u /= TextureScale.X;
+				UVQ.v /= TextureScale.Y;
+				
 				MeshVertexUVQs.Add(UVQ);
 			}
 
