Fix:
one header was missing and it was causing this build error

fatal error C1189:
#error:  Nesting AllowWindowsPlatformTypes.h is not allowed!
#error Nesting AllowWindowsPlatformTypes.h is not allowed

diff --git a/Engine/Source/Programs/Enterprise/Datasmith/DatasmithSketchUpRubyExporter/Private/DatasmithSketchUpMaterial.cpp b/Engine/Source/Programs/Enterprise/Datasmith/DatasmithSketchUpRubyExporter/Private/DatasmithSketchUpMaterial.cpp
--- a/Engine/Source/Programs/Enterprise/Datasmith/DatasmithSketchUpRubyExporter/Private/DatasmithSketchUpMaterial.cpp
+++ b/Engine/Source/Programs/Enterprise/Datasmith/DatasmithSketchUpRubyExporter/Private/DatasmithSketchUpMaterial.cpp
@@ -21,6 +21,7 @@
 #include "SketchUpAPI/model/texture.h"
 #include "SketchUpAPI/model/typed_value.h"
 #include "SketchUpAPI/model/rendering_options.h"
+#include "DatasmithSketchUpSDKCeases.h"
 
 // Datasmith SDK.
 
