[Naviswork] Add support for Naviswork 2026

diff --git a/Engine/Source/Programs/Enterprise/Datasmith/DatasmithNavisworksExporter/DatasmithNavisworksPlugin/DatasmithNavisworksPlugin.csproj b/Engine/Source/Programs/Enterprise/Datasmith/DatasmithNavisworksExporter/DatasmithNavisworksPlugin/DatasmithNavisworksPlugin.csproj
--- a/Engine/Source/Programs/Enterprise/Datasmith/DatasmithNavisworksExporter/DatasmithNavisworksPlugin/DatasmithNavisworksPlugin.csproj
+++ b/Engine/Source/Programs/Enterprise/Datasmith/DatasmithNavisworksExporter/DatasmithNavisworksPlugin/DatasmithNavisworksPlugin.csproj
@@ -8,6 +8,8 @@
     <ProjectGuid>{D16F2120-CCE4-4FCC-92C8-86124AB1FA25}</ProjectGuid>
     <TargetFrameworkProfile />
   </PropertyGroup>
+  <Import Project="..\Navisworks2026.props" Condition="'$(Configuration)|$(Platform)'=='Debug2026|AnyCPU'" />
+  <Import Project="..\Navisworks2026.props" Condition="'$(Configuration)|$(Platform)'=='Release2026|AnyCPU'" />
   <Import Project="..\Navisworks2025.props" Condition="'$(Configuration)|$(Platform)'=='Debug2025|AnyCPU'" />
   <Import Project="..\Navisworks2025.props" Condition="'$(Configuration)|$(Platform)'=='Release2025|AnyCPU'" />
   <Import Project="..\Navisworks2024.props" Condition="'$(Configuration)|$(Platform)'=='Debug2024|AnyCPU'" />
@@ -32,6 +34,30 @@
     <TargetFrameworkVersion>v4.8</TargetFrameworkVersion>
     <FileAlignment>512</FileAlignment>
   </PropertyGroup>
+  <PropertyGroup Condition=" '$(Configuration)|$(Platform)' == 'Debug2026|AnyCPU' ">
+    <DebugSymbols>true</DebugSymbols>
+    <DebugType>full</DebugType>
+    <Optimize>false</Optimize>
+    <OutputPath>bin\Debug2026\</OutputPath>
+    <DefineConstants>DEBUG;TRACE;DATASMITH_NAVISWORKS_2026_OR_GREATER</DefineConstants>
+    <ErrorReport>prompt</ErrorReport>
+    <WarningLevel>4</WarningLevel>
+    <StartAction>Program</StartAction>
+    <StartProgram>$(ProgramW6432)\Autodesk\Navisworks Manage $(Navisworks_Version)\Roamer.exe</StartProgram>
+    <EnableUnmanagedDebugging>true</EnableUnmanagedDebugging>
+    <AllowUnsafeBlocks>true</AllowUnsafeBlocks>
+  </PropertyGroup>
+  <PropertyGroup Condition=" '$(Configuration)|$(Platform)' == 'Release2026|AnyCPU' ">
+    <DebugType>pdbonly</DebugType>
+    <Optimize>true</Optimize>
+    <OutputPath>bin\Release2026\</OutputPath>
+    <DefineConstants>TRACE;DATASMITH_NAVISWORKS_2026_OR_GREATER</DefineConstants>
+    <ErrorReport>prompt</ErrorReport>
+    <WarningLevel>4</WarningLevel>
+    <StartAction>Program</StartAction>
+    <StartProgram>$(ProgramW6432)\Autodesk\Navisworks Manage $(Navisworks_Version)\Roamer.exe</StartProgram>
+    <AllowUnsafeBlocks>true</AllowUnsafeBlocks>
+  </PropertyGroup>
   <PropertyGroup Condition=" '$(Configuration)|$(Platform)' == 'Debug2025|AnyCPU' ">
     <DebugSymbols>true</DebugSymbols>
     <DebugType>full</DebugType>
diff --git a/Engine/Source/Programs/Enterprise/Datasmith/DatasmithNavisworksExporter/DatasmithNavisworksPlugin/PackageContents.xml b/Engine/Source/Programs/Enterprise/Datasmith/DatasmithNavisworksExporter/DatasmithNavisworksPlugin/PackageContents.xml
--- a/Engine/Source/Programs/Enterprise/Datasmith/DatasmithNavisworksExporter/DatasmithNavisworksPlugin/PackageContents.xml
+++ b/Engine/Source/Programs/Enterprise/Datasmith/DatasmithNavisworksExporter/DatasmithNavisworksPlugin/PackageContents.xml
@@ -29,4 +29,8 @@
 	    <RuntimeRequirements OS="Win64" Platform="NAVMAN|NAVSIM" SeriesMin="Nw22" SeriesMax="Nw22" />
         <ComponentEntry AppName="Datasmith Exporter for Autodesk® Navisworks® 2025" AppType="ManagedPlugin" ModuleName="./DatasmithNavisworksPlugin2025.dll" />
     </Components>
+	<Components Description="2026">
+		<RuntimeRequirements OS="Win64" Platform="NAVMAN|NAVSIM" SeriesMin="Nw23" SeriesMax="Nw23" />
+		<ComponentEntry AppName="Datasmith Exporter for Autodesk® Navisworks® 2026" AppType="ManagedPlugin" ModuleName="./DatasmithNavisworksPlugin2026.dll" />
+	</Components>
 </ApplicationPackage>
diff --git a/Engine/Source/Programs/Enterprise/Datasmith/DatasmithNavisworksExporter/UBTModules/DatasmithNavisworks2026.Build.cs b/Engine/Source/Programs/Enterprise/Datasmith/DatasmithNavisworksExporter/UBTModules/DatasmithNavisworks2026.Build.cs
new file mode 100644
--- /dev/null
+++ b/Engine/Source/Programs/Enterprise/Datasmith/DatasmithNavisworksExporter/UBTModules/DatasmithNavisworks2026.Build.cs
@@ -0,0 +1,20 @@
+// Copyright Epic Games, Inc. All Rights Reserved.
+
+using System.IO;
+
+namespace UnrealBuildTool.Rules
+{
+	[SupportedPlatforms("Win64")]
+	public class DatasmithNavisworks2026 : DatasmithNavisworksBase
+	{
+		public DatasmithNavisworks2026(ReadOnlyTargetRules Target)
+			: base(Target)
+		{
+		}
+
+		public override string GetNavisworksVersion()
+		{
+			return "2026";
+		}
+	}
+}
diff --git a/Engine/Source/Programs/Enterprise/Datasmith/DatasmithNavisworksExporter/UBTModules/DatasmithNavisworks2026.Target.cs b/Engine/Source/Programs/Enterprise/Datasmith/DatasmithNavisworksExporter/UBTModules/DatasmithNavisworks2026.Target.cs
new file mode 100644
--- /dev/null
+++ b/Engine/Source/Programs/Enterprise/Datasmith/DatasmithNavisworksExporter/UBTModules/DatasmithNavisworks2026.Target.cs
@@ -0,0 +1,14 @@
+// Copyright Epic Games, Inc. All Rights Reserved.
+
+using UnrealBuildTool;
+using System.IO;
+
+public class DatasmithNavisworks2026Target : DatasmithNavisworksBaseTarget
+{
+	public DatasmithNavisworks2026Target(TargetInfo Target)
+		: base(Target)
+	{
+	}
+
+	public override string GetVersion() { return "2026"; }
+}
