Fixed issues with Datasmith exporters:
  - Remove compilation of 3dsMax exporters for 2017 to 2021.
  - Remove use of prerequisite installation packages
 -  Fixed newly added Navisworks 2026 build

#rnx

diff --git a/Engine/Source/Programs/Enterprise/Datasmith/DatasmithMaxExporter/DatasmithMax2017.Build.cs b/Engine/Source/Programs/Enterprise/Datasmith/DatasmithMaxExporter/DatasmithMax2017.Build.cs
--- a/Engine/Source/Programs/Enterprise/Datasmith/DatasmithMaxExporter/DatasmithMax2017.Build.cs
+++ b/Engine/Source/Programs/Enterprise/Datasmith/DatasmithMaxExporter/DatasmithMax2017.Build.cs
@@ -27,6 +27,7 @@ namespace UnrealBuildTool.Rules
 					"UdpMessaging", // required for DirectLink networking
 					"UEOpenExr",
 
+					"IntelTBB",
 					"Slate",
 					"SlateCore",
 				}
diff --git a/Engine/Source/Programs/Enterprise/Datasmith/DatasmithMaxExporter/DatasmithMax2017.Target.cs b/Engine/Source/Programs/Enterprise/Datasmith/DatasmithMaxExporter/DatasmithMax2017.Target.cs
--- a/Engine/Source/Programs/Enterprise/Datasmith/DatasmithMaxExporter/DatasmithMax2017.Target.cs
+++ b/Engine/Source/Programs/Enterprise/Datasmith/DatasmithMaxExporter/DatasmithMax2017.Target.cs
@@ -54,10 +54,10 @@ public abstract class DatasmithMaxBaseTarget : TargetRules
 		}
 
 		string SrcOutputFileName = string.Format(@"$(EngineDir)\Binaries\Win64\{0}\{1}.dll", ExeBinariesSubFolder, OutputName);
+		string DstOutputFileName = string.Format(@"$(EngineDir)\Binaries\Win64\{0}\{1}.gup", ExeBinariesSubFolder, OutputName);
 
-		string DstOutputFileName;
-
-		DstOutputFileName = string.Format(@"$(EngineDir)\Binaries\Win64\{0}\{1}.gup", ExeBinariesSubFolder, OutputName);
+		PostBuildSteps.Add(string.Format("echo Copying {0} to {1}...", SrcOutputFileName, DstOutputFileName));
+		PostBuildSteps.Add(string.Format("copy /Y \"{0}\" \"{1}\" 1>nul", SrcOutputFileName, DstOutputFileName));
 
 		PostBuildSteps.Add(string.Format("echo Copying {0} to {1}...", SrcOutputFileName, DstOutputFileName));
 		PostBuildSteps.Add(string.Format("copy /Y \"{0}\" \"{1}\" 1>nul", SrcOutputFileName, DstOutputFileName));
diff --git a/Engine/Source/Programs/Enterprise/Datasmith/DatasmithNavisworksExporter/DatasmithNavisworksPlugin/DatasmithNavisworksPlugin.csproj b/Engine/Source/Programs/Enterprise/Datasmith/DatasmithNavisworksExporter/DatasmithNavisworksPlugin/DatasmithNavisworksPlugin.csproj
--- a/Engine/Source/Programs/Enterprise/Datasmith/DatasmithNavisworksExporter/DatasmithNavisworksPlugin/DatasmithNavisworksPlugin.csproj
+++ b/Engine/Source/Programs/Enterprise/Datasmith/DatasmithNavisworksExporter/DatasmithNavisworksPlugin/DatasmithNavisworksPlugin.csproj
@@ -39,7 +39,7 @@
     <DebugType>full</DebugType>
     <Optimize>false</Optimize>
     <OutputPath>bin\Debug2026\</OutputPath>
-    <DefineConstants>DEBUG;TRACE;DATASMITH_NAVISWORKS_2026_OR_GREATER</DefineConstants>
+    <DefineConstants>DEBUG;TRACE;DATASMITH_NAVISWORKS_2025_OR_GREATER</DefineConstants>
     <ErrorReport>prompt</ErrorReport>
     <WarningLevel>4</WarningLevel>
     <StartAction>Program</StartAction>
@@ -51,7 +51,7 @@
     <DebugType>pdbonly</DebugType>
     <Optimize>true</Optimize>
     <OutputPath>bin\Release2026\</OutputPath>
-    <DefineConstants>TRACE;DATASMITH_NAVISWORKS_2026_OR_GREATER</DefineConstants>
+    <DefineConstants>TRACE;DATASMITH_NAVISWORKS_2025_OR_GREATER</DefineConstants>
     <ErrorReport>prompt</ErrorReport>
     <WarningLevel>4</WarningLevel>
     <StartAction>Program</StartAction>
