[SKP DS Exporter] Export the current viewport

diff --git a/Engine/Source/Programs/Enterprise/Datasmith/DatasmithSketchUpRubyExporter/Private/DatasmithSketchUpCamera.cpp b/Engine/Source/Programs/Enterprise/Datasmith/DatasmithSketchUpRubyExporter/Private/DatasmithSketchUpCamera.cpp
--- a/Engine/Source/Programs/Enterprise/Datasmith/DatasmithSketchUpRubyExporter/Private/DatasmithSketchUpCamera.cpp
+++ b/Engine/Source/Programs/Enterprise/Datasmith/DatasmithSketchUpRubyExporter/Private/DatasmithSketchUpCamera.cpp
@@ -207,3 +207,27 @@ FMD5Hash FCamera::GetHash()
 	Hash.Set(MD5);
 	return Hash;
 }
+
+bool FViewportCamera::Update(FExportContext& InContext)
+{
+	SUCameraRef camera_ref = SU_INVALID;
+	if (SUModelGetCamera(InContext.ModelRef, &camera_ref) == SU_ERROR_NONE && SUIsValid(camera_ref))
+	{
+		TSharedPtr<FCamera> MainCamera = FCamera::Create(InContext, camera_ref, TEXT("viewport_camera"));
+		
+		FMD5Hash CurrentHash = MainCamera->GetHash();
+		
+		if(Hash != CurrentHash)
+		{
+			Hash = CurrentHash;
+			MainCamera->Update(InContext);
+			return true;
+		}
+
+		return false;
+	}
+
+	DatasmithSketchUpUtils::ToRuby::LogWarn(TEXT("Cannot find viewport camera."));
+	
+	return false;
+}
diff --git a/Engine/Source/Programs/Enterprise/Datasmith/DatasmithSketchUpRubyExporter/Private/DatasmithSketchUpCamera.h b/Engine/Source/Programs/Enterprise/Datasmith/DatasmithSketchUpRubyExporter/Private/DatasmithSketchUpCamera.h
--- a/Engine/Source/Programs/Enterprise/Datasmith/DatasmithSketchUpRubyExporter/Private/DatasmithSketchUpCamera.h
+++ b/Engine/Source/Programs/Enterprise/Datasmith/DatasmithSketchUpRubyExporter/Private/DatasmithSketchUpCamera.h
@@ -39,4 +39,16 @@ namespace DatasmithSketchUp
 		bool bIsActive;
 	};
 
+
+	// camera that associates with the current viewport in Sketchup
+	class FViewportCamera : FNoncopyable
+	{
+	public:
+		FViewportCamera() {}
+
+		bool Update(FExportContext& InContext);
+		
+	private:
+		FMD5Hash Hash;
+	};
 }
diff --git a/Engine/Source/Programs/Enterprise/Datasmith/DatasmithSketchUpRubyExporter/Private/DatasmithSketchUpExportContext.cpp b/Engine/Source/Programs/Enterprise/Datasmith/DatasmithSketchUpRubyExporter/Private/DatasmithSketchUpExportContext.cpp
--- a/Engine/Source/Programs/Enterprise/Datasmith/DatasmithSketchUpRubyExporter/Private/DatasmithSketchUpExportContext.cpp
+++ b/Engine/Source/Programs/Enterprise/Datasmith/DatasmithSketchUpRubyExporter/Private/DatasmithSketchUpExportContext.cpp
@@ -97,8 +97,9 @@ void FExportContext::Populate()
 	// Parse/convert Model
 	Layers.PopulateFromModel(ModelRef);
 	Scenes.Update();
+	ViewportCamera.Update(*this);
+	
 	ComponentDefinitions.PopulateFromModel(ModelRef);
-
 	ModelDefinition->ParseNode(*this, *RootNode); // Create node hierarchy
 }
 
@@ -167,6 +168,8 @@ bool FExportContext::Update(bool bModifiedHint)
 
 	bModified |= Scenes.Update();
 
+	bModified |= ViewportCamera.Update(*this);
+
 	bModified |= ModelDefinition->UpdateModel(*this);
 
 	if (!bModified)
diff --git a/Engine/Source/Programs/Enterprise/Datasmith/DatasmithSketchUpRubyExporter/Private/DatasmithSketchUpExportContext.h b/Engine/Source/Programs/Enterprise/Datasmith/DatasmithSketchUpRubyExporter/Private/DatasmithSketchUpExportContext.h
--- a/Engine/Source/Programs/Enterprise/Datasmith/DatasmithSketchUpRubyExporter/Private/DatasmithSketchUpExportContext.h
+++ b/Engine/Source/Programs/Enterprise/Datasmith/DatasmithSketchUpRubyExporter/Private/DatasmithSketchUpExportContext.h
@@ -371,6 +371,7 @@ namespace DatasmithSketchUp
 		FLayerCollection Layers;
 		FImageCollection Images;
 		FImageFileCollection ImageFiles;
+		FViewportCamera ViewportCamera;
 
 		bool bColorByLayer = false;
 		bool bColorByLayerInvaliated = true;
