[Rhino] Fix Mac build

- Copy DatasmithFacadeCSharp.dylib dylib only from DatasmithFacadeCSharp. Copying every dylib was overwriting TBB dylibs which now appear in a target directory of a target using InterTBB.

diff --git a/Engine/Source/Programs/Enterprise/Datasmith/DatasmithRhinoExporter/Installer/CreateMacPackages.sh b/Engine/Source/Programs/Enterprise/Datasmith/DatasmithRhinoExporter/Installer/CreateMacPackages.sh
--- a/Engine/Source/Programs/Enterprise/Datasmith/DatasmithRhinoExporter/Installer/CreateMacPackages.sh
+++ b/Engine/Source/Programs/Enterprise/Datasmith/DatasmithRhinoExporter/Installer/CreateMacPackages.sh
@@ -20,7 +20,8 @@ rm -f "$PLUGIN_MACRHI_PATH"
 
 echo "Creating .rhp package"
 mkdir "$PLUGIN_RHP_PATH"
-cp "$DATASMITH_FACADE_BINARY_DIR/"*".dylib" "$PLUGIN_RHP_PATH"
+
+cp "$DATASMITH_FACADE_BINARY_DIR/DatasmithFacadeCSharp.dylib" "$PLUGIN_RHP_PATH"
 cp "$TBB_BINARY_DIR/"*".dylib" "$PLUGIN_RHP_PATH"
 cp "$PLUGIN_BINARY_DIR/$PLUGIN_NAME.dll" "$PLUGIN_RHP_PATH/$PLUGIN_NAME.rhp"
 cp "$PLUGIN_BASE_DIR/Config/DatasmithRhino.rhp.config" "$PLUGIN_RHP_PATH/$PLUGIN_NAME.rhp.config"
