[Rhino3D] Allow Camera Visibility in .udatasmith files

RhinoPlugin now exports viewport cameras alongside named views.
For viewport cameras "Current_" is prefixed to the name.

Aspect ratio and sensor size is also now serialized in cameras.

diff --git a/Engine/Source/Programs/Enterprise/Datasmith/DatasmithRhinoExporter/Private/ElementExporters/DatasmithRhinoActorExporter.cs b/Engine/Source/Programs/Enterprise/Datasmith/DatasmithRhinoExporter/Private/ElementExporters/DatasmithRhinoActorExporter.cs
--- a/Engine/Source/Programs/Enterprise/Datasmith/DatasmithRhinoExporter/Private/ElementExporters/DatasmithRhinoActorExporter.cs
+++ b/Engine/Source/Programs/Enterprise/Datasmith/DatasmithRhinoExporter/Private/ElementExporters/DatasmithRhinoActorExporter.cs
@@ -1,5 +1,6 @@
 // Copyright Epic Games, Inc. All Rights Reserved.
 
+using System;
 using DatasmithRhino.ExportContext;
 using DatasmithRhino.Properties.Localization;
 using DatasmithRhino.Utils;
@@ -9,6 +10,9 @@ using Rhino.DocObjects;
 using Rhino.Geometry;
 using System.Collections.Generic;
 using System.Collections.Specialized;
+using System.Linq;
+using Rhino.Display;
+using Rhino.Input;
 
 namespace DatasmithRhino.ElementExporters
 {
@@ -24,7 +28,13 @@ namespace DatasmithRhino.ElementExporters
 		protected override IEnumerable<DatasmithActorInfo> GetElementsToSynchronize()
 		{
 			// We don't export the root, only its descendants.
-			return ExportContext.SceneRoot.GetDescendantEnumerator();
+			List<DatasmithActorInfo> AllElements = ExportContext.SceneRoot.GetDescendantEnumerator()
+				.ToList();
+			
+			// We add "on the fly" generated actors (like viewport camera, etc)
+			AllElements.AddRange(GetCustomExportActorInfos());
+			
+			return AllElements;
 		}
 
 		protected override FDatasmithFacadeElement CreateElement(DatasmithActorInfo ElementInfo)
@@ -74,6 +84,42 @@ namespace DatasmithRhino.ElementExporters
 
 		///// END IDatasmithRhinoElementExporter Interface /////
 
+		/// <summary>
+		/// Generates Actor Infos that would be used to create custom actors on the fly
+		/// required before exporting. These are actors that are not present in the Scene
+		/// but still are needed
+		/// </summary>
+		/// <returns></returns>
+		private List<DatasmithActorInfo> GetCustomExportActorInfos()
+		{
+			List<DatasmithActorInfo> Actors = new List<DatasmithActorInfo>();
+			
+			// Active viewport views are exported as Camera Actors
+			RhinoView[] ViewsList = ExportContext.RhinoDocument.Views.GetViewList(true, true);
+			
+			foreach (var View in ViewsList)
+			{
+				ViewportInfo ViewportInfo = new ViewportInfo(View.ActiveViewport);
+				
+				string Hash = FDatasmithFacadeElement.GetStringHash(
+					$"{ViewportInfo.CameraAngle}, {ViewportInfo.CameraLocation}, {ViewportInfo.CameraDirection}" +
+					$"{ViewportInfo.CameraUp},{ViewportInfo.ViewScale}");
+
+				// viewport camera names are pre fixed with "Current" substring
+				// cameras without this are "named views"
+				// this is just a design decision in naming convention so users can identify
+				// in UE/TM what sort of camera each one is 
+				string CameraObjectName = $"Current_{View.ActiveViewport.Name}";
+				
+				DatasmithCameraActorInfo ViewportCameraInfo =
+					new DatasmithCameraActorInfo(ViewportInfo, CameraObjectName, CameraObjectName, CameraObjectName, Hash);
+				
+				Actors.Add(ViewportCameraInfo);
+			}
+			
+			return Actors;
+		}
+		
 		private FDatasmithFacadeActor CreateActor(DatasmithActorInfo Node)
 		{
 			FDatasmithFacadeActor ExportedActor = null;
@@ -313,11 +359,14 @@ namespace DatasmithRhino.ElementExporters
 			CameraActor.SetCameraPosition((float)Viewport.CameraLocation.X, (float)Viewport.CameraLocation.Y, (float)Viewport.CameraLocation.Z);
 			CameraActor.SetCameraRotation((float)Viewport.CameraDirection.X, (float)Viewport.CameraDirection.Y, (float)Viewport.CameraDirection.Z,
 				(float)Viewport.CameraUp.X, (float)Viewport.CameraUp.Y, (float)Viewport.CameraUp.Z);
+			
+			CameraActor.SetAspectRatio((float) Viewport.FrustumAspect);
+			CameraActor.SetSensorWidth((float) Viewport.ScreenPort.Width);
 		}
 
 		private static void AddActorToParent(FDatasmithFacadeActor InDatasmithActor, DatasmithActorInfo InNode, FDatasmithFacadeScene InDatasmithScene)
 		{
-			if (InNode.Parent.bIsRoot)
+			if (InNode.Parent == null || InNode.Parent.bIsRoot)
 			{
 				InDatasmithScene.AddActor(InDatasmithActor);
 			}
@@ -429,7 +478,8 @@ namespace DatasmithRhino.ElementExporters
 		private string GetLayers(DatasmithActorInfo InNode)
 		{
 			bool bIsSameAsParentLayer =
-				!(InNode.RhinoCommonObject is Layer
+				(InNode.Parent != null) // this is a manually created objects within this plugin, it has no Parent
+				&& !(InNode.RhinoCommonObject is Layer
 					|| (InNode.Parent.bIsRoot && InNode.RhinoCommonObject == null) //This is a dummy document layer.
 					|| (InNode.Parent.RhinoCommonObject as RhinoObject)?.ObjectType == ObjectType.InstanceReference);
 
