Update supported version of C++ from 17 to 20
  Consequence: ArchiCAD23 to rchiCAD26 exporters have been temporarily disabled on Windows.

#rnx

diff --git a/Engine/Source/Programs/Enterprise/Datasmith/DatasmithARCHICADExporter/Private/SyncData.cpp b/Engine/Source/Programs/Enterprise/Datasmith/DatasmithARCHICADExporter/Private/SyncData.cpp
--- a/Engine/Source/Programs/Enterprise/Datasmith/DatasmithARCHICADExporter/Private/SyncData.cpp
+++ b/Engine/Source/Programs/Enterprise/Datasmith/DatasmithARCHICADExporter/Private/SyncData.cpp
@@ -1388,7 +1388,7 @@ FSyncData::FLight::FLightGDLParameters::FLightGDLParameters(const API_Guid&		InL
 			double TmpAreaShape;
 			if (GetParameter(AutoMemo.Memo.params, "iC4dDetAreaShape", &TmpAreaShape))
 			{
-				if (TmpAreaShape >= kDisc && TmpAreaShape <= kPerpendicularCylinder)
+				if (TmpAreaShape >= static_cast<double>(kDisc) && TmpAreaShape <= static_cast<double>(kPerpendicularCylinder))
 				{
 					AreaShape = (EC4dDetAreaShape) int(TmpAreaShape + 0.5);
 				}
diff --git a/Engine/Source/Programs/Enterprise/Datasmith/DatasmithARCHICADExporter/Private/Utils/APIEnvir.h b/Engine/Source/Programs/Enterprise/Datasmith/DatasmithARCHICADExporter/Private/Utils/APIEnvir.h
--- a/Engine/Source/Programs/Enterprise/Datasmith/DatasmithARCHICADExporter/Private/Utils/APIEnvir.h
+++ b/Engine/Source/Programs/Enterprise/Datasmith/DatasmithARCHICADExporter/Private/Utils/APIEnvir.h
@@ -37,6 +37,7 @@ THIRD_PARTY_INCLUDES_START
 __pragma(warning(push))
 __pragma(warning(disable: 4005)) // 'TEXT': macro redefinition
 __pragma(warning(disable: 5040)) // Support\Modules\GSRoot\vaarray.hpp: dynamic exception specifications are valid only in C++14 and earlier
+__pragma(warning(disable: 5055)) // operator '>=': deprecated between enumerations and floating-point
 
 #else
 
