Datasmith - Physical material is not imported correctly

Exported materials were missing support for PBR materials. DCC legacy non-pbr flow and modern pbr flow.
This change adds support for both in some base properties.

diff --git a/Engine/Source/Programs/Enterprise/Datasmith/DatasmithRhinoExporter/Private/ElementExporters/DatasmithRhinoMaterialExporter.cs b/Engine/Source/Programs/Enterprise/Datasmith/DatasmithRhinoExporter/Private/ElementExporters/DatasmithRhinoMaterialExporter.cs
--- a/Engine/Source/Programs/Enterprise/Datasmith/DatasmithRhinoExporter/Private/ElementExporters/DatasmithRhinoMaterialExporter.cs
+++ b/Engine/Source/Programs/Enterprise/Datasmith/DatasmithRhinoExporter/Private/ElementExporters/DatasmithRhinoMaterialExporter.cs
@@ -7,6 +7,7 @@ using Rhino.Geometry;
 using System.Collections.Generic;
 using System.Drawing;
 using System;
+using Rhino;
 
 namespace DatasmithRhino.ElementExporters
 {
@@ -59,7 +60,17 @@ namespace DatasmithRhino.ElementExporters
 		{
 			Color MaterialDiffuseColor = RhinoMaterial.DiffuseColor;
 			MaterialDiffuseColor = Color.FromArgb(255 - (byte)(255 * RhinoMaterial.Transparency), MaterialDiffuseColor);
-
+			float Roughness = 1f - (float) (RhinoMaterial.Shine / Material.MaxShine);
+			float Reflectivity = (float) RhinoMaterial.Reflectivity;
+			
+			// Rhino Materials have different shading models and properties need to be adjusted based on that
+			if (RhinoMaterial.IsPhysicallyBased)
+			{
+				MaterialDiffuseColor = RhinoMaterial.PhysicallyBased.BaseColor.AsSystemColor();
+				Roughness = (float)RhinoMaterial.PhysicallyBased.Roughness;
+				Reflectivity = (float)RhinoMaterial.PhysicallyBased.Metallic;
+			}
+			
 			Texture[] MaterialTextures = RhinoMaterial.GetTextures();
 			for (int TextureIndex = 0; TextureIndex < MaterialTextures.Length; ++TextureIndex)
 			{
@@ -80,7 +91,6 @@ namespace DatasmithRhino.ElementExporters
 				FDatasmithFacadeMaterialExpressionColor ColorExpression = DSMaterial.AddMaterialExpressionColor();
 				ColorExpression.SetName("Diffuse Color");
 				ColorExpression.SetsRGBColor(MaterialDiffuseColor.R, MaterialDiffuseColor.G, MaterialDiffuseColor.B, MaterialDiffuseColor.A);
-
 				ColorExpression.ConnectExpression(DSMaterial.GetBaseColor());
 			}
 
@@ -112,24 +122,16 @@ namespace DatasmithRhino.ElementExporters
 					DSMaterial.GetOpacity().SetExpression(Multiply);
 				}
 			}
-
-			float Shininess = (float) (RhinoMaterial.Shine / Material.MaxShine);
-			if (Math.Abs(Shininess) > float.Epsilon)
-			{
-				FDatasmithFacadeMaterialExpressionScalar ShininessExpression = DSMaterial.AddMaterialExpressionScalar();
-				ShininessExpression.SetName("Roughness");
-				ShininessExpression.SetScalar(1f - Shininess);
-				ShininessExpression.ConnectExpression(DSMaterial.GetRoughness());
-			}
-
-			float Reflectivity = (float) RhinoMaterial.Reflectivity;
-			if (Math.Abs(Reflectivity) > float.Epsilon)
-			{
-				FDatasmithFacadeMaterialExpressionScalar ReflectivityExpression = DSMaterial.AddMaterialExpressionScalar();
-				ReflectivityExpression.SetName("Metallic");
-				ReflectivityExpression.SetScalar(Reflectivity);
-				ReflectivityExpression.ConnectExpression(DSMaterial.GetMetallic());
-			}
+			
+			FDatasmithFacadeMaterialExpressionScalar ShininessExpression = DSMaterial.AddMaterialExpressionScalar();
+			ShininessExpression.SetName("Roughness");
+			ShininessExpression.SetScalar(Roughness);
+			ShininessExpression.ConnectExpression(DSMaterial.GetRoughness());
+			
+			FDatasmithFacadeMaterialExpressionScalar ReflectivityExpression = DSMaterial.AddMaterialExpressionScalar();
+			ReflectivityExpression.SetName("Metallic");
+			ReflectivityExpression.SetScalar(Reflectivity);
+			ReflectivityExpression.ConnectExpression(DSMaterial.GetMetallic());
 		}
 
 		private static void AddTextureToMaterial(FDatasmithFacadeUEPbrMaterial DSMaterial, DatasmithTextureInfo TextureInfo, Color DiffuseColor)
