Change Datasmith exporter build settings and remove unsupported .net frameworks

Exporter plugin -  csproj was ported to .net8 core
New dll dependencies and build were updated accordingly

diff --git a/Engine/Source/Programs/Enterprise/Datasmith/DatasmithRhinoExporter/DatasmithRhino8/DatasmithRhino8.csproj b/Engine/Source/Programs/Enterprise/Datasmith/DatasmithRhinoExporter/DatasmithRhino8/DatasmithRhino8.csproj
--- a/Engine/Source/Programs/Enterprise/Datasmith/DatasmithRhinoExporter/DatasmithRhino8/DatasmithRhino8.csproj
+++ b/Engine/Source/Programs/Enterprise/Datasmith/DatasmithRhinoExporter/DatasmithRhino8/DatasmithRhino8.csproj
@@ -1,55 +1,25 @@
-﻿<?xml version="1.0" encoding="utf-8"?>
-<Project ToolsVersion="4.0" DefaultTargets="Build" xmlns="http://schemas.microsoft.com/developer/msbuild/2003">
+﻿<Project Sdk="Microsoft.NET.Sdk" xmlns="http://schemas.microsoft.com/developer/msbuild/2003">
   <PropertyGroup>
-    <Configuration Condition=" '$(Configuration)' == '' ">Debug</Configuration>
+    <TargetFramework>net8.0</TargetFramework>
+    <AppendTargetFrameworkToOutputPath>false</AppendTargetFrameworkToOutputPath>
     <Platform Condition=" '$(Platform)' == '' ">x64</Platform>
     <AssemblyInfo Condition=" '$(AssemblyInfo)' == ''">..\Properties\DefaultAssemblyInfo.cs</AssemblyInfo>
     <Platform_OS Condition=" '$(Platform_OS)' == '' ">Win64</Platform_OS>
-    <!--If the evironment variable is not define, try to find the Rhino files on build machines-->
-    <ProductVersion>8.0.30703</ProductVersion>
-    <SchemaVersion>2.0</SchemaVersion>
-    <ProjectGuid>{D1FDC795-B334-4933-B680-088119CDC6BB}</ProjectGuid>
     <OutputType>Library</OutputType>
-    <AppDesignerFolder>Properties</AppDesignerFolder>
     <RootNamespace>DatasmithRhino</RootNamespace>
     <AssemblyName Condition="'$(Configuration)' == 'Debug'">DatasmithRhino8-Debug</AssemblyName>
-    <AssemblyName Condition="'$(Configuration)' == 'Shipping' OR '$(Configuration)' == 'Release'">DatasmithRhino8-Shipping</AssemblyName>
-    <AssemblyName Condition="'$(Configuration)' != 'Debug' AND '$(Configuration)' != 'Shipping' AND '$(Configuration)' != 'Release'">DatasmithRhino8</AssemblyName>
-    <TargetFrameworkVersion>v4.5</TargetFrameworkVersion>
-    <FileAlignment>512</FileAlignment>
+    <AssemblyName Condition="'$(Configuration)' == 'Development'">DatasmithRhino8</AssemblyName>
+    <AssemblyName Condition="'$(Configuration)' == 'Shipping' OR '$(Configuration)' == 'Release'">DatasmithRhino8</AssemblyName>
     <IsWebBootstrapper>false</IsWebBootstrapper>
-    <NuGetPackageImportStamp>
-    </NuGetPackageImportStamp>
+    <GenerateAssemblyInfo>false</GenerateAssemblyInfo>
+    <Configurations>Debug;Development;Release</Configurations>
+    <NoWarn>$(NoWarn);CA1416;CS0612</NoWarn>
   </PropertyGroup>
-
-  <PropertyGroup Condition="'$(Configuration)|$(Platform)' == 'Debug|x64'">
-    <PlatformTarget>AnyCPU</PlatformTarget>
-    <DebugSymbols>true</DebugSymbols>
-    <DebugType>full</DebugType>
-    <Optimize>false</Optimize>
+  <PropertyGroup>
     <OutputPath>..\..\..\..\..\..\Binaries\$(Platform_OS)\Rhino\8</OutputPath>
-    <DefineConstants>DEBUG;TRACE;RHINO_VERSION_8</DefineConstants>
-    <ErrorReport>prompt</ErrorReport>
+    <DefineConstants>DEBUG;TRACE;RHINO_VERSION_8;USING_NET_CORE</DefineConstants>
     <CodeAnalysisIgnoreBuiltInRules>false</CodeAnalysisIgnoreBuiltInRules>
   </PropertyGroup>
-  <PropertyGroup Condition="'$(Configuration)|$(Platform)' == 'Development|x64'">
-    <PlatformTarget>AnyCPU</PlatformTarget>
-    <DebugType>pdbonly</DebugType>
-    <Optimize>true</Optimize>
-    <OutputPath>..\..\..\..\..\..\Binaries\$(Platform_OS)\Rhino\8</OutputPath>
-    <DefineConstants>TRACE;RHINO_VERSION_8</DefineConstants>
-    <ErrorReport>prompt</ErrorReport>
-    <WarningLevel>4</WarningLevel>
-  </PropertyGroup>
-  <PropertyGroup Condition="'$(Configuration)|$(Platform)' == 'Release|x64'">
-    <PlatformTarget>AnyCPU</PlatformTarget>
-    <DebugType>pdbonly</DebugType>
-    <Optimize>true</Optimize>
-    <OutputPath>..\..\..\..\..\..\Binaries\$(Platform_OS)\Rhino\8</OutputPath>
-    <DefineConstants>TRACE;RHINO_VERSION_8</DefineConstants>
-    <ErrorReport>prompt</ErrorReport>
-    <WarningLevel>4</WarningLevel>
-  </PropertyGroup>
   <PropertyGroup Condition="'$(Platform_OS)' == 'Win64'">
     <DefineConstants>$(DefineConstants);WIN64_OS</DefineConstants>
   </PropertyGroup>
@@ -60,69 +30,57 @@
     <DefineConstants>$(DefineConstants);DATASMITHRHINO_EXPERIMENTAL</DefineConstants>
   </PropertyGroup>
   <ItemGroup>
-    <Reference Include="System" />
-    <Reference Include="System.Core" />
-    <Reference Include="System.Drawing" />
-    <Reference Include="Microsoft.CSharp" />
-    <Reference Include="RhinoCommon">
+    <Reference Include="Eto">
       <SpecificVersion>False</SpecificVersion>
-      <HintPath>$(RHINO8_PATH)\rhinocommon.dll</HintPath>
+      <HintPath>$(RHINO8_PATH)\netcore\Eto.dll</HintPath>
       <Private>False</Private>
     </Reference>
-    <Reference Include="Eto">
+    <Reference Include="RhinoCommon">
       <SpecificVersion>False</SpecificVersion>
-      <HintPath>$(RHINO8_PATH)\Eto.dll</HintPath>
+      <HintPath>$(RHINO8_PATH)\netcore\RhinoCommon.dll</HintPath>
       <Private>False</Private>
     </Reference>
     <Reference Include="Rhino.UI">
       <SpecificVersion>False</SpecificVersion>
-      <HintPath>$(RHINO8_PATH)\Rhino.UI.dll</HintPath>
+      <HintPath>$(RHINO8_PATH)\netcore\Rhino.UI.dll</HintPath>
       <Private>False</Private>
     </Reference>
   </ItemGroup>
   <ItemGroup Condition="'$(Platform_OS)' == 'Mac'">
-    <Reference Include="RhinoMac">
-      <SpecificVersion>False</SpecificVersion>
-      <HintPath>$(RHINO8_PATH)\RhinoMac.dll</HintPath>
+    <Reference Include="Eto">
+      <HintPath>$(RHINO8_PATH)/netcore/mac/Eto.dll</HintPath>
+      <Private>False</Private>
+    </Reference>
+    <Reference Include="RhinoCommon">
+      <HintPath>$(RHINO8_PATH)/netcore/mac/RhinoCommon.dll</HintPath>
+      <Private>False</Private>
+    </Reference>
+    <Reference Include="Rhino.UI">
+      <HintPath>$(RHINO8_PATH)/netcore/mac/Rhino.UI.dll</HintPath>
       <Private>False</Private>
     </Reference>
   </ItemGroup>
   <ItemGroup>
-    <Compile Include="DatasmithRhino8.cs" />
     <Compile Include="$(AssemblyInfo)">
       <Link>Properties\AssemblyInfo.cs</Link>
     </Compile>
-    <!-- DatasmithCSharpFacade API, possibly change that. -->
     <Compile Include="..\..\..\..\..\..\Binaries\$(Platform_OS)\DatasmithFacadeCSharp\Public\*.cs">
       <Link>DatasmithSwigFacade\%(Filename)%(Extension)</Link>
     </Compile>
   </ItemGroup>
-  <ItemGroup Condition="'$(Platform_OS)' == 'Win64'">
+  <ItemGroup>
     <Compile Include="..\Private\**\*.cs">
       <Link>DatasmithRhinoExporter\%(RecursiveDir)%(Filename)%(Extension)</Link>
     </Compile>
   </ItemGroup>
-  <!-- TODO Remove this explicit definition once .net uses msbuild on Mac -->
-  <ItemGroup Condition="'$(Platform_OS)' == 'Mac'">
-    <Compile Include="..\Private\Commands\*.cs">
-      <Link>DatasmithRhinoExporter\Commands\%(Filename)%(Extension)</Link>
-    </Compile>
-    <Compile Include="..\Private\DirectLink\*.cs">
-      <Link>DatasmithRhinoExporter\DirectLink\%(Filename)%(Extension)</Link>
-    </Compile>
-    <Compile Include="..\Private\ElementExporters\*.cs">
-      <Link>DatasmithRhinoExporter\ElementExporters\%(Filename)%(Extension)</Link>
-    </Compile>
-    <Compile Include="..\Private\ExportContext\*.cs">
-      <Link>DatasmithRhinoExporter\ExportContext\%(Filename)%(Extension)</Link>
-    </Compile>
-    <Compile Include="..\Private\Utils\*.cs">
-      <Link>DatasmithRhinoExporter\Utils\%(Filename)%(Extension)</Link>
-    </Compile>
-    <Compile Include="..\Private\*.cs">
-      <Link>DatasmithRhinoExporter\%(Filename)%(Extension)</Link>
-    </Compile>
-  </ItemGroup>
+<!--  <ItemGroup Condition="'$(Platform_OS)' == 'Mac'">-->
+<!--    <Compile Include="..\Private\Commands\*.cs" />-->
+<!--    <Compile Include="..\Private\DirectLink\*.cs" />-->
+<!--    <Compile Include="..\Private\ElementExporters\*.cs" />-->
+<!--    <Compile Include="..\Private\ExportContext\*.cs" />-->
+<!--    <Compile Include="..\Private\Utils\*.cs" />-->
+<!--    <Compile Include="..\Private\*.cs" />-->
+<!--  </ItemGroup>-->
   <ItemGroup>
     <Compile Include="..\SharedResources\Localization\Resources.Designer.cs">
       <Link>Properties\Localization\Resources.Designer.cs</Link>
@@ -216,30 +174,26 @@
       <LastGenOutput>Resources.zh.Designer.cs</LastGenOutput>
     </EmbeddedResource>
   </ItemGroup>
-  <Import Project="$(MSBuildToolsPath)\Microsoft.CSharp.targets" />
-  <!-- To modify your build process, add your task inside one of the targets below and uncomment it. 
-       Other similar extension points exist, see Microsoft.Common.targets.
-  <Target Name="BeforeBuild">
-  </Target>
-  <Target Name="AfterBuild">
-  </Target>
-  -->
+  <ItemGroup Condition="'$(Platform_OS)' == 'Win64'">
+    <PackageReference Include="Microsoft.CSharp"  Version="4.7.0"/>
+  </ItemGroup>
   <PropertyGroup Condition="'$(Platform_OS)' == 'Win64'">
     <PostBuildEvent>
-		Copy "$(TargetPath)" "$(TargetDir)$(TargetName).rhp"
-		Copy "$(MSBuildProjectDirectory)\..\SharedResources\DatasmithRhino.rui" "$(TargetDir)$(TargetName).rui"
+		Copy "$(MSBuildProjectDirectory)\$(OutputPath)\$(AssemblyName).dll" "$(MSBuildProjectDirectory)\$(OutputPath)\$(AssemblyName).rhp"
+		Copy "$(MSBuildProjectDirectory)\..\SharedResources\DatasmithRhino.rui" "$(MSBuildProjectDirectory)\$(OutputPath)\$(AssemblyName).rui"
     </PostBuildEvent>
   </PropertyGroup>
   <PropertyGroup Condition="'$(Platform_OS)' == 'Mac'">
     <PostBuildEvent>
-    sh "$(ProjectDir)../Installer/CreateMacPackages.sh" $(TargetName) 8
+    sh "$(MSBuildProjectDirectory)/../Installer/CreateMacPackages.sh" $(AssemblyName) 8
     </PostBuildEvent>
   </PropertyGroup>
   <PropertyGroup>
     <FallbackCulture>en-US</FallbackCulture>
   </PropertyGroup>
-  <PropertyGroup Condition="'$(Configuration)' == 'Debug' And '$(Platform_OS)' == 'Win64'">
-    <StartProgram>$(RHINO8_PATH)\Rhino.exe</StartProgram>
+  <PropertyGroup>
+<!--    Condition="'$(Configuration)' == 'Debug' And '$(Platform_OS)' == 'Win64'">-->
+    <StartProgram>C:\Program Files\Rhino 8\System\Rhino.exe</StartProgram>
     <StartArguments>
     </StartArguments>
     <StartAction>Program</StartAction>
diff --git a/Engine/Source/Programs/Enterprise/Datasmith/DatasmithRhinoExporter/Private/DatasmithRhinoPlugin.cs b/Engine/Source/Programs/Enterprise/Datasmith/DatasmithRhinoExporter/Private/DatasmithRhinoPlugin.cs
--- a/Engine/Source/Programs/Enterprise/Datasmith/DatasmithRhinoExporter/Private/DatasmithRhinoPlugin.cs
+++ b/Engine/Source/Programs/Enterprise/Datasmith/DatasmithRhinoExporter/Private/DatasmithRhinoPlugin.cs
@@ -6,7 +6,7 @@ using Rhino;
 using Rhino.PlugIns;
 using System;
 
-#if MAC_OS
+#if MAC_OS && !USING_NET_CORE
 using RhinoMac;
 #endif
 
@@ -56,7 +56,7 @@ namespace DatasmithRhino
 			return result;
 		}
 
-#if MAC_OS
+#if MAC_OS && !USING_NET_CORE
 		protected override LoadReturnCode OnLoad(ref string errorMessage)
 		{
 			string PluginPath = System.IO.Path.GetDirectoryName(Assembly.Location);
