FIX
Compile errors in DatasmithRhino6.csproj and DatasmithRhinoMaterialExporter.cs

A previously submitted task was missing DCC version pre compile check

diff --git a/Engine/Source/Programs/Enterprise/Datasmith/DatasmithRhinoExporter/DatasmithRhino8/DatasmithRhino8.csproj b/Engine/Source/Programs/Enterprise/Datasmith/DatasmithRhinoExporter/DatasmithRhino8/DatasmithRhino8.csproj
--- a/Engine/Source/Programs/Enterprise/Datasmith/DatasmithRhinoExporter/DatasmithRhino8/DatasmithRhino8.csproj
+++ b/Engine/Source/Programs/Enterprise/Datasmith/DatasmithRhinoExporter/DatasmithRhino8/DatasmithRhino8.csproj
@@ -17,7 +17,7 @@
   </PropertyGroup>
   <PropertyGroup>
     <OutputPath>..\..\..\..\..\..\Binaries\$(Platform_OS)\Rhino\8</OutputPath>
-    <DefineConstants>DEBUG;TRACE;RHINO_VERSION_8;USING_NET_CORE</DefineConstants>
+    <DefineConstants>DEBUG;TRACE;RHINO_VERSION_8;USING_NET_CORE;RHINO_VERSION_8_OR_NEWER</DefineConstants>
     <CodeAnalysisIgnoreBuiltInRules>false</CodeAnalysisIgnoreBuiltInRules>
   </PropertyGroup>
   <PropertyGroup Condition="'$(Platform_OS)' == 'Win64'">
diff --git a/Engine/Source/Programs/Enterprise/Datasmith/DatasmithRhinoExporter/Private/ElementExporters/DatasmithRhinoMaterialExporter.cs b/Engine/Source/Programs/Enterprise/Datasmith/DatasmithRhinoExporter/Private/ElementExporters/DatasmithRhinoMaterialExporter.cs
--- a/Engine/Source/Programs/Enterprise/Datasmith/DatasmithRhinoExporter/Private/ElementExporters/DatasmithRhinoMaterialExporter.cs
+++ b/Engine/Source/Programs/Enterprise/Datasmith/DatasmithRhinoExporter/Private/ElementExporters/DatasmithRhinoMaterialExporter.cs
@@ -63,6 +63,7 @@ namespace DatasmithRhino.ElementExporters
 			float Roughness = 1f - (float) (RhinoMaterial.Shine / Material.MaxShine);
 			float Reflectivity = (float) RhinoMaterial.Reflectivity;
 			
+			#if RHINO_VERSION_8_OR_NEWER
 			// Rhino Materials have different shading models and properties need to be adjusted based on that
 			if (RhinoMaterial.IsPhysicallyBased)
 			{
@@ -70,6 +71,7 @@ namespace DatasmithRhino.ElementExporters
 				Roughness = (float)RhinoMaterial.PhysicallyBased.Roughness;
 				Reflectivity = (float)RhinoMaterial.PhysicallyBased.Metallic;
 			}
+			#endif
 			
 			Texture[] MaterialTextures = RhinoMaterial.GetTextures();
 			for (int TextureIndex = 0; TextureIndex < MaterialTextures.Length; ++TextureIndex)
