[On behalf of Kerim Borchaev][3ds Max Exporter] 3ds Max 2026 support
#rnx

diff --git a/Engine/Source/Programs/Enterprise/Datasmith/DatasmithMaxExporter/DatasmithMax2017.Target.cs b/Engine/Source/Programs/Enterprise/Datasmith/DatasmithMaxExporter/DatasmithMax2017.Target.cs
--- a/Engine/Source/Programs/Enterprise/Datasmith/DatasmithMaxExporter/DatasmithMax2017.Target.cs
+++ b/Engine/Source/Programs/Enterprise/Datasmith/DatasmithMaxExporter/DatasmithMax2017.Target.cs
@@ -1,6 +1,8 @@
 // Copyright Epic Games, Inc. All Rights Reserved.
 
 using UnrealBuildTool;
+using System.IO;
+using System;
 
 [SupportedPlatforms("Win64")]
 public abstract class DatasmithMaxBaseTarget : TargetRules
@@ -19,10 +21,31 @@ public abstract class DatasmithMaxBaseTarget : TargetRules
 
 		WindowsPlatform.ModuleDefinitionFile = "Programs/Enterprise/Datasmith/DatasmithMaxExporter/DatasmithMaxExporterWithDirectLink.def";
 
-		WindowsPlatform.Compiler = WindowsCompiler.VisualStudio2022;
+		if(StaticAnalyzer != StaticAnalyzer.None)
+		{
+			string MaxSDKLocation = "";
+			string SDKRootEnvVar = System.Environment.GetEnvironmentVariable("UE_SDKS_ROOT");
+			if (SDKRootEnvVar != null && SDKRootEnvVar != "")
+			{
+				MaxSDKLocation = Path.Combine(SDKRootEnvVar, "HostWin64", "Win64", "3dsMax");
+			}
+
+			if (!Directory.Exists(MaxSDKLocation))
+			{
+			}
+
+			WindowsPlatform.Compiler = WindowsCompiler.Clang;
+			Console.WriteLine("Skipping " + MaxSDKLocation);
+			AdditionalCompilerArguments = "-Xanalyzer,-isystem-after" + MaxSDKLocation;
+		}
+		else
+		{
+			WindowsPlatform.Compiler = WindowsCompiler.VisualStudio2022;
+			AdditionalCompilerArguments = "/Zc:referenceBinding- /Zc:strictStrings- /Zc:rvalueCast- /Zc:preprocessor-";
+		}
+
 		WindowsPlatform.bStrictConformanceMode = false;
 		WindowsPlatform.bStrictPreprocessorConformance = false;
-		AdditionalCompilerArguments = "/Zc:referenceBinding- /Zc:strictStrings- /Zc:rvalueCast- /Zc:preprocessor-";
 
 		bBuildDeveloperTools = false;
 		bBuildWithEditorOnlyData = true;
@@ -36,7 +59,8 @@ public abstract class DatasmithMaxBaseTarget : TargetRules
 
 		// For DirectLinkUI (see FDatasmithExporterManager::FInitOptions)
 		GlobalDefinitions.Add("UE_EXTERNAL_PROFILING_ENABLED=0");
-		
+		GlobalDefinitions.Add("NO_INIUTIL_USING");
+
 		// Disable CrashReporter
 		GlobalDefinitions.Add("NOINITCRASHREPORTER=1");
 
diff --git a/Engine/Source/Programs/Enterprise/Datasmith/DatasmithMaxExporter/DatasmithMax2026.Build.cs b/Engine/Source/Programs/Enterprise/Datasmith/DatasmithMaxExporter/DatasmithMax2026.Build.cs
new file mode 100644
--- /dev/null
+++ b/Engine/Source/Programs/Enterprise/Datasmith/DatasmithMaxExporter/DatasmithMax2026.Build.cs
@@ -0,0 +1,16 @@
+// Copyright Epic Games, Inc. All Rights Reserved.
+
+using System.Diagnostics;
+
+namespace UnrealBuildTool.Rules
+{
+	[SupportedPlatforms("Win64")]
+	public class DatasmithMax2026 : DatasmithMaxBase
+	{
+		public DatasmithMax2026(ReadOnlyTargetRules Target)
+			: base(Target)
+		{
+		}
+		public override string GetMaxVersion() { return "2026"; }
+	}
+}
\ No newline at end of file
diff --git a/Engine/Source/Programs/Enterprise/Datasmith/DatasmithMaxExporter/DatasmithMax2026.Target.cs b/Engine/Source/Programs/Enterprise/Datasmith/DatasmithMaxExporter/DatasmithMax2026.Target.cs
new file mode 100644
--- /dev/null
+++ b/Engine/Source/Programs/Enterprise/Datasmith/DatasmithMaxExporter/DatasmithMax2026.Target.cs
@@ -0,0 +1,15 @@
+// Copyright Epic Games, Inc. All Rights Reserved.
+
+using UnrealBuildTool;
+using System.Collections.Generic;
+
+public class DatasmithMax2026Target : DatasmithMaxBaseTarget
+{
+	public DatasmithMax2026Target(TargetInfo Target) : base(Target)
+	{
+		LaunchModuleName = "DatasmithMax2026";
+		ExeBinariesSubFolder = @"3DSMax\2026";
+
+		AddCopyPostBuildStep(Target);
+	}
+}
diff --git a/Engine/Source/Programs/Enterprise/Datasmith/DatasmithMaxExporter/Plugin/PackageContents.xml b/Engine/Source/Programs/Enterprise/Datasmith/DatasmithMaxExporter/Plugin/PackageContents.xml
--- a/Engine/Source/Programs/Enterprise/Datasmith/DatasmithMaxExporter/Plugin/PackageContents.xml
+++ b/Engine/Source/Programs/Enterprise/Datasmith/DatasmithMaxExporter/Plugin/PackageContents.xml
@@ -16,7 +16,12 @@
 	<!--  todo: Icon="./DatasmithIcon32.png" -->
     <CompanyDetails Name="EpicGames, Inc." Url="www.epicgames.com"/>
 
-    <RuntimeRequirements OS="Win64" Platform="3ds Max" SeriesMin="2017" SeriesMax="2025" />
+    <RuntimeRequirements OS="Win64" Platform="3ds Max" SeriesMin="2017" SeriesMax="2026" />
+
+    <Components Description="plugins parts">
+        <RuntimeRequirements OS="Win64" Platform="3ds Max" SeriesMin="2026" SeriesMax="2026" />
+        <ComponentEntry ModuleName="./2025/DatasmithMax2026.gup" />
+    </Components>
 
     <Components Description="plugins parts">
         <RuntimeRequirements OS="Win64" Platform="3ds Max" SeriesMin="2025" SeriesMax="2025" />
@@ -64,18 +69,18 @@
     </Components>
 
     <Components Description="pre-start-up scripts parts">
-        <RuntimeRequirements OS="Win64" Platform="3ds Max" SeriesMin="2017" SeriesMax="2025" />
+        <RuntimeRequirements OS="Win64" Platform="3ds Max" SeriesMin="2017" SeriesMax="2026" />
         <ComponentEntry ModuleName="./UI/DatasmithPreStartUp.ms" />
         <ComponentEntry ModuleName="./UI/UnrealStudio_Datasmith_Attributes_Holder_Modifier.ms" />
     </Components>
 
     <Components Description="macroscripts parts">
-        <RuntimeRequirements OS="Win64" Platform="3ds Max" SeriesMin="2017" SeriesMax="2025" />
+        <RuntimeRequirements OS="Win64" Platform="3ds Max" SeriesMin="2017" SeriesMax="2026" />
         <ComponentEntry ModuleName="./UI/DatasmithMacros.mcr" />
     </Components>
 
     <Components Description="post-start-up scripts parts">
-        <RuntimeRequirements OS="Win64" Platform="3ds Max" SeriesMin="2017" SeriesMax="2025" />
+        <RuntimeRequirements OS="Win64" Platform="3ds Max" SeriesMin="2017" SeriesMax="2026" />
         <ComponentEntry ModuleName="./UI/DatasmithPostStartUp.ms" />
     </Components>
 
diff --git a/Engine/Source/Programs/Enterprise/Datasmith/DatasmithMaxExporter/Public/DatasmithMaxExporterDefines.h b/Engine/Source/Programs/Enterprise/Datasmith/DatasmithMaxExporter/Public/DatasmithMaxExporterDefines.h
--- a/Engine/Source/Programs/Enterprise/Datasmith/DatasmithMaxExporter/Public/DatasmithMaxExporterDefines.h
+++ b/Engine/Source/Programs/Enterprise/Datasmith/DatasmithMaxExporter/Public/DatasmithMaxExporterDefines.h
@@ -1,6 +1,17 @@
 // Copyright Epic Games, Inc. All Rights Reserved.
 #pragma once
 
+#ifdef __clang__
+#ifndef IGNORE_CLANG_WARNINGS
+#define IGNORE_CLANG_WARNINGS
+#pragma clang error "TODO: replace deprecated feature"
+#pragma clang diagnostic push
+#pragma clang diagnostic ignored "-Wnon-pod-varargs"
+#pragma clang diagnostic ignored "-Waddress-of-temporary"
+#endif // IGNORE_CLANG_WARNINGS
+#endif // __clang__
+
+
 #ifndef MAX_INCLUDES_START
 #define MAX_INCLUDES_START \
 	THIRD_PARTY_INCLUDES_START \
