[Datasmith] Fixed TBB dylibs packaging for Mac of SketchUp and ArchiCAD plugins

diff --git a/Engine/Source/Programs/Enterprise/Datasmith/DatasmithARCHICADExporter/Build/Config/setupdylibs.sh b/Engine/Source/Programs/Enterprise/Datasmith/DatasmithARCHICADExporter/Build/Config/setupdylibs.sh
--- a/Engine/Source/Programs/Enterprise/Datasmith/DatasmithARCHICADExporter/Build/Config/setupdylibs.sh
+++ b/Engine/Source/Programs/Enterprise/Datasmith/DatasmithARCHICADExporter/Build/Config/setupdylibs.sh
@@ -89,8 +89,12 @@ SetUpDll() {
 }
 
 SetUpThirdPartyDll $dylibLibFreeImage "$EnginePath/Binaries/ThirdParty/FreeImage/Mac/$dylibLibFreeImage"
-SetUpThirdPartyDll $dylibtbb "$EnginePath/Source/ThirdParty/Intel/TBB/Deploy/oneTBB-2021.13.0/Mac/lib/$dylibtbb"
-SetUpThirdPartyDll $dylibtbbmalloc "$EnginePath/Source/ThirdParty/Intel/TBB/Deploy/oneTBB-2021.13.0/Mac/lib/$dylibtbbmalloc"
+
+tbbLibDir="$EnginePath/Source/ThirdParty/Intel/TBB/Deploy/oneTBB-2021.13.0/Mac/lib"
+dylibLibTbbFilenames=$(find "${tbbLibDir}" -name "libtbb*.dylib" -maxdepth 1 -exec basename {} \;)
+for dylibName in $dylibLibTbbFilenames; do 
+    SetUpThirdPartyDll $dylibName "${tbbLibDir}/$dylibName"
+done
 
 SetUpDll DatasmithUE4ArchiCAD.dylib
 SetUpDll DatasmithUE4ArchiCAD-Mac-Debug.dylib
diff --git a/Engine/Source/Programs/Enterprise/Datasmith/DatasmithSketchUpRubyExporter/Build/build_mac.sh b/Engine/Source/Programs/Enterprise/Datasmith/DatasmithSketchUpRubyExporter/Build/build_mac.sh
--- a/Engine/Source/Programs/Enterprise/Datasmith/DatasmithSketchUpRubyExporter/Build/build_mac.sh
+++ b/Engine/Source/Programs/Enterprise/Datasmith/DatasmithSketchUpRubyExporter/Build/build_mac.sh
@@ -29,8 +29,8 @@ SetupIntermediateDirectory() {
     mkdir -p "$IntermediatePath"
 
     dylibLibFreeImage=libfreeimage-3.18.0.dylib
-    dylibLibTbb=libtbb.dylib
-    dylibLibTbbMalloc=libtbbmalloc.dylib
+    tbbLibDir="$EnginePath/Source/ThirdParty/Intel/TBB/Deploy/oneTBB-2021.13.0/Mac/lib"
+    dylibLibTbbFilenames=$(find "${tbbLibDir}" -name "libtbb*.dylib" -maxdepth 1 -exec basename {} \;)
 
     # Copy and fixup dylibs
     Dylibs="$IntermediatePath/Dylibs"
@@ -39,12 +39,13 @@ SetupIntermediateDirectory() {
 
     cp "$EnginePath/Binaries/ThirdParty/FreeImage/Mac/$dylibLibFreeImage" "$Dylibs"
     cp "$EnginePath/Binaries/Mac/DatasmithSDK/DatasmithSDK.dylib" "$Dylibs"
-    cp "$EnginePath/Source/ThirdParty/Intel/TBB/Deploy/oneTBB-2021.13.0/Mac/lib/$dylibLibTbb" "$Dylibs"
-    cp "$EnginePath/Source/ThirdParty/Intel/TBB/Deploy/oneTBB-2021.13.0/Mac/lib/$dylibLibTbbMalloc" "$Dylibs"
+
+    for dylibName in $dylibLibTbbFilenames; do 
+        cp "${tbbLibDir}/${dylibName}" "$Dylibs"
+        chmod 777 "$Dylibs/${dylibName}"
+    done
 
     chmod 777 "$Dylibs/$dylibLibFreeImage"
-    chmod 777 "$Dylibs/$dylibLibTbb"
-    chmod 777 "$Dylibs/$dylibLibTbbMalloc"
 
     install_name_tool -id @loader_path/Dylibs/DatasmithSDK.dylib "$Dylibs/DatasmithSDK.dylib" > /dev/null 2>&1
     install_name_tool -id @loader_path/Dylibs/$dylibLibFreeImage "$Dylibs/$dylibLibFreeImage" > /dev/null 2>&1
