[Datasmith] Fix 3dsmax plugin for 2026

diff --git a/Engine/Source/Programs/Enterprise/Datasmith/DatasmithMaxExporter/Plugin/PackageContents.xml b/Engine/Source/Programs/Enterprise/Datasmith/DatasmithMaxExporter/Plugin/PackageContents.xml
--- a/Engine/Source/Programs/Enterprise/Datasmith/DatasmithMaxExporter/Plugin/PackageContents.xml
+++ b/Engine/Source/Programs/Enterprise/Datasmith/DatasmithMaxExporter/Plugin/PackageContents.xml
@@ -20,7 +20,7 @@
 
     <Components Description="plugins parts">
         <RuntimeRequirements OS="Win64" Platform="3ds Max" SeriesMin="2026" SeriesMax="2026" />
-        <ComponentEntry ModuleName="./2025/DatasmithMax2026.gup" />
+        <ComponentEntry ModuleName="./2026/DatasmithMax2026.gup" />
     </Components>
 
     <Components Description="plugins parts">
