Fixed compilation error of ArchiCAD exporter on Mac

#rnx

diff --git a/Engine/Source/Programs/Enterprise/Datasmith/DatasmithARCHICADExporter/Build/Config/SetupBuildEnvironment.sh b/Engine/Source/Programs/Enterprise/Datasmith/DatasmithARCHICADExporter/Build/Config/SetupBuildEnvironment.sh
new file mode 100644
--- /dev/null
+++ b/Engine/Source/Programs/Enterprise/Datasmith/DatasmithARCHICADExporter/Build/Config/SetupBuildEnvironment.sh
@@ -0,0 +1,112 @@
+#!/bin/sh
+
+set -e
+
+echo "Setting up build environment..."
+
+ConfigPath=`dirname "$0"`
+pushd $ConfigPath/..
+projectPath=`pwd`
+cd ../../../../../..
+EnginePath=`pwd`
+cd ..
+RootDir=`pwd`
+popd
+
+echo "UE_SDKS_ROOT = ${UE_SDKS_ROOT}"
+
+if [ -z ${UE_SDKS_ROOT+x} ]; then
+    echo "UE_SDKS_ROOT is unset";
+    exit -1;
+fi
+
+if [ ! -f "$ConfigPath/SDKsRoot.xcconfig" ]; then
+    echo "Create $ConfigPath/SDKsRoot.xcconfig"
+    echo UESDKRoot = $UE_SDKS_ROOT > "$ConfigPath/SDKsRoot.xcconfig"
+    echo ArchiCADLocalMacDir = "$RootDir/ArchiCADLocalMacDir" >> "$ConfigPath/SDKsRoot.xcconfig"
+    echo UE_Engine = $EnginePath >> "$ConfigPath/SDKsRoot.xcconfig"
+fi
+
+# ls -l $UE_SDKS_ROOT/HostMac/Mac/Archicad/Archicad.zip
+echo "Unzipping $UE_SDKS_ROOT/HostMac/Mac/Archicad/Archicad.zip"
+unzip -qq $UE_SDKS_ROOT/HostMac/Mac/Archicad/Archicad.zip -d $RootDir/ArchiCADLocalMacDir
+# ls -l $RootDir/ArchiCADLocalMacDir/Archicad/23/Support/Frameworks/ACOperations.framework/Versions/Current
+
+# Remove ArchiCAD resource tool from quarantine
+pushd "$RootDir/ArchiCADLocalMacDir/Archicad"
+	chmod 777 23.1/Support/Tools/OSX/ResConv
+	/usr/bin/xattr -r -d com.apple.quarantine 23.1/Support/Tools/OSX/ResConv
+
+	chmod 777 24/Support/Tools/OSX/ResConv
+	/usr/bin/xattr -r -d com.apple.quarantine 24/Support/Tools/OSX/ResConv
+
+	chmod 777 25/Support/Tools/OSX/ResConv
+	/usr/bin/xattr -r -d com.apple.quarantine 25/Support/Tools/OSX/ResConv
+
+	chmod 777 26/Support/Tools/OSX/ResConv
+	/usr/bin/xattr -r -d com.apple.quarantine 26/Support/Tools/OSX/ResConv
+
+	chmod 777 27/Support/Tools/OSX/ResConv
+	/usr/bin/xattr -r -d com.apple.quarantine 27/Support/Tools/OSX/ResConv
+
+	chmod 777 28/Support/Tools/OSX/ResConv
+	/usr/bin/xattr -r -d com.apple.quarantine 28/Support/Tools/OSX/ResConv
+popd
+
+OurDylibFolder=$projectPath/Dylibs
+
+mkdir -p "$OurDylibFolder"
+
+dylibLibFreeImage=libfreeimage-3.18.0.dylib
+dylibtbb=libtbb.dylib
+dylibtbbmalloc=libtbbmalloc.dylib
+
+SetUpThirdPartyDll() {
+	DylibName=$1
+	DylibPath=$2
+	if [[ "$DylibPath" -nt "$OurDylibFolder/$DylibName" ]]; then
+		if [ -f "$OurDylibFolder/$DylibName" ]; then
+			unlink "$OurDylibFolder/$DylibName"
+		fi
+	fi
+	if [ ! -f "$OurDylibFolder/$DylibName" ]; then
+		echo "Copy $DylibName"
+		cp "$DylibPath" "$OurDylibFolder"
+		chmod +w "$OurDylibFolder/$DylibName"
+		install_name_tool -id @loader_path/$DylibName "$OurDylibFolder/$DylibName" > /dev/null 2>&1
+	fi
+}
+
+SetUpDll() {
+	DylibName=$1
+	OriginalDylibPath="$EnginePath/Binaries/Mac/DatasmithUE4ArchiCAD/$DylibName"
+
+	if [[ "$OriginalDylibPath" -nt "$OurDylibFolder/$DylibName" ]]; then
+		if [ -f "$OurDylibFolder/$DylibName" ]; then
+			unlink "$OurDylibFolder/$DylibName"
+		fi
+	fi
+	if [ ! -f "$OurDylibFolder/$DylibName" ]; then
+		if [ -f "$OriginalDylibPath" ]; then
+			echo "Copy $DylibName"
+			cp "$OriginalDylibPath" "$OurDylibFolder"
+			install_name_tool -id @loader_path/$DylibName "$OurDylibFolder/$DylibName" > /dev/null 2>&1
+			install_name_tool -change @rpath/$dylibLibFreeImage @loader_path/$dylibLibFreeImage "$OurDylibFolder/$DylibName" > /dev/null 2>&1
+		else
+			echo "Missing $DylibName"
+		fi
+	fi
+}
+
+SetUpThirdPartyDll $dylibLibFreeImage "$EnginePath/Binaries/ThirdParty/FreeImage/Mac/$dylibLibFreeImage"
+
+tbbLibDir="$EnginePath/Source/ThirdParty/Intel/TBB/Deploy/oneTBB-2021.13.0/Mac/lib"
+dylibLibTbbFilenames=$(find "${tbbLibDir}" -name "libtbb*.dylib" -maxdepth 1 -exec basename {} \;)
+for dylibName in $dylibLibTbbFilenames; do 
+    SetUpThirdPartyDll $dylibName "${tbbLibDir}/$dylibName"
+done
+
+SetUpDll DatasmithUE4ArchiCAD.dylib
+SetUpDll DatasmithUE4ArchiCAD-Mac-Debug.dylib
+
+echo "Build environment set."
diff --git a/Engine/Source/Programs/Enterprise/Datasmith/DatasmithARCHICADExporter/Build/Config/UE_AC.xcconfig b/Engine/Source/Programs/Enterprise/Datasmith/DatasmithARCHICADExporter/Build/Config/UE_AC.xcconfig
--- a/Engine/Source/Programs/Enterprise/Datasmith/DatasmithARCHICADExporter/Build/Config/UE_AC.xcconfig
+++ b/Engine/Source/Programs/Enterprise/Datasmith/DatasmithARCHICADExporter/Build/Config/UE_AC.xcconfig
@@ -2,7 +2,7 @@
 
 #include? "SDKsRoot.xcconfig"
 
-AC_SDK_PATH = $(UESDKRoot)/HostMac/Mac/Archicad/$(AC_SDK_DIR)
+AC_SDK_PATH = $(ArchiCADLocalMacDir)/Archicad/$(AC_SDK_DIR)
 
 AC_INCLUDES_PATH = "$(AC_SDK_PATH)/Support/Inc" "$(AC_SDK_PATH)/Support/Modules/GSRoot/STL/imp" "$(AC_SDK_PATH)/Support/Modules/GSRoot" "$(AC_SDK_PATH)/Support/Modules/InputOutput"/** "$(AC_SDK_PATH)/Support/Modules/GSModelDevLib"/** "$(AC_SDK_PATH)/Support/Modules/GSUtils"/** "$(AC_SDK_PATH)/Support/Modules/GSModeler"/** "$(AC_SDK_PATH)/Support/Modules/Model3D"/** "$(AC_SDK_PATH)/Support/Modules/MPIInterfaceDevLib"/** "$(AC_SDK_PATH)/Support/Modules/Geometry"/** "$(AC_SDK_PATH)/Support/Modules/DGLib"/** "$(AC_SDK_PATH)/Support/Modules/GDL"/** "$(AC_SDK_PATH)/Support/Modules/GXImage"/** "$(AC_SDK_PATH)/Support/Modules/GX"/** "$(AC_SDK_PATH)/Support/Modules/DGGraphix"/** "$(AC_SDK_PATH)/Support/Modules/Graphix"/** "$(AC_SDK_PATH)/Support/Modules/TextEngine"/** "$(AC_SDK_PATH)/Support/Modules/GXImageBase"/** "$(AC_SDK_PATH)/Support/Modules/PointCloudManager" "$(AC_SDK_PATH)/Support/Modules/PointCloud" "$(AC_SDK_PATH)/Support/Modules/Pattern" "$(AC_SDK_PATH)/Support/Modules/Brep" "$(AC_SDK_PATH)/Support/Modules/VBAttributes" "$(AC_SDK_PATH)/Support/Modules/GSModeler2D" "$(AC_SDK_PATH)/Support/Modules/ObjectDatabase" "$(AC_SDK_PATH)/Support/Modules/UDLib" "$(AC_SDK_PATH)/Support/Modules/RS" "$(UE_Engine)/Binaries/Mac/DatasmithSDK"
 
diff --git a/Engine/Source/Programs/Enterprise/Datasmith/DatasmithARCHICADExporter/Build/Config/UE_AC_ARM.xcconfig b/Engine/Source/Programs/Enterprise/Datasmith/DatasmithARCHICADExporter/Build/Config/UE_AC_ARM.xcconfig
--- a/Engine/Source/Programs/Enterprise/Datasmith/DatasmithARCHICADExporter/Build/Config/UE_AC_ARM.xcconfig
+++ b/Engine/Source/Programs/Enterprise/Datasmith/DatasmithARCHICADExporter/Build/Config/UE_AC_ARM.xcconfig
@@ -2,7 +2,7 @@
 
 #include? "SDKsRoot.xcconfig"
 
-AC_SDK_PATH = $(UESDKRoot)/HostMac/Mac/Archicad/$(AC_SDK_DIR)
+AC_SDK_PATH = $(ArchiCADLocalMacDir)/Archicad/$(AC_SDK_DIR)
 
 AC_INCLUDES_PATH = "$(AC_SDK_PATH)/Support/Inc" "$(AC_SDK_PATH)/Support/Modules/GSRoot/STL/imp" "$(AC_SDK_PATH)/Support/Modules/GSRoot" "$(AC_SDK_PATH)/Support/Modules/InputOutput"/** "$(AC_SDK_PATH)/Support/Modules/GSModelDevLib"/** "$(AC_SDK_PATH)/Support/Modules/GSUtils"/** "$(AC_SDK_PATH)/Support/Modules/GSModeler"/** "$(AC_SDK_PATH)/Support/Modules/Model3D"/** "$(AC_SDK_PATH)/Support/Modules/MPIInterfaceDevLib"/** "$(AC_SDK_PATH)/Support/Modules/Geometry"/** "$(AC_SDK_PATH)/Support/Modules/DGLib"/** "$(AC_SDK_PATH)/Support/Modules/GDL"/** "$(AC_SDK_PATH)/Support/Modules/GXImage"/** "$(AC_SDK_PATH)/Support/Modules/GX"/** "$(AC_SDK_PATH)/Support/Modules/DGGraphix"/** "$(AC_SDK_PATH)/Support/Modules/Graphix"/** "$(AC_SDK_PATH)/Support/Modules/TextEngine"/** "$(AC_SDK_PATH)/Support/Modules/GXImageBase"/** "$(AC_SDK_PATH)/Support/Modules/PointCloudManager" "$(AC_SDK_PATH)/Support/Modules/PointCloud" "$(AC_SDK_PATH)/Support/Modules/Pattern" "$(AC_SDK_PATH)/Support/Modules/CADInfrastructureBase" "$(AC_SDK_PATH)/Support/Modules/Brep" "$(AC_SDK_PATH)/Support/Modules/VBAttributes" "$(AC_SDK_PATH)/Support/Modules/GSModeler2D" "$(AC_SDK_PATH)/Support/Modules/ObjectDatabase" "$(AC_SDK_PATH)/Support/Modules/UDLib" "$(AC_SDK_PATH)/Support/Modules/RS" "$(UE_Engine)/Binaries/Mac/DatasmithSDK"
 
diff --git a/Engine/Source/Programs/Enterprise/Datasmith/DatasmithARCHICADExporter/Build/Config/setupdylibs.sh b/Engine/Source/Programs/Enterprise/Datasmith/DatasmithARCHICADExporter/Build/Config/setupdylibs.sh
--- a/Engine/Source/Programs/Enterprise/Datasmith/DatasmithARCHICADExporter/Build/Config/setupdylibs.sh
+++ b/Engine/Source/Programs/Enterprise/Datasmith/DatasmithARCHICADExporter/Build/Config/setupdylibs.sh
@@ -1,100 +1,3 @@
 #!/bin/sh
 
 set -e
-
-ConfigPath=`dirname "$0"`
-pushd $ConfigPath/..
-projectPath=`pwd`
-cd ../../../../../..
-EnginePath=`pwd`
-popd
-
-echo "UE_SDKS_ROOT = ${UE_SDKS_ROOT}"
-
-if [ -z ${UE_SDKS_ROOT+x} ]; then
-    echo "UE_SDKS_ROOT is unset";
-    exit -1;
-fi
-
-if [ ! -f "$ConfigPath/SDKsRoot.xcconfig" ]; then
-    echo "Create $ConfigPath/SDKsRoot.xcconfig"
-    echo UESDKRoot = $UE_SDKS_ROOT > "$ConfigPath/SDKsRoot.xcconfig"
-    echo UE_Engine = $EnginePath >> "$ConfigPath/SDKsRoot.xcconfig"
-fi
-
-# Remove ArchiCAD resource tool from quarantine
-pushd "$UE_SDKS_ROOT/HostMac/Mac/Archicad"
-	chmod 777 23.1/Support/Tools/OSX/ResConv
-	/usr/bin/xattr -r -d com.apple.quarantine 23.1/Support/Tools/OSX/ResConv
-
-	chmod 777 24/Support/Tools/OSX/ResConv
-	/usr/bin/xattr -r -d com.apple.quarantine 24/Support/Tools/OSX/ResConv
-
-	chmod 777 25/Support/Tools/OSX/ResConv
-	/usr/bin/xattr -r -d com.apple.quarantine 25/Support/Tools/OSX/ResConv
-
-	chmod 777 26/Support/Tools/OSX/ResConv
-	/usr/bin/xattr -r -d com.apple.quarantine 26/Support/Tools/OSX/ResConv
-
-	chmod 777 27/Support/Tools/OSX/ResConv
-	/usr/bin/xattr -r -d com.apple.quarantine 27/Support/Tools/OSX/ResConv
-
-	chmod 777 28/Support/Tools/OSX/ResConv
-	/usr/bin/xattr -r -d com.apple.quarantine 28/Support/Tools/OSX/ResConv
-popd
-
-OurDylibFolder=$projectPath/Dylibs
-
-mkdir -p "$OurDylibFolder"
-
-dylibLibFreeImage=libfreeimage-3.18.0.dylib
-dylibtbb=libtbb.dylib
-dylibtbbmalloc=libtbbmalloc.dylib
-
-SetUpThirdPartyDll() {
-	DylibName=$1
-	DylibPath=$2
-	if [[ "$DylibPath" -nt "$OurDylibFolder/$DylibName" ]]; then
-		if [ -f "$OurDylibFolder/$DylibName" ]; then
-			unlink "$OurDylibFolder/$DylibName"
-		fi
-	fi
-	if [ ! -f "$OurDylibFolder/$DylibName" ]; then
-		echo "Copy $DylibName"
-		cp "$DylibPath" "$OurDylibFolder"
-		chmod +w "$OurDylibFolder/$DylibName"
-		install_name_tool -id @loader_path/$DylibName "$OurDylibFolder/$DylibName" > /dev/null 2>&1
-	fi
-}
-
-SetUpDll() {
-	DylibName=$1
-	OriginalDylibPath="$EnginePath/Binaries/Mac/DatasmithUE4ArchiCAD/$DylibName"
-
-	if [[ "$OriginalDylibPath" -nt "$OurDylibFolder/$DylibName" ]]; then
-		if [ -f "$OurDylibFolder/$DylibName" ]; then
-			unlink "$OurDylibFolder/$DylibName"
-		fi
-	fi
-	if [ ! -f "$OurDylibFolder/$DylibName" ]; then
-		if [ -f "$OriginalDylibPath" ]; then
-			echo "Copy $DylibName"
-			cp "$OriginalDylibPath" "$OurDylibFolder"
-			install_name_tool -id @loader_path/$DylibName "$OurDylibFolder/$DylibName" > /dev/null 2>&1
-			install_name_tool -change @rpath/$dylibLibFreeImage @loader_path/$dylibLibFreeImage "$OurDylibFolder/$DylibName" > /dev/null 2>&1
-		else
-			echo "Missing $DylibName"
-		fi
-	fi
-}
-
-SetUpThirdPartyDll $dylibLibFreeImage "$EnginePath/Binaries/ThirdParty/FreeImage/Mac/$dylibLibFreeImage"
-
-tbbLibDir="$EnginePath/Source/ThirdParty/Intel/TBB/Deploy/oneTBB-2021.13.0/Mac/lib"
-dylibLibTbbFilenames=$(find "${tbbLibDir}" -name "libtbb*.dylib" -maxdepth 1 -exec basename {} \;)
-for dylibName in $dylibLibTbbFilenames; do 
-    SetUpThirdPartyDll $dylibName "${tbbLibDir}/$dylibName"
-done
-
-SetUpDll DatasmithUE4ArchiCAD.dylib
-SetUpDll DatasmithUE4ArchiCAD-Mac-Debug.dylib
