Uv repetition is not imported correctly
This change implements exporting UVs

diff --git a/Engine/Source/Programs/Enterprise/Datasmith/DatasmithRhinoExporter/Private/ElementExporters/DatasmithRhinoMaterialExporter.cs b/Engine/Source/Programs/Enterprise/Datasmith/DatasmithRhinoExporter/Private/ElementExporters/DatasmithRhinoMaterialExporter.cs
--- a/Engine/Source/Programs/Enterprise/Datasmith/DatasmithRhinoExporter/Private/ElementExporters/DatasmithRhinoMaterialExporter.cs
+++ b/Engine/Source/Programs/Enterprise/Datasmith/DatasmithRhinoExporter/Private/ElementExporters/DatasmithRhinoMaterialExporter.cs
@@ -3,11 +3,11 @@
 using DatasmithRhino.ExportContext;
 using DatasmithRhino.Utils;
 using Rhino.DocObjects;
-using Rhino.Geometry;
 using System.Collections.Generic;
 using System.Drawing;
 using System;
 using Rhino;
+using Rhino.Geometry;
 
 namespace DatasmithRhino.ElementExporters
 {
@@ -192,11 +192,19 @@ namespace DatasmithRhino.ElementExporters
 					}
 					break;
 			}
+
+			
+			// Sadly this can't be done in GetUVParameter() since the Expression is not exposed in FUVEditParameters
+			// Replace this if FUVEditParameters.SetTilingPivot() is implemented in Datasmith Facade
+			FDatasmithFacadeMaterialExpressionColor TilingPivotExpression = DatasmithRhinoUtilities.GetMaterialExpressionByName("Tiling_Pivot", DSMaterial) as FDatasmithFacadeMaterialExpressionColor;
+			// Tiling Pivot for Textures need to be adjusted to (0, 1) so it matches Rhino Textures
+			TilingPivotExpression?.SetColor(0, 1, 0, 1);
 		}
 
 		/// <summary>
 		/// Returns the UV Parameters for the given Texture.
-		/// This function uses the same operations as in FOpenNurbsTranslatorImpl::TranslateMaterialTable(), improvements should be applied to both functions.
+		/// If the texture uses a non-surface projection in Rhino, it'll be applied like if it was a surface projection instead
+		/// since it's the way it's used in DataSmith and log a warning (warning only available from Version 8+)
 		/// </summary>
 		/// <param name="RhinoTexture"></param>
 		/// <returns></returns>
@@ -208,38 +216,60 @@ namespace DatasmithRhino.ElementExporters
 			// Use cached texture coordinates(channel 0)
 			UVParameters.SetChannelIndex(0);
 
-			//// Extract the UV tiling, offset and rotation angle from the UV transform matrix
-			Transform RotationTransform, OrthogonalTransform;
-			Vector3d Translation, Scale;
-			RhinoTexture.UvwTransform.DecomposeAffine(out Translation, out RotationTransform, out OrthogonalTransform, out Scale);
+			float RepeatX;
+			float RepeatY;
+			float Degrees;
+			float OffsetX;
+			float OffsetY;
 
-			double RotX, RotY, RotZ;
-			if (!RotationTransform.GetYawPitchRoll(out RotX, out RotY, out RotZ))
+			// These two implementations are equivalent, however from Rhino8 and above, the API implemented
+			// some new ways to gather UV transformation data that is said to be recommended and more precise
+			// hence this separation so it supports a proper implementation in newer versions
+			#if RHINO_VERSION_8_OR_NEWER
 			{
-				//This is not a valid rotation make sure the angles are at 0;
-				RotX = RotY = RotZ = 0;
+				// this check is only available in 8+
+				if(RhinoTexture.ProjectionMode != TextureProjectionModes.MappingChannel)
+					RhinoApp.WriteLine($"Warning. Texture projection mode <{RhinoTexture.ProjectionMode}> is not supported in texture: <{RhinoTexture.FileName}>.");
+				
+				RhinoTexture.UvwTransform.DecomposeTextureMapping(out Vector3d Offset, out Vector3d Repeat, out Vector3d Rotation);
+				RepeatX = (float)Repeat.X;
+				RepeatY = (float)Repeat.Y;
+				float RotationRadians = - (float) Rotation.Z; 
+				Degrees = RotationRadians * (float)( 180.0f / Math.PI);
+				float UserOffsetU = (float) - Offset.X;
+				float UserOffsetV = (float) Offset.Y;
+
+				float Cos = (float)Math.Cos(RotationRadians);
+				float Sin = (float)Math.Sin(RotationRadians);
+				OffsetX = Cos * UserOffsetU - Sin * UserOffsetV;
+				OffsetY = Sin * UserOffsetU + Cos * UserOffsetV;
 			}
-			else
+			#else
 			{
-				RotX = DatasmithRhinoUtilities.RadianToDegree(RotX);
-				RotY = DatasmithRhinoUtilities.RadianToDegree(RotY);
-				RotZ = DatasmithRhinoUtilities.RadianToDegree(RotZ);
-			}
-
-			UVParameters.SetUVTiling((float)Scale.X, (float)Scale.Y);
+				float a = (float) RhinoTexture.UvwTransform[0, 0]; 
+				float b = (float) RhinoTexture.UvwTransform[0, 1];
+				float c = (float) RhinoTexture.UvwTransform[1, 0]; 
+				float d = (float) RhinoTexture.UvwTransform[1, 1]; 
+				float tx = (float) RhinoTexture.UvwTransform[0, 3];
+				float ty = (float) RhinoTexture.UvwTransform[1, 3];
 			
-			//If the tiling vector is not zero.
-			if (Math.Abs(Scale.X) > float.Epsilon && Math.Abs(Scale.Y) > float.Epsilon)
-			{
-				float UVOffsetX = (float) (Translation.X / Scale.X);
-				float UVOffsetY = (float)-(Translation.Y / Scale.Y + 0.5f - 0.5f / Scale.Y);
+				RepeatX = (float) Math.Sqrt( (a * a) + (b * b));
+				RepeatY = (float) Math.Sqrt( (c * c) + (d * d));
+				float Radians =  (float) Math.Atan2(b, a);
+				Degrees = (float) (Radians * 180 / Math.PI);
 
-				UVParameters.SetUVOffset(UVOffsetX, UVOffsetY); // V-coordinate is inverted in Unreal
+				OffsetX = tx / RepeatX;
+				OffsetY = - ty / RepeatY;
 			}
-
-			// Rotation angle is reversed because V-axis points down in Unreal while it points up in OpenNurbs
-			UVParameters.SetRotationAngle((float)-RotX);
-
+			#endif
+			
+			UVParameters.SetUVTiling(RepeatX, RepeatY);
+			UVParameters.SetRotationAngle(Degrees);
+			UVParameters.SetUVOffset(OffsetX, OffsetY);
+			
+			// Rotation Pivot for Textures need to be adjusted to (0, 1) so it matches Rhino Textures
+			UVParameters.SetRotationPivot(0, 1, 0);
+			
 			return UVParameters;
 		}
 	}
diff --git a/Engine/Source/Programs/Enterprise/Datasmith/DatasmithRhinoExporter/Private/Utils/DatasmithRhinoUtilities.cs b/Engine/Source/Programs/Enterprise/Datasmith/DatasmithRhinoExporter/Private/Utils/DatasmithRhinoUtilities.cs
--- a/Engine/Source/Programs/Enterprise/Datasmith/DatasmithRhinoExporter/Private/Utils/DatasmithRhinoUtilities.cs
+++ b/Engine/Source/Programs/Enterprise/Datasmith/DatasmithRhinoExporter/Private/Utils/DatasmithRhinoUtilities.cs
@@ -404,5 +404,23 @@ namespace DatasmithRhino.Utils
 
 			return true;
 		}
+
+		public static FDatasmithFacadeMaterialExpression GetMaterialExpressionByName(string Name, FDatasmithFacadeUEPbrMaterial Material)
+		{
+			int QtyExpressions = Material.GetExpressionsCount();
+
+			for (int i = 0; i < QtyExpressions; i++)
+			{
+				FDatasmithFacadeMaterialExpression Expression = Material.GetExpression(i);
+				
+				if (Expression.GetName() == Name)
+				{
+					return Expression;
+				}
+					
+			}
+
+			return null;
+		}
 	}
 }
\ No newline at end of file
