[SKP DS Exporter] Export the current viewport
Redo/fix export viewport view to datasmith and directlink

diff --git a/Engine/Source/Programs/Enterprise/Datasmith/DatasmithSketchUpRubyExporter/Private/DatasmithSketchUpCamera.cpp b/Engine/Source/Programs/Enterprise/Datasmith/DatasmithSketchUpRubyExporter/Private/DatasmithSketchUpCamera.cpp
--- a/Engine/Source/Programs/Enterprise/Datasmith/DatasmithSketchUpRubyExporter/Private/DatasmithSketchUpCamera.cpp
+++ b/Engine/Source/Programs/Enterprise/Datasmith/DatasmithSketchUpRubyExporter/Private/DatasmithSketchUpCamera.cpp
@@ -213,7 +213,10 @@ bool FViewportCamera::Update(FExportContext& InContext)
 	SUCameraRef camera_ref = SU_INVALID;
 	if (SUModelGetCamera(InContext.ModelRef, &camera_ref) == SU_ERROR_NONE && SUIsValid(camera_ref))
 	{
-		TSharedPtr<FCamera> MainCamera = FCamera::Create(InContext, camera_ref, TEXT("viewport_camera"));
+		if(!MainCamera.IsValid())
+		{
+			MainCamera = FCamera::Create(InContext, camera_ref, TEXT("viewport_camera"));
+		}
 		
 		FMD5Hash CurrentHash = MainCamera->GetHash();
 		
diff --git a/Engine/Source/Programs/Enterprise/Datasmith/DatasmithSketchUpRubyExporter/Private/DatasmithSketchUpCamera.h b/Engine/Source/Programs/Enterprise/Datasmith/DatasmithSketchUpRubyExporter/Private/DatasmithSketchUpCamera.h
--- a/Engine/Source/Programs/Enterprise/Datasmith/DatasmithSketchUpRubyExporter/Private/DatasmithSketchUpCamera.h
+++ b/Engine/Source/Programs/Enterprise/Datasmith/DatasmithSketchUpRubyExporter/Private/DatasmithSketchUpCamera.h
@@ -50,5 +50,6 @@ namespace DatasmithSketchUp
 		
 	private:
 		FMD5Hash Hash;
+		TSharedPtr<FCamera> MainCamera;
 	};
 }
