[Rhino3D] Allow Viewport Camera and Current Camera Visibility in .udatasmith files
Reimplementation and fixes (direct link imp also included)

diff --git a/Engine/Source/Programs/Enterprise/Datasmith/DatasmithRhinoExporter/Private/DirectLink/DatasmithRhinoDirectLinkManager.cs b/Engine/Source/Programs/Enterprise/Datasmith/DatasmithRhinoExporter/Private/DirectLink/DatasmithRhinoDirectLinkManager.cs
--- a/Engine/Source/Programs/Enterprise/Datasmith/DatasmithRhinoExporter/Private/DirectLink/DatasmithRhinoDirectLinkManager.cs
+++ b/Engine/Source/Programs/Enterprise/Datasmith/DatasmithRhinoExporter/Private/DirectLink/DatasmithRhinoDirectLinkManager.cs
@@ -10,6 +10,7 @@ namespace DatasmithRhino.DirectLink
 	public class DatasmithRhinoDirectLinkManager
 	{
 		private const string UntitledSceneName = "Untitled";
+		private const double CameraUpdateInterval = 0.25;
 
 		public FDatasmithFacadeDirectLink DirectLink { get; private set; } = null;
 		public FDatasmithFacadeScene DatasmithScene { get; private set; } = null;
@@ -19,6 +20,8 @@ namespace DatasmithRhino.DirectLink
 		public bool bAutoSyncActive { get; private set; } = false;
 		// Slate UI is not available on Mac, we can't use the Connection UI and must provide this feature ourselves.
 		private string InternalCacheDirectory = Path.GetTempPath();
+		private double LastCameraUpdateTime = double.MinValue;
+		
 		public string CacheDirectory
 		{
 			get
@@ -163,6 +166,13 @@ namespace DatasmithRhino.DirectLink
 
 		private void OnRhinoIdle(object Sender, EventArgs e)
 		{
+			double ETA = DateTime.Now.TimeOfDay.TotalSeconds;
+			if (ETA - LastCameraUpdateTime > CameraUpdateInterval)
+			{
+				LastCameraUpdateTime = ETA;
+				ExportContext.ParseRhinoCameras();
+			}
+			
 			if (ExportContext.RhinoDocument != null && ExportContext.bIsDirty)
 			{
 				Synchronize(ExportContext.RhinoDocument);
diff --git a/Engine/Source/Programs/Enterprise/Datasmith/DatasmithRhinoExporter/Private/ElementExporters/DatasmithRhinoActorExporter.cs b/Engine/Source/Programs/Enterprise/Datasmith/DatasmithRhinoExporter/Private/ElementExporters/DatasmithRhinoActorExporter.cs
--- a/Engine/Source/Programs/Enterprise/Datasmith/DatasmithRhinoExporter/Private/ElementExporters/DatasmithRhinoActorExporter.cs
+++ b/Engine/Source/Programs/Enterprise/Datasmith/DatasmithRhinoExporter/Private/ElementExporters/DatasmithRhinoActorExporter.cs
@@ -28,13 +28,7 @@ namespace DatasmithRhino.ElementExporters
 		protected override IEnumerable<DatasmithActorInfo> GetElementsToSynchronize()
 		{
 			// We don't export the root, only its descendants.
-			List<DatasmithActorInfo> AllElements = ExportContext.SceneRoot.GetDescendantEnumerator()
-				.ToList();
-			
-			// We add "on the fly" generated actors (like viewport camera, etc)
-			AllElements.AddRange(GetCustomExportActorInfos());
-			
-			return AllElements;
+			return ExportContext.SceneRoot.GetDescendantEnumerator();
 		}
 
 		protected override FDatasmithFacadeElement CreateElement(DatasmithActorInfo ElementInfo)
@@ -83,43 +77,7 @@ namespace DatasmithRhino.ElementExporters
 		}
 
 		///// END IDatasmithRhinoElementExporter Interface /////
-
-		/// <summary>
-		/// Generates Actor Infos that would be used to create custom actors on the fly
-		/// required before exporting. These are actors that are not present in the Scene
-		/// but still are needed
-		/// </summary>
-		/// <returns></returns>
-		private List<DatasmithActorInfo> GetCustomExportActorInfos()
-		{
-			List<DatasmithActorInfo> Actors = new List<DatasmithActorInfo>();
-			
-			// Active viewport views are exported as Camera Actors
-			RhinoView[] ViewsList = ExportContext.RhinoDocument.Views.GetViewList(true, true);
-			
-			foreach (var View in ViewsList)
-			{
-				ViewportInfo ViewportInfo = new ViewportInfo(View.ActiveViewport);
-				
-				string Hash = FDatasmithFacadeElement.GetStringHash(
-					$"{ViewportInfo.CameraAngle}, {ViewportInfo.CameraLocation}, {ViewportInfo.CameraDirection}" +
-					$"{ViewportInfo.CameraUp},{ViewportInfo.ViewScale}");
-
-				// viewport camera names are pre fixed with "Current" substring
-				// cameras without this are "named views"
-				// this is just a design decision in naming convention so users can identify
-				// in UE/TM what sort of camera each one is 
-				string CameraObjectName = $"Current_{View.ActiveViewport.Name}";
-				
-				DatasmithCameraActorInfo ViewportCameraInfo =
-					new DatasmithCameraActorInfo(ViewportInfo, CameraObjectName, CameraObjectName, CameraObjectName, Hash);
-				
-				Actors.Add(ViewportCameraInfo);
-			}
-			
-			return Actors;
-		}
-		
+	
 		private FDatasmithFacadeActor CreateActor(DatasmithActorInfo Node)
 		{
 			FDatasmithFacadeActor ExportedActor = null;
@@ -353,15 +311,14 @@ namespace DatasmithRhino.ElementExporters
 			ViewportInfo Viewport = ActorInfo.RhinoCommonObject as ViewportInfo;
 			CameraActor.SetLabel(ActorInfo.UniqueLabel);
 
-			CameraActor.SetAspectRatio((float)Viewport.FrustumAspect);
-			CameraActor.SetFocalLength((float)Viewport.Camera35mmLensLength);
-
+			// 360 instead of 180 because Rhino gives the half fov
+			bool bIsVerticalFov = Viewport.ScreenPortAspect >= 1f;
+			CameraActor.SetFocalLength((float) Viewport.CameraAngle * 360f / (float)Math.PI, bIsVerticalFov);
+			CameraActor.SetAspectRatio((float)Viewport.ScreenPortAspect);
+			
 			CameraActor.SetCameraPosition((float)Viewport.CameraLocation.X, (float)Viewport.CameraLocation.Y, (float)Viewport.CameraLocation.Z);
 			CameraActor.SetCameraRotation((float)Viewport.CameraDirection.X, (float)Viewport.CameraDirection.Y, (float)Viewport.CameraDirection.Z,
 				(float)Viewport.CameraUp.X, (float)Viewport.CameraUp.Y, (float)Viewport.CameraUp.Z);
-			
-			CameraActor.SetAspectRatio((float) Viewport.FrustumAspect);
-			CameraActor.SetSensorWidth((float) Viewport.ScreenPort.Width);
 		}
 
 		private static void AddActorToParent(FDatasmithFacadeActor InDatasmithActor, DatasmithActorInfo InNode, FDatasmithFacadeScene InDatasmithScene)
diff --git a/Engine/Source/Programs/Enterprise/Datasmith/DatasmithRhinoExporter/Private/ExportContext/DatasmithCameraActorInfo.cs b/Engine/Source/Programs/Enterprise/Datasmith/DatasmithRhinoExporter/Private/ExportContext/DatasmithCameraActorInfo.cs
--- a/Engine/Source/Programs/Enterprise/Datasmith/DatasmithRhinoExporter/Private/ExportContext/DatasmithCameraActorInfo.cs
+++ b/Engine/Source/Programs/Enterprise/Datasmith/DatasmithRhinoExporter/Private/ExportContext/DatasmithCameraActorInfo.cs
@@ -27,12 +27,12 @@ namespace DatasmithRhino.ExportContext
 		/// We use this because Rhino does not provide an event for when the named view table is modified.
 		/// That means every time the document is modified we must check for named view changes, the hash allows us to avoid false-positives modifications.
 		/// </summary>
-		public String NamedViewHash { get; private set; }
+		public string ViewHash { get; private set; }
 
 		public DatasmithCameraActorInfo(ViewportInfo InViewportInfo, string InName, string InUniqueLabel, string InBaseLabel, string InHash)
 			: base(InViewportInfo, InName, InUniqueLabel, InBaseLabel)
 		{
-			NamedViewHash = InHash;
+			ViewHash = InHash;
 		}
 
 		public override void ApplyDiffs(DatasmithInfoBase OtherInfo)
@@ -41,7 +41,7 @@ namespace DatasmithRhino.ExportContext
 
 			if (OtherInfo is DatasmithCameraActorInfo OtherActorCameraInfo)
 			{
-				NamedViewHash = OtherActorCameraInfo.NamedViewHash;
+				ViewHash = OtherActorCameraInfo.ViewHash;
 			}
 			else
 			{
diff --git a/Engine/Source/Programs/Enterprise/Datasmith/DatasmithRhinoExporter/Private/ExportContext/DatasmithRhinoExportContext.cs b/Engine/Source/Programs/Enterprise/Datasmith/DatasmithRhinoExporter/Private/ExportContext/DatasmithRhinoExportContext.cs
--- a/Engine/Source/Programs/Enterprise/Datasmith/DatasmithRhinoExporter/Private/ExportContext/DatasmithRhinoExportContext.cs
+++ b/Engine/Source/Programs/Enterprise/Datasmith/DatasmithRhinoExporter/Private/ExportContext/DatasmithRhinoExportContext.cs
@@ -39,7 +39,8 @@ namespace DatasmithRhino.ExportContext
 		public Dictionary<InstanceDefinition, DatasmithActorInfo> InstanceDefinitionHierarchyNodeDictionary = new Dictionary<InstanceDefinition, DatasmithActorInfo>();
 		public Dictionary<Guid, DatasmithActorInfo> ObjectIdToHierarchyActorNodeDictionary = new Dictionary<Guid, DatasmithActorInfo>();
 		public Dictionary<Guid, DatasmithMeshInfo> ObjectIdToMeshInfoDictionary = new Dictionary<Guid, DatasmithMeshInfo>();
-		public Dictionary<string, DatasmithCameraActorInfo> NamedViewNameToCameraInfo = new Dictionary<string, DatasmithCameraActorInfo>();
+		public Dictionary<string, DatasmithCameraActorInfo> ViewNameToCameraInfo = new Dictionary<string, DatasmithCameraActorInfo>();
+		public Dictionary<string, DatasmithCameraActorInfo> ViewportViewNameToCameraInfo = new Dictionary<string, DatasmithCameraActorInfo>();
 		public Dictionary<string, DatasmithMaterialInfo> MaterialHashToMaterialInfo = new Dictionary<string, DatasmithMaterialInfo>();
 		public Dictionary<string, DatasmithTextureInfo> TextureHashToTextureInfo = new Dictionary<string, DatasmithTextureInfo>();
 		public Dictionary<int, string> GroupIndexToName = new Dictionary<int, string>();
@@ -128,7 +129,7 @@ namespace DatasmithRhino.ExportContext
 			InstanceDefinitionHierarchyNodeDictionary = new Dictionary<InstanceDefinition, DatasmithActorInfo>();
 			ObjectIdToHierarchyActorNodeDictionary = new Dictionary<Guid, DatasmithActorInfo>();
 			ObjectIdToMeshInfoDictionary = new Dictionary<Guid, DatasmithMeshInfo>();
-			NamedViewNameToCameraInfo = new Dictionary<string, DatasmithCameraActorInfo>();
+			ViewNameToCameraInfo = new Dictionary<string, DatasmithCameraActorInfo>();
 			MaterialHashToMaterialInfo = new Dictionary<string, DatasmithMaterialInfo>();
 			TextureHashToTextureInfo = new Dictionary<string, DatasmithTextureInfo>();
 			GroupIndexToName = new Dictionary<int, string>();
@@ -1390,57 +1391,82 @@ namespace DatasmithRhino.ExportContext
 
 		public void ParseRhinoCameras()
 		{
-			HashSet<string> ViewsToRemove = new HashSet<string>(NamedViewNameToCameraInfo.Keys);
-			bool bHasChanged = false;
+			List<ViewInfo> AllViewInfos = new List<ViewInfo>();
+			
+			// Active viewport views are exported as Camera Actors
+			RhinoView[] ViewsList = RhinoDocument.Views.GetViewList(true, true);
+			foreach (var View in ViewsList)
+			{
+				ViewInfo ViewportViewInfo = new ViewInfo(View.ActiveViewport);
+				// the name convention decided for viewport cameras is that they start with "Current_"
+				ViewportViewInfo.Name = "Current_" + ViewportViewInfo.Name;
+				AllViewInfos.Add(ViewportViewInfo);
+			}
+			
+			AllViewInfos.AddRange(RhinoDocument.NamedViews);
 
-			foreach (ViewInfo NamedView in RhinoDocument.NamedViews)
+			HashSet<string> ViewsToRemoveNames = new HashSet<string>(ViewNameToCameraInfo.Keys);
+			bool bHasChanged = false;
+			
+			foreach (ViewInfo View in AllViewInfos)
 			{
-				string NamedViewHash = DatasmithRhinoUtilities.GetNamedViewHash(NamedView);
+				string ViewHash = DatasmithRhinoUtilities.GetViewHash(View);
 				
-				if (NamedViewNameToCameraInfo.TryGetValue(NamedView.Name, out DatasmithCameraActorInfo ExistingCameraInfo))
+				if (ViewNameToCameraInfo.TryGetValue(View.Name, out DatasmithCameraActorInfo ExistingCameraInfo))
 				{
-					// Named view exists in DatasmithScene, check if it has changed.
-					ViewsToRemove.Remove(NamedView.Name);
-					
-					if (NamedViewHash != ExistingCameraInfo.NamedViewHash)
+					// view exists in DatasmithScene, check if it has changed.
+					ViewsToRemoveNames.Remove(View.Name);
+
+					if (string.Equals(ViewHash, ExistingCameraInfo.ViewHash))
 					{
-						DatasmithCameraActorInfo DiffInfo = GenerateActorCameraInfo(NamedView, NamedViewHash);
-						ExistingCameraInfo.ApplyDiffs(DiffInfo);
-						bHasChanged = true;
+						continue;
 					}
+						
+					DatasmithCameraActorInfo DiffInfo = GenerateActorCameraInfo(View, ViewHash);
+					ExistingCameraInfo.ApplyDiffs(DiffInfo);
+					bHasChanged = true;
 				}
 				else
 				{
 					// Named view does not exists in DatasmithScene, create a new one and add it to scene root actor.
-					DatasmithCameraActorInfo CameraInfo = GenerateActorCameraInfo(NamedView, NamedViewHash);
-
-					NamedViewNameToCameraInfo.Add(NamedView.Name, CameraInfo);
+					DatasmithCameraActorInfo CameraInfo = GenerateActorCameraInfo(View, ViewHash);
+					ViewNameToCameraInfo.Add(View.Name, CameraInfo);
 					SceneRoot.AddChild(CameraInfo);
 					bHasChanged = true;
 				}
 			}
-
+			
 			// Mark for removal all the named views that were not
-			foreach (string ViewToRemoveName in ViewsToRemove)
+			foreach (string ViewToRemoveName in ViewsToRemoveNames)
 			{
-				if (NamedViewNameToCameraInfo.TryGetValue(ViewToRemoveName, out DatasmithCameraActorInfo ViewToRemove))
+				if (ViewNameToCameraInfo.TryGetValue(ViewToRemoveName, out DatasmithCameraActorInfo ViewToRemove))
 				{
 					ViewToRemove.ApplyDeletedStatus();
 				}
 				bHasChanged = true;
 			}
+			
+			// removing name prefix from viewport cameras, ortherwise they will stay
+			// renamed in Rhino, causing all sort of problems
+			foreach (var View in AllViewInfos)
+			{
+				if (View.Name.StartsWith("Current_"))
+				{
+					View.Name = View.Name.Remove(0, "Current_".Length);
+				}
+			}
 
 			bIsDirty |= bHasChanged;
 		}
 
-		private DatasmithCameraActorInfo GenerateActorCameraInfo(ViewInfo NamedView, string NamedViewHash)
+		private DatasmithCameraActorInfo GenerateActorCameraInfo(ViewInfo View, string ViewHash)
 		{
-			string BaseLabel = DatasmithRhinoUniqueNameGenerator.GetTargetName(NamedView);
+			string BaseLabel = DatasmithRhinoUniqueNameGenerator.GetTargetName(View);
 			string UniqueLabel = ActorLabelGenerator.GenerateUniqueNameFromBaseName(BaseLabel);
 			//We can't use the NamedView.Viewport.Id, as it may be unset.
 			string Name = UniqueLabel;
 
-			return new DatasmithCameraActorInfo(NamedView.Viewport, Name, UniqueLabel, BaseLabel, NamedViewHash);
+			return new DatasmithCameraActorInfo(View.Viewport, Name, UniqueLabel, BaseLabel, ViewHash);
 		}
 
 		private void ParseAllRhinoMeshes()
diff --git a/Engine/Source/Programs/Enterprise/Datasmith/DatasmithRhinoExporter/Private/Utils/DatasmithRhinoUtilities.cs b/Engine/Source/Programs/Enterprise/Datasmith/DatasmithRhinoExporter/Private/Utils/DatasmithRhinoUtilities.cs
--- a/Engine/Source/Programs/Enterprise/Datasmith/DatasmithRhinoExporter/Private/Utils/DatasmithRhinoUtilities.cs
+++ b/Engine/Source/Programs/Enterprise/Datasmith/DatasmithRhinoExporter/Private/Utils/DatasmithRhinoUtilities.cs
@@ -15,7 +15,9 @@ namespace DatasmithRhino.Utils
 {
 	public static class DatasmithRhinoUtilities
 	{
-		public static string GetNamedViewHash(ViewInfo NamedView)
+		private const int PRECISION_DECIMALS = 7;
+		
+		public static string GetViewHash(ViewInfo NamedView)
 		{
 			return ComputeHashStringFromBytes(SerializeNamedView(NamedView));
 		}
@@ -68,11 +70,30 @@ namespace DatasmithRhino.Utils
 					Writer.Write(NamedView.Name);
 
 					ViewportInfo Viewport = NamedView.Viewport;
-					Writer.Write(Viewport.FrustumAspect);
-					SerializeVector3d(Writer, new Vector3d(Viewport.CameraLocation));
-					SerializeVector3d(Writer, Viewport.CameraDirection);
-					SerializeVector3d(Writer, Viewport.CameraUp);
-					Writer.Write(Viewport.Camera35mmLensLength);
+
+					// Some Views might be recalculated constantly (specially viewport views)
+					// leading to imprecision, hence we cut down a bit the precision here to avoid
+					// forcing unnecessary updates
+					double FrustumAspect = Math.Round(Viewport.FrustumAspect, PRECISION_DECIMALS);
+					double Camera35mmLensLength = Math.Round(Viewport.Camera35mmLensLength, PRECISION_DECIMALS);
+					Vector3d CameraLocation = new Vector3d(
+						Math.Round(Viewport.CameraLocation.X, PRECISION_DECIMALS), 
+						Math.Round(Viewport.CameraLocation.Y, PRECISION_DECIMALS),
+						Math.Round(Viewport.CameraLocation.Z, PRECISION_DECIMALS));
+					Vector3d CameraDirection = new Vector3d(
+						Math.Round(Viewport.CameraDirection.X, PRECISION_DECIMALS), 
+						Math.Round(Viewport.CameraDirection.Y, PRECISION_DECIMALS),
+						Math.Round(Viewport.CameraDirection.Z, PRECISION_DECIMALS));
+					Vector3d CameraUp = new Vector3d(
+						Math.Round(Viewport.CameraUp.X, PRECISION_DECIMALS), 
+						Math.Round(Viewport.CameraUp.Y, PRECISION_DECIMALS),
+						Math.Round(Viewport.CameraUp.Z, PRECISION_DECIMALS));
+					
+					Writer.Write(FrustumAspect);
+					Writer.Write(Camera35mmLensLength);
+					SerializeVector3d(Writer, CameraLocation);
+					SerializeVector3d(Writer, CameraDirection);
+					SerializeVector3d(Writer, CameraUp);
 				}
 				return Stream.ToArray();
 			}
