Fixed compilation error for DatasmithSDKSample

#rnx

diff --git a/Engine/Source/Programs/Enterprise/Datasmith/DatasmithSDK/DatasmithSDK.Target.cs b/Engine/Source/Programs/Enterprise/Datasmith/DatasmithSDK/DatasmithSDK.Target.cs
--- a/Engine/Source/Programs/Enterprise/Datasmith/DatasmithSDK/DatasmithSDK.Target.cs
+++ b/Engine/Source/Programs/Enterprise/Datasmith/DatasmithSDK/DatasmithSDK.Target.cs
@@ -159,6 +159,8 @@ public class DatasmithSDKTarget : TargetRules
 			(@"$(EngineDir)\Source\Runtime\Core\Public\", Path.Combine(DstDatasmithSDK,"Private\\"), "*.inl"),
 			(@"$(EngineDir)\Source\Runtime\CoreUObject\Public\", Path.Combine(DstDatasmithSDK,"Private\\"), "*.h"),
 			(@"$(EngineDir)\Source\Runtime\CoreUObject\Public\", Path.Combine(DstDatasmithSDK,"Private\\"), "*.inl"),
+			(@"$(EngineDir)\Intermediate\Build\Win64\DatasmithSDK\Inc\CoreUObject\UHT\", Path.Combine(DstDatasmithSDK,"Private\\UObject\\"), "TopLevelAssetPath.generated.h"),
+			(@"$(EngineDir)\Intermediate\Build\Win64\DatasmithSDK\Inc\CoreUObject\UHT\", Path.Combine(DstDatasmithSDK,"Private\\UObject\\"), "GCObject.generated.h"),
 		};
 
 		foreach ((string, string, string) ToCopy in FilesToCopy)
