RN to Datasmith - Fix/Add multiple material export types that were not  supported.

diff --git a/Engine/Source/Programs/Enterprise/Datasmith/DatasmithRhinoExporter/DatasmithRhino7/DatasmithRhino7.csproj b/Engine/Source/Programs/Enterprise/Datasmith/DatasmithRhinoExporter/DatasmithRhino7/DatasmithRhino7.csproj
--- a/Engine/Source/Programs/Enterprise/Datasmith/DatasmithRhinoExporter/DatasmithRhino7/DatasmithRhino7.csproj
+++ b/Engine/Source/Programs/Enterprise/Datasmith/DatasmithRhinoExporter/DatasmithRhino7/DatasmithRhino7.csproj
@@ -27,7 +27,7 @@
     <DebugType>full</DebugType>
     <Optimize>false</Optimize>
     <OutputPath>..\..\..\..\..\..\Binaries\$(Platform_OS)\Rhino\7</OutputPath>
-    <DefineConstants>DEBUG;TRACE;RHINO_VERSION_7</DefineConstants>
+    <DefineConstants>DEBUG;TRACE;RHINO_VERSION_7;RHINO_VERSION_7_OR_NEWER</DefineConstants>
     <ErrorReport>prompt</ErrorReport>
     <CodeAnalysisIgnoreBuiltInRules>false</CodeAnalysisIgnoreBuiltInRules>
   </PropertyGroup>
diff --git a/Engine/Source/Programs/Enterprise/Datasmith/DatasmithRhinoExporter/DatasmithRhino8/DatasmithRhino8.csproj b/Engine/Source/Programs/Enterprise/Datasmith/DatasmithRhinoExporter/DatasmithRhino8/DatasmithRhino8.csproj
--- a/Engine/Source/Programs/Enterprise/Datasmith/DatasmithRhinoExporter/DatasmithRhino8/DatasmithRhino8.csproj
+++ b/Engine/Source/Programs/Enterprise/Datasmith/DatasmithRhinoExporter/DatasmithRhino8/DatasmithRhino8.csproj
@@ -17,7 +17,7 @@
   </PropertyGroup>
   <PropertyGroup>
     <OutputPath>..\..\..\..\..\..\Binaries\$(Platform_OS)\Rhino\8</OutputPath>
-    <DefineConstants>DEBUG;TRACE;RHINO_VERSION_8;USING_NET_CORE;RHINO_VERSION_8_OR_NEWER</DefineConstants>
+    <DefineConstants>DEBUG;TRACE;RHINO_VERSION_8;USING_NET_CORE;RHINO_VERSION_7_OR_NEWER;RHINO_VERSION_8_OR_NEWER</DefineConstants>
     <CodeAnalysisIgnoreBuiltInRules>false</CodeAnalysisIgnoreBuiltInRules>
   </PropertyGroup>
   <PropertyGroup Condition="'$(Platform_OS)' == 'Win64'">
diff --git a/Engine/Source/Programs/Enterprise/Datasmith/DatasmithRhinoExporter/Private/ElementExporters/DatasmithRhinoMaterialExporter.cs b/Engine/Source/Programs/Enterprise/Datasmith/DatasmithRhinoExporter/Private/ElementExporters/DatasmithRhinoMaterialExporter.cs
--- a/Engine/Source/Programs/Enterprise/Datasmith/DatasmithRhinoExporter/Private/ElementExporters/DatasmithRhinoMaterialExporter.cs
+++ b/Engine/Source/Programs/Enterprise/Datasmith/DatasmithRhinoExporter/Private/ElementExporters/DatasmithRhinoMaterialExporter.cs
@@ -58,20 +58,91 @@ namespace DatasmithRhino.ElementExporters
 
 		private void ParseMaterial(FDatasmithFacadeUEPbrMaterial DSMaterial, Material RhinoMaterial)
 		{
-			Color MaterialDiffuseColor = RhinoMaterial.DiffuseColor;
-			MaterialDiffuseColor = Color.FromArgb(255 - (byte)(255 * RhinoMaterial.Transparency), MaterialDiffuseColor);
-			float Roughness = 1f - (float) (RhinoMaterial.Shine / Material.MaxShine);
-			float Reflectivity = (float) RhinoMaterial.Reflectivity;
+			// Default values for legacy materials. Then, each material type could 
+			// replace some of these values in different ways 
+			Color DiffuseColor = RhinoMaterial.DiffuseColor;
+			float Roughness = 1f - (float) RhinoMaterial.Reflectivity;
+			float Metallic = (float) (RhinoMaterial.Shine / Material.MaxShine);
+			float Alpha = (float) (1f - RhinoMaterial.Transparency);
+			float IOR = 1f;
+			float EmissivePower = 0;
 			
-			#if RHINO_VERSION_8_OR_NEWER
-			// Rhino Materials have different shading models and properties need to be adjusted based on that
-			if (RhinoMaterial.IsPhysicallyBased)
+			// Rhino has multiple material types, some with reflective PBR and some with special
+			// case rendering (like medium scattering, etc). Many of these materials are not 1:1 translatable
+			// to a Datasmith Material, so we apply some custom transformation to make it look similar and
+			// also log warnings for user visibility
+			switch (RhinoMaterial.RenderMaterial.TypeName)
 			{
-				MaterialDiffuseColor = RhinoMaterial.PhysicallyBased.BaseColor.AsSystemColor();
-				Roughness = (float)RhinoMaterial.PhysicallyBased.Roughness;
-				Reflectivity = (float)RhinoMaterial.PhysicallyBased.Metallic;
+				case "Blend":
+					RhinoApp.WriteLine($"Warning. Blend Material ({RhinoMaterial.Name}) export not supported at the time.");
+
+				break;
+				case "Custom":
+					// No Special Behavior
+				break;
+				case "Double Sided":
+					// Rhino Double Sided materials are actually 2 different materials, so we only export one
+					// A different implementation would require to duplicate the geometry per face (possibly todo?)
+					RhinoApp.WriteLine($"Warning. Double Sided Material ({RhinoMaterial.Name}) would export only the front side as single.");
+				break;
+				case "Emission":
+					// The Rhino SDK doesn't seem to have a way to extract the emission intensity
+					// from emissive materials, only the color
+					RhinoApp.WriteLine($"Warning. Emission Material ({RhinoMaterial.Name}) would export with Intensity 1.");
+					EmissivePower = 1;
+					DiffuseColor = RhinoMaterial.EmissionColor;
+				break;
+				case "Gem":
+					RhinoApp.WriteLine($"Warning. Gem Material ({RhinoMaterial.Name}) is partially supported.");
+					DiffuseColor = RhinoMaterial.TransparentColor;
+					IOR = (float)RhinoMaterial.IndexOfRefraction;
+					// not real value, but avoids the material being invisible in Unreal
+					// the Rhino material doesn't have any translatable value from Gem to Alpha
+					Alpha = 0.1f;
+					Metallic = 0;
+					Roughness = 0;
+				break;
+				case "Glass":
+					RhinoApp.WriteLine($"Warning. Glass Material ({RhinoMaterial.Name}) is partially supported.");
+					DiffuseColor = RhinoMaterial.TransparentColor;
+					IOR = (float)RhinoMaterial.IndexOfRefraction;
+					// Glass use Shine for "clarity", and while it's not the same as alpha, it's the closest
+					// we can get from this Shader
+					Alpha = (float)RhinoMaterial.Shine / 255f;
+					Metallic = 0;
+					Roughness = 0;
+				break;
+				case "Metal":
+					DiffuseColor = RhinoMaterial.SpecularColor;
+					Metallic = 1;
+				break;
+				case "Paint":
+					Metallic = 0;
+				break;
+				case "Physically Based":
+					// PBR Materials were introduced in Rhino 7, and only present from that API version
+					#if RHINO_VERSION_7_OR_NEWER
+					if (RhinoMaterial.IsPhysicallyBased)
+					{
+						DiffuseColor = RhinoMaterial.PhysicallyBased.BaseColor.AsSystemColor();
+						Roughness = (float)RhinoMaterial.PhysicallyBased.Roughness;
+						Metallic = (float)RhinoMaterial.PhysicallyBased.Metallic;
+						Alpha = (float)RhinoMaterial.PhysicallyBased.Opacity;
+						IOR = (float)RhinoMaterial.PhysicallyBased.OpacityIOR;
+					}
+					#endif
+				break;
+				case "Picture":
+					RhinoApp.WriteLine($"Warning. Picture Material ({RhinoMaterial.Name}) is partially supported.");
+				break;
+				case "Plaster":
+					Metallic = 0;
+					Roughness = 0;
+				break;
+				case "Plastic":
+					Metallic = 0;
+				break;
 			}
-			#endif
 			
 			Texture[] MaterialTextures = RhinoMaterial.GetTextures();
 			for (int TextureIndex = 0; TextureIndex < MaterialTextures.Length; ++TextureIndex)
@@ -82,21 +153,41 @@ namespace DatasmithRhino.ElementExporters
 					DatasmithTextureInfo TextureInfo = ExportContext.GetTextureInfoFromRhinoTexture(RhinoTexture.Id);
 					if (TextureInfo != null)
 					{
-						AddTextureToMaterial(DSMaterial, TextureInfo, MaterialDiffuseColor);
+						AddTextureToMaterial(DSMaterial, TextureInfo, DiffuseColor);
 					}
 				}
 			}
 
+			if (EmissivePower > 0)
+			{
+				FDatasmithFacadeMaterialExpressionGeneric Multiply = DSMaterial.AddMaterialExpressionGeneric();
+				Multiply.SetExpressionName("Multiply");
+				Multiply.ConnectExpression(DSMaterial.GetEmissiveColor());
+				
+				var ColorEmission = DSMaterial.AddMaterialExpressionColor();
+				ColorEmission.SetName("Emission Color");
+				ColorEmission.SetsRGBColor(DiffuseColor.R, DiffuseColor.G, DiffuseColor.B, DiffuseColor.A);
+				ColorEmission.ConnectExpression(Multiply.GetInput(0));
+				
+				FDatasmithFacadeMaterialExpressionScalar ColorIntensityScalar = DSMaterial.AddMaterialExpressionScalar();
+				ColorIntensityScalar.SetName("Emission Power");
+				ColorIntensityScalar.SetScalar(EmissivePower);
+				ColorIntensityScalar.ConnectExpression(Multiply.GetInput(1));
+				
+				// the base color becomes white and we only use the emissive color
+				DiffuseColor = Color.White;
+			}
+			
 			// Set a diffuse color if there's nothing in the BaseColor
 			if (DSMaterial.GetBaseColor().GetExpression() == null)
 			{
 				FDatasmithFacadeMaterialExpressionColor ColorExpression = DSMaterial.AddMaterialExpressionColor();
 				ColorExpression.SetName("Diffuse Color");
-				ColorExpression.SetsRGBColor(MaterialDiffuseColor.R, MaterialDiffuseColor.G, MaterialDiffuseColor.B, MaterialDiffuseColor.A);
+				ColorExpression.SetsRGBColor(DiffuseColor.R, DiffuseColor.G, DiffuseColor.B, DiffuseColor.A);
 				ColorExpression.ConnectExpression(DSMaterial.GetBaseColor());
 			}
 
-			if (RhinoMaterial.Transparency > 0)
+			if (Alpha < 1f)
 			{
 				DSMaterial.SetBlendMode(/*EBlendMode::BLEND_Translucent*/2);
 				if (DSMaterial.GetOpacity().GetExpression() == null)
@@ -104,7 +195,7 @@ namespace DatasmithRhino.ElementExporters
 					// Transparent color
 					FDatasmithFacadeMaterialExpressionScalar Scalar = DSMaterial.AddMaterialExpressionScalar();
 					Scalar.SetName("Opacity");
-					Scalar.SetScalar(1 - (float)RhinoMaterial.Transparency);
+					Scalar.SetScalar(Alpha);
 					Scalar.ConnectExpression(DSMaterial.GetOpacity());
 				}
 				else
@@ -123,6 +214,16 @@ namespace DatasmithRhino.ElementExporters
 
 					DSMaterial.GetOpacity().SetExpression(Multiply);
 				}
+
+				if (Math.Abs(IOR - 1) > 0.00001f)
+				{
+					// Note, currently there is no way to set the Refraction Mode in the Material
+					// so while the ior might be properly set, refraction input is going to be disabled in the material
+					FDatasmithFacadeMaterialExpressionScalar Refraction = DSMaterial.AddMaterialExpressionScalar();
+					Refraction.SetName("Refraction");
+					Refraction.SetScalar(IOR);
+					Refraction.ConnectExpression(DSMaterial.GetRefraction());
+				}
 			}
 			
 			FDatasmithFacadeMaterialExpressionScalar ShininessExpression = DSMaterial.AddMaterialExpressionScalar();
@@ -132,7 +233,7 @@ namespace DatasmithRhino.ElementExporters
 			
 			FDatasmithFacadeMaterialExpressionScalar ReflectivityExpression = DSMaterial.AddMaterialExpressionScalar();
 			ReflectivityExpression.SetName("Metallic");
-			ReflectivityExpression.SetScalar(Reflectivity);
+			ReflectivityExpression.SetScalar(Metallic);
 			ReflectivityExpression.ConnectExpression(DSMaterial.GetMetallic());
 		}
 
