Fixed compilation error of Datasmith 3ds max exporter

#rnx

diff --git a/Engine/Source/Programs/Enterprise/Datasmith/DatasmithMaxExporter/DatasmithMax2017.Build.cs b/Engine/Source/Programs/Enterprise/Datasmith/DatasmithMaxExporter/DatasmithMax2017.Build.cs
deleted file mode 100644
diff --git a/Engine/Source/Programs/Enterprise/Datasmith/DatasmithMaxExporter/DatasmithMax2017.Target.cs b/Engine/Source/Programs/Enterprise/Datasmith/DatasmithMaxExporter/DatasmithMax2017.Target.cs
deleted file mode 100644
diff --git a/Engine/Source/Programs/Enterprise/Datasmith/DatasmithMaxExporter/DatasmithMax2018.Build.cs b/Engine/Source/Programs/Enterprise/Datasmith/DatasmithMaxExporter/DatasmithMax2018.Build.cs
deleted file mode 100644
diff --git a/Engine/Source/Programs/Enterprise/Datasmith/DatasmithMaxExporter/DatasmithMax2018.Target.cs b/Engine/Source/Programs/Enterprise/Datasmith/DatasmithMaxExporter/DatasmithMax2018.Target.cs
deleted file mode 100644
diff --git a/Engine/Source/Programs/Enterprise/Datasmith/DatasmithMaxExporter/DatasmithMax2019.Build.cs b/Engine/Source/Programs/Enterprise/Datasmith/DatasmithMaxExporter/DatasmithMax2019.Build.cs
deleted file mode 100644
diff --git a/Engine/Source/Programs/Enterprise/Datasmith/DatasmithMaxExporter/DatasmithMax2019.Target.cs b/Engine/Source/Programs/Enterprise/Datasmith/DatasmithMaxExporter/DatasmithMax2019.Target.cs
deleted file mode 100644
diff --git a/Engine/Source/Programs/Enterprise/Datasmith/DatasmithMaxExporter/DatasmithMax2020.Build.cs b/Engine/Source/Programs/Enterprise/Datasmith/DatasmithMaxExporter/DatasmithMax2020.Build.cs
deleted file mode 100644
diff --git a/Engine/Source/Programs/Enterprise/Datasmith/DatasmithMaxExporter/DatasmithMax2020.Target.cs b/Engine/Source/Programs/Enterprise/Datasmith/DatasmithMaxExporter/DatasmithMax2020.Target.cs
deleted file mode 100644
diff --git a/Engine/Source/Programs/Enterprise/Datasmith/DatasmithMaxExporter/DatasmithMax2021.Build.cs b/Engine/Source/Programs/Enterprise/Datasmith/DatasmithMaxExporter/DatasmithMax2021.Build.cs
deleted file mode 100644
diff --git a/Engine/Source/Programs/Enterprise/Datasmith/DatasmithMaxExporter/DatasmithMax2021.Target.cs b/Engine/Source/Programs/Enterprise/Datasmith/DatasmithMaxExporter/DatasmithMax2021.Target.cs
deleted file mode 100644
diff --git a/Engine/Source/Programs/Enterprise/Datasmith/DatasmithMaxExporter/DatasmithMax2022.Build.cs b/Engine/Source/Programs/Enterprise/Datasmith/DatasmithMaxExporter/DatasmithMax2022.Build.cs
--- a/Engine/Source/Programs/Enterprise/Datasmith/DatasmithMaxExporter/DatasmithMax2022.Build.cs
+++ b/Engine/Source/Programs/Enterprise/Datasmith/DatasmithMaxExporter/DatasmithMax2022.Build.cs
@@ -1,9 +1,104 @@
 // Copyright Epic Games, Inc. All Rights Reserved.
 
+using System;
 using System.Diagnostics;
+using System.IO;
+using System.Linq;
 
 namespace UnrealBuildTool.Rules
 {
+	public abstract class DatasmithMaxBase : ModuleRules
+	{
+		public DatasmithMaxBase(ReadOnlyTargetRules Target)
+			: base(Target)
+		{
+
+			bUseRTTI = true;
+
+			// To avoid clashes with Max SDK
+			// todo: separate Slate code from 3ds max to a module
+			bUseUnity = false;
+
+
+			PublicDefinitions.Add("NEW_DIRECTLINK_PLUGIN=1");
+
+			PrivateDependencyModuleNames.AddRange(
+				new string[]
+				{
+					"DatasmithExporter",
+					"DatasmithExporterUI",
+
+					"UdpMessaging", // required for DirectLink networking
+					"UEOpenExr",
+
+					"IntelTBB",
+					"Slate",
+					"SlateCore",
+				}
+			);
+
+			PrivateIncludePathModuleNames.Add("Launch");
+
+			// Max SDK setup
+			{
+				string MaxVersionString = GetMaxVersion();
+				string MaxSDKLocation = "";
+
+				// Try with build machine setup
+				string SDKRootEnvVar = System.Environment.GetEnvironmentVariable("UE_SDKS_ROOT");
+				if (SDKRootEnvVar != null && SDKRootEnvVar != "")
+				{
+					MaxSDKLocation = Path.Combine(SDKRootEnvVar, "HostWin64", "Win64", "3dsMax", MaxVersionString);
+				}
+
+				if (!Directory.Exists(MaxSDKLocation))
+				{
+					// Try with custom setup
+					string MaxSDKEnvVar = System.Environment.GetEnvironmentVariable("ADSK_3DSMAX_SDK_" + MaxVersionString);
+					if (MaxSDKEnvVar != null && MaxSDKEnvVar != "")
+					{
+						MaxSDKLocation = MaxSDKEnvVar;
+					}
+				}
+
+				// Make sure this version of Max is actually installed
+				if (Directory.Exists(MaxSDKLocation))
+				{
+					PublicSystemIncludePaths.Add(Path.Combine(MaxSDKLocation, "include"));
+
+					string LibraryPaths = Path.Combine(MaxSDKLocation, "lib", "x64", "Release");
+					PublicAdditionalLibraries.Add(Path.Combine(LibraryPaths, "assetmanagement.lib"));
+					PublicAdditionalLibraries.Add(Path.Combine(LibraryPaths, "bmm.lib"));
+					PublicAdditionalLibraries.Add(Path.Combine(LibraryPaths, "core.lib"));
+					PublicAdditionalLibraries.Add(Path.Combine(LibraryPaths, "geom.lib"));
+					PublicAdditionalLibraries.Add(Path.Combine(LibraryPaths, "maxutil.lib"));
+					PublicAdditionalLibraries.Add(Path.Combine(LibraryPaths, "Maxscrpt.lib"));
+					PublicAdditionalLibraries.Add(Path.Combine(LibraryPaths, "mesh.lib"));
+				}
+			}
+
+			// Itoo ForestPack/RailClone API
+			string ItooInterfaceLocation = Path.Combine(ModuleDirectory, "ThirdParty", "Itoo");
+			bool bWithItooInterface = Directory.Exists(ItooInterfaceLocation);
+			PublicDefinitions.Add("WITH_ITOO_INTERFACE=" + (bWithItooInterface ? "1" : "0"));
+			if (bWithItooInterface)
+			{
+				PrivateIncludePaths.Add(ItooInterfaceLocation);
+			}
+
+			PrivateIncludePaths.Add(ModuleDirectory);
+
+			// Only 3ds max version 2022 and higher is supported now.
+			//if (int.Parse(GetMaxVersion()) < 2022)
+			//{
+			//	// Being cheeky here as suppressing the warning doesn't work in ModuleRules for some reason
+			//	CppStandard = Enum.GetValues<CppStandardVersion>().Where(x => x.ToString().Contains("Cpp17")).First();
+			//}
+		}
+
+		public abstract string GetMaxVersion();
+	}
+
 	[SupportedPlatforms("Win64")]
 	public class DatasmithMax2022 : DatasmithMaxBase
 	{
diff --git a/Engine/Source/Programs/Enterprise/Datasmith/DatasmithMaxExporter/DatasmithMax2022.Target.cs b/Engine/Source/Programs/Enterprise/Datasmith/DatasmithMaxExporter/DatasmithMax2022.Target.cs
--- a/Engine/Source/Programs/Enterprise/Datasmith/DatasmithMaxExporter/DatasmithMax2022.Target.cs
+++ b/Engine/Source/Programs/Enterprise/Datasmith/DatasmithMaxExporter/DatasmithMax2022.Target.cs
@@ -1,7 +1,93 @@
 // Copyright Epic Games, Inc. All Rights Reserved.
 
 using UnrealBuildTool;
+using System;
 using System.Collections.Generic;
+using System.IO;
+
+[SupportedPlatforms("Win64")]
+public abstract class DatasmithMaxBaseTarget : TargetRules
+{
+	public DatasmithMaxBaseTarget(TargetInfo Target)
+		: base(Target)
+	{
+		Type = TargetType.Program;
+		IncludeOrderVersion = EngineIncludeOrderVersion.Latest;
+		SolutionDirectory = "Programs/Datasmith";
+		bBuildInSolutionByDefault = false;
+		bLegalToDistributeBinary = true;
+
+		bShouldCompileAsDLL = true;
+		LinkType = TargetLinkType.Monolithic;
+
+		WindowsPlatform.ModuleDefinitionFile = "Programs/Enterprise/Datasmith/DatasmithMaxExporter/DatasmithMaxExporterWithDirectLink.def";
+
+		if(StaticAnalyzer != StaticAnalyzer.None)
+		{
+			string MaxSDKLocation = "";
+			string SDKRootEnvVar = System.Environment.GetEnvironmentVariable("UE_SDKS_ROOT");
+			if (SDKRootEnvVar != null && SDKRootEnvVar != "")
+			{
+				MaxSDKLocation = Path.Combine(SDKRootEnvVar, "HostWin64", "Win64", "3dsMax");
+			}
+
+			if (!Directory.Exists(MaxSDKLocation))
+			{
+			}
+
+			WindowsPlatform.Compiler = WindowsCompiler.Clang;
+			Console.WriteLine("Skipping " + MaxSDKLocation);
+			AdditionalCompilerArguments = "-Xanalyzer,-isystem-after" + MaxSDKLocation;
+		}
+		else
+		{
+			WindowsPlatform.Compiler = WindowsCompiler.VisualStudio2022;
+			AdditionalCompilerArguments = "/Zc:referenceBinding- /Zc:strictStrings- /Zc:rvalueCast";
+		}
+
+		WindowsPlatform.bStrictConformanceMode = false;
+		WindowsPlatform.bStrictPreprocessorConformance = false;
+
+		bBuildDeveloperTools = false;
+		bBuildWithEditorOnlyData = true;
+		bCompileAgainstEngine = false;
+		bCompileAgainstCoreUObject = true;
+		bCompileICU = false;
+		bUsesSlate = false;
+
+		bHasExports = true;
+		bForceEnableExceptions = true;
+
+		// For DirectLinkUI (see FDatasmithExporterManager::FInitOptions)
+		GlobalDefinitions.Add("UE_EXTERNAL_PROFILING_ENABLED=0");
+		GlobalDefinitions.Add("NO_INIUTIL_USING");
+
+		// Disable CrashReporter
+		GlobalDefinitions.Add("NOINITCRASHREPORTER=1");
+
+		// todo: remove?
+		// bSupportEditAndContinue = true;
+	}
+
+	protected void AddCopyPostBuildStep(TargetInfo Target)
+	{
+		// Add a post-build step that copies the output to a file with the .dle extension
+		string OutputName = "$(TargetName)";
+		if (Target.Configuration != UnrealTargetConfiguration.Development)
+		{
+			OutputName = string.Format("{0}-{1}-{2}", OutputName, Target.Platform, Target.Configuration);
+		}
+
+		string SrcOutputFileName = string.Format(@"$(EngineDir)\Binaries\Win64\{0}\{1}.dll", ExeBinariesSubFolder, OutputName);
+		string DstOutputFileName = string.Format(@"$(EngineDir)\Binaries\Win64\{0}\{1}.gup", ExeBinariesSubFolder, OutputName);
+
+		PostBuildSteps.Add(string.Format("echo Copying {0} to {1}...", SrcOutputFileName, DstOutputFileName));
+		PostBuildSteps.Add(string.Format("copy /Y \"{0}\" \"{1}\" 1>nul", SrcOutputFileName, DstOutputFileName));
+
+		PostBuildSteps.Add(string.Format("echo Copying {0} to {1}...", SrcOutputFileName, DstOutputFileName));
+		PostBuildSteps.Add(string.Format("copy /Y \"{0}\" \"{1}\" 1>nul", SrcOutputFileName, DstOutputFileName));
+	}
+}
 
 public class DatasmithMax2022Target : DatasmithMaxBaseTarget
 {
diff --git a/Engine/Source/Programs/Enterprise/Datasmith/DatasmithMaxExporter/Private/DatasmithMaxDirectLinkMaxScript.cpp b/Engine/Source/Programs/Enterprise/Datasmith/DatasmithMaxExporter/Private/DatasmithMaxDirectLinkMaxScript.cpp
--- a/Engine/Source/Programs/Enterprise/Datasmith/DatasmithMaxExporter/Private/DatasmithMaxDirectLinkMaxScript.cpp
+++ b/Engine/Source/Programs/Enterprise/Datasmith/DatasmithMaxExporter/Private/DatasmithMaxDirectLinkMaxScript.cpp
@@ -2,6 +2,9 @@
 
 #ifdef NEW_DIRECTLINK_PLUGIN
 
+#pragma warning( push )
+#pragma warning( disable: 5103 )
+
 #include "DatasmithMaxDirectLink.h"
 #include "DatasmithMaxLogger.h"
 
@@ -661,7 +664,7 @@ public:
 			{
 				if (GetExporter())
 				{
-					MCHAR* ScriptCode = _T("Datasmith_ExportDialog selected:false");
+					const MCHAR* ScriptCode = _T("Datasmith_ExportDialog selected:false");
 #if MAX_PRODUCT_YEAR_NUMBER >= 2022
 					 ExecuteMAXScriptScript(ScriptCode, MAXScript::ScriptSource::NonEmbedded);
 #else
@@ -679,7 +682,7 @@ public:
 			{
 				if (GetExporter())
 				{
-					MCHAR* ScriptCode = _T("Datasmith_ExportDialog selected:true");
+					const MCHAR* ScriptCode = _T("Datasmith_ExportDialog selected:true");
 #if MAX_PRODUCT_YEAR_NUMBER >= 2022
 					 ExecuteMAXScriptScript(ScriptCode, MAXScript::ScriptSource::NonEmbedded);
 #else
@@ -752,4 +755,6 @@ Primitive SetupActions_pf(_M("Datasmith_SetupActions"), SetupActions_cf);
 
 #include "Windows/HideWindowsPlatformTypes.h"
 
+#pragma warning( pop )
+
 #endif // NEW_DIRECTLINK_PLUGIN
diff --git a/Engine/Source/Programs/Enterprise/Datasmith/DatasmithMaxExporter/Public/DatasmithMaxDirectLink.h b/Engine/Source/Programs/Enterprise/Datasmith/DatasmithMaxExporter/Public/DatasmithMaxDirectLink.h
--- a/Engine/Source/Programs/Enterprise/Datasmith/DatasmithMaxExporter/Public/DatasmithMaxDirectLink.h
+++ b/Engine/Source/Programs/Enterprise/Datasmith/DatasmithMaxExporter/Public/DatasmithMaxDirectLink.h
@@ -386,9 +386,9 @@ struct FSceneUpdateStats
 	int64 UpdateTexturesBakedPixels = 0;
 };
 
-#define SCENE_UPDATE_STAT_INC(Category, Name) {Stats.##Category##Name++;}
-#define SCENE_UPDATE_STAT_SET(Category, Name, Value) Stats.##Category##Name = Value
-#define SCENE_UPDATE_STAT_GET(Category, Name) Stats.##Category##Name
+#define SCENE_UPDATE_STAT_INC(Category, Name) {Stats.Category##Name++;}
+#define SCENE_UPDATE_STAT_SET(Category, Name, Value) Stats.Category##Name = Value
+#define SCENE_UPDATE_STAT_GET(Category, Name) Stats.Category##Name
 
 
 /**
