Fixed compilation errors of Datasmith 3dsMax exporter

#rnx

diff --git a/Engine/Source/Programs/Enterprise/Datasmith/DatasmithMaxExporter/DatasmithMax2022.Target.cs b/Engine/Source/Programs/Enterprise/Datasmith/DatasmithMaxExporter/DatasmithMax2022.Target.cs
--- a/Engine/Source/Programs/Enterprise/Datasmith/DatasmithMaxExporter/DatasmithMax2022.Target.cs
+++ b/Engine/Source/Programs/Enterprise/Datasmith/DatasmithMaxExporter/DatasmithMax2022.Target.cs
@@ -42,11 +42,9 @@ public abstract class DatasmithMaxBaseTarget : TargetRules
 		else
 		{
 			WindowsPlatform.Compiler = WindowsCompiler.VisualStudio2022;
-			AdditionalCompilerArguments = "/Zc:referenceBinding- /Zc:strictStrings- /Zc:rvalueCast";
 		}
 
-		WindowsPlatform.bStrictConformanceMode = false;
-		WindowsPlatform.bStrictPreprocessorConformance = false;
+		WindowsPlatform.bStrictPreprocessorConformance = true;
 
 		bBuildDeveloperTools = false;
 		bBuildWithEditorOnlyData = true;
diff --git a/Engine/Source/Programs/Enterprise/Datasmith/DatasmithMaxExporter/Private/DatasmithMaxWriter.cpp b/Engine/Source/Programs/Enterprise/Datasmith/DatasmithMaxExporter/Private/DatasmithMaxWriter.cpp
--- a/Engine/Source/Programs/Enterprise/Datasmith/DatasmithMaxExporter/Private/DatasmithMaxWriter.cpp
+++ b/Engine/Source/Programs/Enterprise/Datasmith/DatasmithMaxExporter/Private/DatasmithMaxWriter.cpp
@@ -952,7 +952,7 @@ void FDatasmithMaxMatWriter::DumpWeightedColor(TSharedPtr<IDatasmithCompositeTex
 	CompTex->AddParamVal1(IDatasmithCompositeTexture::ParamVal(1.0f-(float)Weight, TEXT("Weight")));
 }
 
-FString FDatasmithMaxMatWriter::DumpColorOfTexmap(TSharedPtr<IDatasmithCompositeTexture>& CompTex, Texmap* InTexmap, const TCHAR* ColorPrefix, TCHAR* Property)
+FString FDatasmithMaxMatWriter::DumpColorOfTexmap(TSharedPtr<IDatasmithCompositeTexture>& CompTex, Texmap* InTexmap, const TCHAR* ColorPrefix, const TCHAR* Property)
 {
 	int NumParamBlocks = InTexmap->NumParamBlocks();
 
diff --git a/Engine/Source/Programs/Enterprise/Datasmith/DatasmithMaxExporter/Public/DatasmithMaxWriter.h b/Engine/Source/Programs/Enterprise/Datasmith/DatasmithMaxExporter/Public/DatasmithMaxWriter.h
--- a/Engine/Source/Programs/Enterprise/Datasmith/DatasmithMaxExporter/Public/DatasmithMaxWriter.h
+++ b/Engine/Source/Programs/Enterprise/Datasmith/DatasmithMaxExporter/Public/DatasmithMaxWriter.h
@@ -130,7 +130,7 @@ public:
 	// selects from any kind of texture the ColorA is mixed with the Color B using Weight
 	static void DumpWeightedColor(TSharedPtr<IDatasmithCompositeTexture>& CompTex, BMM_Color_fl ColorA, BMM_Color_fl ColorB, float Weight, const TCHAR* Prefix);
 	// if we want to export sub-texture but it is void
-	static FString DumpColorOfTexmap(TSharedPtr<IDatasmithCompositeTexture>& CompTex, Texmap* InTexmap, const TCHAR* ColorPrefix, TCHAR* property);
+	static FString DumpColorOfTexmap(TSharedPtr<IDatasmithCompositeTexture>& CompTex, Texmap* InTexmap, const TCHAR* ColorPrefix, const TCHAR* property);
 
 	//---------------------------------------------------------------
 	// SCANLINE DEPENDANT
