Add VisualStudio2026 to preferred compiler list so it can be used if 2022 isn't installed

#rnx
#Synced-CL 49631673

diff --git a/Engine/Source/Programs/Enterprise/Datasmith/DatasmithMaxExporter/DatasmithMax2022.Target.cs b/Engine/Source/Programs/Enterprise/Datasmith/DatasmithMaxExporter/DatasmithMax2022.Target.cs
--- a/Engine/Source/Programs/Enterprise/Datasmith/DatasmithMaxExporter/DatasmithMax2022.Target.cs
+++ b/Engine/Source/Programs/Enterprise/Datasmith/DatasmithMaxExporter/DatasmithMax2022.Target.cs
@@ -41,7 +41,7 @@ public abstract class DatasmithMaxBaseTarget : TargetRules
 		}
 		else
 		{
-			WindowsPlatform.Compiler = WindowsCompiler.VisualStudio2022;
+			WindowsPlatform.PreferredCompilers = [nameof(WindowsCompiler.VisualStudio2022), nameof(WindowsCompiler.VisualStudio2026)];
 		}
 
 		WindowsPlatform.bStrictPreprocessorConformance = true;
