Fixed Compile Datasmith 3DSMax Exporter

- compile failed due to changes in TSet definition from SparseSet to CompactSet in CL 48966924 with bUseCompactSetAsDefault setting to true in BuildConfiguration.xml. Fixed by removing redundand method definitions.

- trying to preflight in ue5/main  (where can run Datasmith External Plugins)- 696ded7ec2542fb19a2b33be

#Synced-CL 49917634

diff --git a/Engine/Source/Programs/Enterprise/Datasmith/DatasmithMaxExporter/Private/DatasmithMaxExporterWithDirectLink.cpp b/Engine/Source/Programs/Enterprise/Datasmith/DatasmithMaxExporter/Private/DatasmithMaxExporterWithDirectLink.cpp
--- a/Engine/Source/Programs/Enterprise/Datasmith/DatasmithMaxExporter/Private/DatasmithMaxExporterWithDirectLink.cpp
+++ b/Engine/Source/Programs/Enterprise/Datasmith/DatasmithMaxExporter/Private/DatasmithMaxExporterWithDirectLink.cpp
@@ -771,7 +771,7 @@ public:
 
 	public:
 
-		typedef TSet<FNodeTracker*>::TRangedForIterator ElementItType;
+		using ElementItType = TSet<FNodeTracker*>::TRangedForIterator;
 
 		FInvalidatedNodeTrackers& InvalidatedNodeTrackers;
 
@@ -793,16 +793,6 @@ public:
 			return *this;
 		}
 
-		FORCEINLINE explicit operator bool() const
-		{ 
-			return !!ElementIt; 
-		}
-		
-		FORCEINLINE bool operator !() const 
-		{
-			return !static_cast<bool>(*this);
-		}
-
 		FORCEINLINE ItElementType* operator->() const
 		{
 			return *ElementIt;
