Adjust preferred compilers to be able to use any version of MSVC rather than specific VisualStudio versions

#rnx

#Synced-CL 50271719

diff --git a/Engine/Source/Programs/Enterprise/Datasmith/DatasmithMaxExporter/DatasmithMax2022.Target.cs b/Engine/Source/Programs/Enterprise/Datasmith/DatasmithMaxExporter/DatasmithMax2022.Target.cs
--- a/Engine/Source/Programs/Enterprise/Datasmith/DatasmithMaxExporter/DatasmithMax2022.Target.cs
+++ b/Engine/Source/Programs/Enterprise/Datasmith/DatasmithMaxExporter/DatasmithMax2022.Target.cs
@@ -41,7 +41,7 @@ public abstract class DatasmithMaxBaseTarget : TargetRules
 		}
 		else
 		{
-			WindowsPlatform.PreferredCompilers = [nameof(WindowsCompiler.VisualStudio2022), nameof(WindowsCompiler.VisualStudio2026)];
+			WindowsPlatform.Compiler = WindowsCompiler.VisualStudio;
 		}
 
 		WindowsPlatform.bStrictPreprocessorConformance = true;
